[UBT]
* Added more scopes around loading engine rules assembly

diff --git a/Engine/Source/Programs/UnrealBuildTool/System/RulesCompiler.cs b/Engine/Source/Programs/UnrealBuildTool/System/RulesCompiler.cs
--- a/Engine/Source/Programs/UnrealBuildTool/System/RulesCompiler.cs
+++ b/Engine/Source/Programs/UnrealBuildTool/System/RulesCompiler.cs
@@ -73,27 +73,26 @@ namespace UnrealBuildTool
 			{
 				if (EngineRulesAssembly == null)
 				{
-					using var ite = Timeline.ScopeEvent("EngineRulesAssembly", Logger);
-
 					// Prevent multiple conflicting processes building rules assembly at the same time
 					using GlobalSingleInstanceMutex _ = new(GlobalSingleInstanceMutex.GetUniqueMutexForPath("UnrealBuildTool_CreateTargetRulesAssembly", Unreal.EngineDirectory), true);
 					{
-						List<PluginInfo> EnginePlugins = new List<PluginInfo>();
-						List<PluginInfo> MarketplacePlugins = new List<PluginInfo>();
-
-						DirectoryReference MarketplaceDirectory = DirectoryReference.Combine(Unreal.EngineDirectory, "Plugins", "Marketplace");
-						foreach (PluginInfo PluginInfo in Plugins.ReadEnginePlugins(Unreal.EngineDirectory))
+						List<PluginInfo> EnginePlugins = [];
+						List<PluginInfo> MarketplacePlugins = [];
 						{
-							if (PluginInfo.File.IsUnderDirectory(MarketplaceDirectory))
-							{
-								MarketplacePlugins.Add(PluginInfo);
-							}
-							else
+							using var ite = Timeline.ScopeEvent("EngineRulesPlugins", Logger);
+							DirectoryReference MarketplaceDirectory = DirectoryReference.Combine(Unreal.EngineDirectory, "Plugins", "Marketplace");
+							foreach (PluginInfo PluginInfo in Plugins.ReadEnginePlugins(Unreal.EngineDirectory))
 							{
-								EnginePlugins.Add(PluginInfo);
+								if (PluginInfo.File.IsUnderDirectory(MarketplaceDirectory))
+								{
+									MarketplacePlugins.Add(PluginInfo);
+								}
+								else
+								{
+									EnginePlugins.Add(PluginInfo);
+								}
 							}
 						}
-
 						EngineRulesAssembly = CreateEngineRulesAssemblyInternal(Unreal.GetExtensionDirs(Unreal.EngineDirectory), ProjectFileGenerator.EngineRulesAssemblyName, EnginePlugins, Unreal.IsEngineInstalled() || bUsePrecompiled, bSkipCompile, bForceCompile, null, Logger);
 
 						if (MarketplacePlugins.Count > 0)
@@ -132,6 +131,8 @@ namespace UnrealBuildTool
 
 			// We want to prefetch the directories in one call to make it go as wide as possible against filesystem
 			{
+				using var fdt = Timeline.ScopeEvent("EngineRulesFetchDirs", Logger);
+
 				List<DirectoryReference> populateDirs = [];
 				foreach (PluginInfo info in Plugins)
 				{
@@ -192,10 +193,13 @@ namespace UnrealBuildTool
 			}
 
 			// Create the assembly
+			using var erat = Timeline.ScopeEvent("EngineRulesAssembly", Logger);
 			DirectoryReference AssemblyDir = RootDirectories[0];
 			FileReference EngineAssemblyFileName = FileReference.Combine(AssemblyDir, "Intermediate", "Build", "BuildRules", AssemblyPrefix + "Rules" + FrameworkAssemblyExtension);
 			RulesAssembly EngineAssembly = new RulesAssembly(Scope, RootDirectories, Plugins, ModuleFileToContext, new List<FileReference>(), EngineAssemblyFileName, bContainsEngineModules: true, DefaultBuildSettings: DefaultEngineBuildSettingsVersion, bReadOnly: bReadOnly, bSkipCompile: bSkipCompile, bForceCompile: bForceCompile, Parent: Parent, Logger: Logger);
+			erat.Finish();
 
+			using var prat = Timeline.ScopeEvent("ProgramRulesAssembly", Logger);
 			Dictionary<FileReference, ModuleRulesContext> ProgramModuleFiles = new Dictionary<FileReference, ModuleRulesContext>();
 			foreach (DirectoryReference RootDirectory in RootDirectories)
 			{
