[UBT]
* Moved Macos stub files into a proper folder that is in line with where all intermediate build files are located

diff --git a/Engine/Source/Programs/UnrealBuildTool/Platform/Mac/MacToolChain.cs b/Engine/Source/Programs/UnrealBuildTool/Platform/Mac/MacToolChain.cs
--- a/Engine/Source/Programs/UnrealBuildTool/Platform/Mac/MacToolChain.cs
+++ b/Engine/Source/Programs/UnrealBuildTool/Platform/Mac/MacToolChain.cs
@@ -653,8 +653,16 @@ namespace UnrealBuildTool
 			FileReference DylibPathRef = new FileReference(DylibPath);
 			bool bIsEngineDylib = DylibPathRef.IsUnderDirectory(Unreal.EngineDirectory);
 
-			string IntermediateDirectory = (ProjectFile == null || LinkEnvironment.OutputFilePaths.All(x => x.IsUnderDirectory(Unreal.EngineDirectory)) || bIsEngineDylib ? Unreal.EngineDirectory : ProjectFile.Directory) + "/Intermediate/Mac";
-			return FileItem.GetItemByPath(Path.Combine(IntermediateDirectory, "Stubs", LinkEnvironment.Architecture.ToString(), LinkEnvironment.Configuration.ToString(), Path.GetFileName(DylibPath)));
+			string DylibFilename = Path.GetFileName(DylibPath);
+
+			string AppName = string.Empty;
+			int firstDash = DylibFilename.IndexOf('-');
+			if (firstDash != -1)
+			{
+				AppName = DylibFilename.Substring(0, firstDash);
+			}
+			string IntermediateDirectory = (ProjectFile == null || LinkEnvironment.OutputFilePaths.All(x => x.IsUnderDirectory(Unreal.EngineDirectory)) || bIsEngineDylib ? Unreal.EngineDirectory : ProjectFile.Directory) + "/Intermediate/Build/Mac";
+			return FileItem.GetItemByPath(Path.Combine(IntermediateDirectory, LinkEnvironment.Architecture.ToString(), AppName, LinkEnvironment.Configuration.ToString(), "Stubs", DylibFilename));
 		}
 
 		private FileItem LinkArchitectureFiles(LinkEnvironment LinkEnvironment, FileReference MultiArchOutputFile, bool bBuildImportLibraryOnly, IActionGraphBuilder Graph)
