[UBT]
* Fixed so temporary script file for extracting framework zip is stored in correct directory
* Fixed so link response file for IOS is written to correct directory

diff --git a/Engine/Source/Programs/UnrealBuildTool/Platform/IOS/IOSToolChain.cs b/Engine/Source/Programs/UnrealBuildTool/Platform/IOS/IOSToolChain.cs
--- a/Engine/Source/Programs/UnrealBuildTool/Platform/IOS/IOSToolChain.cs
+++ b/Engine/Source/Programs/UnrealBuildTool/Platform/IOS/IOSToolChain.cs
@@ -450,7 +450,7 @@ namespace UnrealBuildTool
 			{
 				if (Framework.ZipFile != null)
 				{
-					FileItem ExtractedTokenFile = ExtractFramework(Framework, Graph, Logger);
+					FileItem ExtractedTokenFile = ExtractFramework(Framework, Graph, LinkEnvironment.IntermediateDirectory!, Logger);
 					LinkAction.PrerequisiteItems.Add(ExtractedTokenFile);
 				}
 			}
@@ -490,9 +490,9 @@ namespace UnrealBuildTool
 			}
 			else
 			{
-				bool bIsUnrealGame = LinkEnvironment.OutputFilePath.FullName.Contains("UnrealGame");
-				FileReference ResponsePath = FileReference.Combine(((!bIsUnrealGame && ProjectFile != null) ? ProjectFile.Directory : Unreal.EngineDirectory), "Intermediate", "Build", LinkEnvironment.Platform.ToString(), "LinkFileList_" + LinkEnvironment.OutputFilePath.GetFileNameWithoutExtension() + ".tmp");
-				Graph.CreateIntermediateTextFile(ResponsePath, InputFileNames);
+				FileReference ResponsePath = FileReference.Combine(LinkEnvironment.IntermediateDirectory!, "LinkFileList_" + LinkEnvironment.OutputFilePath.GetFileNameWithoutExtension() + ".tmp");
+				FileItem fileItem = Graph.CreateIntermediateTextFile(ResponsePath, InputFileNames);
+				LinkAction.ProducedItems.Add(fileItem);
 				LinkCommandArguments += String.Format(" \\\"@{0}\\\"", ResponsePath.FullName);
 			}
 
diff --git a/Engine/Source/Programs/UnrealBuildTool/Platform/Mac/MacToolChain.cs b/Engine/Source/Programs/UnrealBuildTool/Platform/Mac/MacToolChain.cs
--- a/Engine/Source/Programs/UnrealBuildTool/Platform/Mac/MacToolChain.cs
+++ b/Engine/Source/Programs/UnrealBuildTool/Platform/Mac/MacToolChain.cs
@@ -866,7 +866,7 @@ namespace UnrealBuildTool
 				{
 					if (Framework.ZipFile != null)
 					{
-						FileItem ExtractedTokenFile = ExtractFramework(Framework, Graph, Logger);
+						FileItem ExtractedTokenFile = ExtractFramework(Framework, Graph, LinkEnvironment.IntermediateDirectory!, Logger);
 						LinkAction.PrerequisiteItems.Add(ExtractedTokenFile);
 					}
 
diff --git a/Engine/Source/Programs/UnrealBuildTool/ToolChain/AppleToolChain.cs b/Engine/Source/Programs/UnrealBuildTool/ToolChain/AppleToolChain.cs
--- a/Engine/Source/Programs/UnrealBuildTool/ToolChain/AppleToolChain.cs
+++ b/Engine/Source/Programs/UnrealBuildTool/ToolChain/AppleToolChain.cs
@@ -308,7 +308,7 @@ namespace UnrealBuildTool
 				{
 					// even in modern, we have to try to unzip in UBT, so that commandline builds extract the frameworks _before_ building (the Xcode finalize
 					// happens after building). We still need the unzip in Xcode pre-build phase, because Xcode wants the frameworks around for its build graph
-					ExtractFramework(Framework, Graph, Logger);
+					ExtractFramework(Framework, Graph, OutputDir, Logger);
 					FrameworkTokenFiles.Add(Framework.ExtractedTokenFile!);
 				}
 			}
@@ -406,7 +406,7 @@ namespace UnrealBuildTool
 			return DestFile;
 		}
 
-		protected FileItem ExtractFramework(UEBuildFramework Framework, IActionGraphBuilder Graph, ILogger Logger)
+		protected FileItem ExtractFramework(UEBuildFramework Framework, IActionGraphBuilder Graph, DirectoryReference IntermediateDirectory, ILogger Logger)
 		{
 			if (Framework.ZipFile == null)
 			{
@@ -426,7 +426,8 @@ namespace UnrealBuildTool
 				ExtractScript.AppendLine(String.Format("unzip -q -o {0} -d {1}", Utils.MakePathSafeToUseWithCommandLine(Framework.ZipFile.FullName), Utils.MakePathSafeToUseWithCommandLine(Framework.ZipOutputDirectory.ParentDirectory!.FullName))); // Zip contains folder with the same name, hence ParentDirectory
 				ExtractScript.AppendLine(String.Format("touch {0}", Utils.MakePathSafeToUseWithCommandLine(Framework.ExtractedTokenFile!.AbsolutePath)));
 
-				FileItem ExtractScriptFileItem = Graph.CreateIntermediateTextFile(new FileReference(Framework.ZipOutputDirectory.FullName + ".sh"), ExtractScript.ToString());
+				FileReference ScriptFile = FileReference.Combine(IntermediateDirectory, Framework.ZipOutputDirectory.GetDirectoryName() + ".sh");
+				FileItem ExtractScriptFileItem = Graph.CreateIntermediateTextFile(ScriptFile, ExtractScript.ToString());
 
 				Action UnzipAction = Graph.CreateAction(ActionType.BuildProject);
 				UnzipAction.CommandPath = new FileReference("/bin/sh");
