[UBT]
* Removed prefetching of generated code directories when not being the last makefile creation. Prefetching will then overlap with fake-population of generated directories and cause race conditions. Hopefully this fixes an issue seen on the farm where second run found different amount of files than first

diff --git a/Engine/Source/Programs/UnrealBuildTool/Configuration/UEBuildTarget.cs b/Engine/Source/Programs/UnrealBuildTool/Configuration/UEBuildTarget.cs
--- a/Engine/Source/Programs/UnrealBuildTool/Configuration/UEBuildTarget.cs
+++ b/Engine/Source/Programs/UnrealBuildTool/Configuration/UEBuildTarget.cs
@@ -2899,7 +2899,7 @@ namespace UnrealBuildTool
 			}
 		}
 
-		internal static void WaitCodeGen(List<CodeGenContext> contexts)
+		internal static void WaitCodeGen(List<CodeGenContext> contexts, bool prefetchAfterReset)
 		{
 			foreach (var context in contexts)
 			{
@@ -2907,8 +2907,10 @@ namespace UnrealBuildTool
 				{
 					UHTExecution.WaitForExecution(context.uhtRun, context.Makefile);
 
-					// TODO: Kick prefetch of folders
-					Task task = FileMetadataPrefetch.QueueDirectoryTrees(context.Makefile.UObjectModules.Select(m => m.GeneratedCodeDirectory));
+					if (prefetchAfterReset)
+					{
+						Task task = FileMetadataPrefetch.QueueDirectoryTrees(context.Makefile.UObjectModules.Select(m => m.GeneratedCodeDirectory));
+					}
 				}
 #if __VPROJECT_AVAILABLE__
 				if (context.vniRun != null)
diff --git a/Engine/Source/Programs/UnrealBuildTool/Modes/BuildMode.cs b/Engine/Source/Programs/UnrealBuildTool/Modes/BuildMode.cs
--- a/Engine/Source/Programs/UnrealBuildTool/Modes/BuildMode.cs
+++ b/Engine/Source/Programs/UnrealBuildTool/Modes/BuildMode.cs
@@ -1349,7 +1349,7 @@ namespace UnrealBuildTool
 							// Make sure to wait for previous pass codegen and timestamp writes
 							if (lastContextsForCodeGen != null)
 							{
-								UEBuildTarget.WaitCodeGen(lastContextsForCodeGen);
+								UEBuildTarget.WaitCodeGen(lastContextsForCodeGen, false);
 								Task.WaitAll(UEBuildTarget.GetTimestampTasks(lastContextsForCodeGen).ToArray());
 							}
 
@@ -1451,7 +1451,7 @@ namespace UnrealBuildTool
 					if (lastContextsForCodeGen != null)
 					{
 						List<Task> timestampTasks = UEBuildTarget.GetTimestampTasks(lastContextsForCodeGen);
-						UEBuildTarget.WaitCodeGen(lastContextsForCodeGen);
+						UEBuildTarget.WaitCodeGen(lastContextsForCodeGen, true); // We want to prefetch generated dirs since actiongraph will read file info
 
 						lock (outTasks)
 						{
