[UBT]
* Fixed so GenerateExport action is using vfs paths

diff --git a/Engine/Source/Programs/UnrealBuildTool/Configuration/UEBuildTarget.cs b/Engine/Source/Programs/UnrealBuildTool/Configuration/UEBuildTarget.cs
--- a/Engine/Source/Programs/UnrealBuildTool/Configuration/UEBuildTarget.cs
+++ b/Engine/Source/Programs/UnrealBuildTool/Configuration/UEBuildTarget.cs
@@ -3405,21 +3405,23 @@ namespace UnrealBuildTool
 			HashSet<FileItem> allObjects = new(Binaries.SelectMany(x => x.InputObjects));
 			SortedSet<FileItem> allExis = new(allObjects.Select(x => FileItem.GetItemByFileReference(x.Location.ChangeExtension(".exi"))), FileItem.SortByPath);
 
-			void CreateStripAction(string Name, UEBuildBinaryType binaryType, IEnumerable< FileItem> toStrip, DirectoryReference IntermediateDirectory)
+			void CreateGenerateExportsAction(string Name, UEBuildBinaryType binaryType, IEnumerable< FileItem> toStrip, DirectoryReference IntermediateDirectory)
 			{
 				SortedSet<FileItem> toStripExis = new(toStrip.Select(x => FileItem.GetItemByFileReference(x.Location.ChangeExtension(".exi"))), FileItem.SortByPath);
 				SortedSet<FileItem> otherExis = new(allExis, FileItem.SortByPath);
 				otherExis.ExceptWith(toStripExis);
 
+				string Normalize(FileReference file) => TargetToolChain.NormalizeCommandLinePath(file, CompileEnvironment.RootPaths);
+
 				FileItem exportsRsp = FileItem.GetItemByFileReference(FileReference.Combine(IntermediateDirectory, $"{Name}.exports.rsp"));
 				string Ext = TargetToolChain.GetExtraLinkFileExtension();
 				FileReference outputFile = FileReference.Combine(IntermediateDirectory, $"{Name}.exports.{Ext}");
 
 				IEnumerable<string> arguments =
 				[
-					.. toStripExis.Select(path => $"/S:{path}"),
-					.. otherExis.Select(path => $"/D:{path}"),
-					$"/O:{outputFile}",
+					.. toStripExis.Select(path => $"/S:{Normalize(path.Location)}"),
+					.. otherExis.Select(path => $"/D:{Normalize(path.Location)}"),
+					$"/O:{Normalize(outputFile)}",
 					$"/T:{Platform}",
 					$"/M:{Name}",
 					..TargetToolChain.GetExtraLinkFileAdditionalSymbols(binaryType).Select(symbol=> $"/E:{symbol}"),
@@ -3433,7 +3435,7 @@ namespace UnrealBuildTool
 				exportsAction.PrerequisiteItems.Add(exportsRsp);
 				exportsAction.ProducedItems.Add(FileItem.GetItemByFileReference(outputFile));
 				exportsAction.CommandPath = commandPath;
-				exportsAction.CommandArguments = $"@\"{TargetToolChain.NormalizeCommandLinePath(exportsRsp.Location, CompileEnvironment.RootPaths)}\"";
+				exportsAction.CommandArguments = $"@\"{Normalize(exportsRsp.Location)}\"";
 				exportsAction.CommandDescription = "GenerateExports";
 				exportsAction.WorkingDirectory = Unreal.EngineSourceDirectory;
 				exportsAction.StatusDescription = Name;
@@ -3447,7 +3449,7 @@ namespace UnrealBuildTool
 
 			foreach (UEBuildBinary Binary in Binaries.Where(x => x.bStripUnusedExports))
 			{
-				CreateStripAction(Binary.OutputFilePaths.First().GetFileNameWithoutExtension(), Binary.Type, Binary.InputObjects, Binary.IntermediateDirectory);
+				CreateGenerateExportsAction(Binary.OutputFilePaths.First().GetFileNameWithoutExtension(), Binary.Type, Binary.InputObjects, Binary.IntermediateDirectory);
 			}
 		}
 
