Replaced PREPROCESSOR_NOTHING and PREPROCESSOR_NOTHING_FUNCTION with UE_EMPTY and UE_EMPTY_FUNCTION.

diff --git a/Engine/Source/Programs/UnrealBuildTool/Configuration/UEBuildModuleCPP.cs b/Engine/Source/Programs/UnrealBuildTool/Configuration/UEBuildModuleCPP.cs
--- a/Engine/Source/Programs/UnrealBuildTool/Configuration/UEBuildModuleCPP.cs
+++ b/Engine/Source/Programs/UnrealBuildTool/Configuration/UEBuildModuleCPP.cs
@@ -1567,8 +1567,8 @@ namespace UnrealBuildTool
 						Definitions.RemoveAll(x => x.Contains("UE_DEPRECATED_FORENGINE", StringComparison.Ordinal));
 
 						Definitions.Add($"UE_IS_ENGINE_MODULE={(Rules.bTreatAsEngineModule || Rules.Target.bWarnAboutMonolithicHeadersIncluded ? "1" : "0")}");
-						Definitions.Add($"UE_DEPRECATED_FORGAME={(Rules.bTreatAsEngineModule ? "PREPROCESSOR_NOTHING_FUNCTION" : "UE_DEPRECATED")}");
-						Definitions.Add($"UE_DEPRECATED_FORENGINE={(Rules.bTreatAsEngineModule || Rules.Target.bDisableEngineDeprecations ? "PREPROCESSOR_NOTHING_FUNCTION" : "UE_DEPRECATED")}");
+						Definitions.Add($"UE_DEPRECATED_FORGAME={(Rules.bTreatAsEngineModule ? "UE_EMPTY_FUNCTION" : "UE_DEPRECATED")}");
+						Definitions.Add($"UE_DEPRECATED_FORENGINE={(Rules.bTreatAsEngineModule || Rules.Target.bDisableEngineDeprecations ? "UE_EMPTY_FUNCTION" : "UE_DEPRECATED")}");
 					}
 
 					// Modify definitions if we need to create a new shared pch for validating format strings
@@ -2612,7 +2612,7 @@ namespace UnrealBuildTool
 			if (!Rules.bTreatAsEngineModule)
 			{
 				Result.Definitions.Add("UE_DEPRECATED_FORGAME=UE_DEPRECATED");
-				Result.Definitions.Add($"UE_DEPRECATED_FORENGINE={(Rules.Target.bDisableEngineDeprecations ? "PREPROCESSOR_NOTHING_FUNCTION" : "UE_DEPRECATED")}");
+				Result.Definitions.Add($"UE_DEPRECATED_FORENGINE={(Rules.Target.bDisableEngineDeprecations ? "UE_EMPTY_FUNCTION" : "UE_DEPRECATED")}");
 			}
 
 			Result.Definitions.Add($"UE_VALIDATE_FORMAT_STRINGS={(Rules.bValidateFormatStrings ? "1" : "0")}");
@@ -2739,7 +2739,7 @@ namespace UnrealBuildTool
 			if (!Rules.bTreatAsEngineModule)
 			{
 				CompileEnvironment.Definitions.Add("UE_DEPRECATED_FORGAME=UE_DEPRECATED");
-				CompileEnvironment.Definitions.Add($"UE_DEPRECATED_FORENGINE={(Rules.Target.bDisableEngineDeprecations ? "PREPROCESSOR_NOTHING_FUNCTION" : "UE_DEPRECATED")}");
+				CompileEnvironment.Definitions.Add($"UE_DEPRECATED_FORENGINE={(Rules.Target.bDisableEngineDeprecations ? "UE_EMPTY_FUNCTION" : "UE_DEPRECATED")}");
 			}
 
 			CompileEnvironment.Definitions.Add($"UE_VALIDATE_FORMAT_STRINGS={(Rules.bValidateFormatStrings ? "1" : "0")}");
diff --git a/Engine/Source/Runtime/Core/Public/Misc/CoreMiscDefines.h b/Engine/Source/Runtime/Core/Public/Misc/CoreMiscDefines.h
--- a/Engine/Source/Runtime/Core/Public/Misc/CoreMiscDefines.h
+++ b/Engine/Source/Runtime/Core/Public/Misc/CoreMiscDefines.h
@@ -372,7 +372,7 @@ enum EPerElement {PerElement};
 #endif // !defined(__INTELLISENSE__)
 
 #ifndef UE_DEPRECATED_FORGAME
-	#define UE_DEPRECATED_FORGAME PREPROCESSOR_NOTHING_FUNCTION
+	#define UE_DEPRECATED_FORGAME UE_EMPTY_FUNCTION
 #endif
 
 #ifndef UE_DEPRECATED_FORENGINE
diff --git a/Engine/Source/Runtime/Core/Public/Serialization/MemoryLayout.h b/Engine/Source/Runtime/Core/Public/Serialization/MemoryLayout.h
--- a/Engine/Source/Runtime/Core/Public/Serialization/MemoryLayout.h
+++ b/Engine/Source/Runtime/Core/Public/Serialization/MemoryLayout.h
@@ -538,7 +538,7 @@ PRAGMA_ENABLE_DEPRECATION_WARNINGS
 			FTypeLayoutDesc::Initialize(TypeDesc); \
 		} \
 		return TypeDesc; } \
-	public: INTERNAL_LAYOUT_INTERFACE_PREFIX(InInterface)(PREPROCESSOR_NOTHING) const FTypeLayoutDesc& GetTypeLayout() const INTERNAL_LAYOUT_INTERFACE_INLINE_IMPL(InInterface) \
+	public: INTERNAL_LAYOUT_INTERFACE_PREFIX(InInterface)(UE_EMPTY) const FTypeLayoutDesc& GetTypeLayout() const INTERNAL_LAYOUT_INTERFACE_INLINE_IMPL(InInterface) \
 	INTERNAL_DECLARE_TYPE_LAYOUT_COMMON(T, InInterface)
 
 #define INTERNAL_DECLARE_TYPE_LAYOUT(T, InInterface, RequiredAPI) \
@@ -554,11 +554,11 @@ PRAGMA_ENABLE_DEPRECATION_WARNINGS
 #define INTERNAL_DECLARE_LAYOUT_EXPLICIT_BASES(T, ...) \
 	template<typename InternalType> static void InternalInitializeBases(FTypeLayoutDesc& TypeDesc) { InternalInitializeBasesHelper<InternalType, __VA_ARGS__>(TypeDesc); }
 
-#define DECLARE_TYPE_LAYOUT(T, Interface) INTERNAL_DECLARE_LAYOUT_BASE(T); INTERNAL_DECLARE_TYPE_LAYOUT(T, Interface, PREPROCESSOR_NOTHING)
+#define DECLARE_TYPE_LAYOUT(T, Interface) INTERNAL_DECLARE_LAYOUT_BASE(T); INTERNAL_DECLARE_TYPE_LAYOUT(T, Interface, UE_EMPTY)
 #define DECLARE_INLINE_TYPE_LAYOUT(T, Interface) INTERNAL_DECLARE_LAYOUT_BASE(T); INTERNAL_DECLARE_INLINE_TYPE_LAYOUT(T, Interface)
 #define DECLARE_EXPORTED_TYPE_LAYOUT(T, RequiredAPI, Interface) INTERNAL_DECLARE_LAYOUT_BASE(T); INTERNAL_DECLARE_TYPE_LAYOUT(T, Interface, RequiredAPI)
 
-#define DECLARE_TYPE_LAYOUT_EXPLICIT_BASES(T, Interface, ...) INTERNAL_DECLARE_LAYOUT_EXPLICIT_BASES(T, __VA_ARGS__); INTERNAL_DECLARE_TYPE_LAYOUT(T, Interface, PREPROCESSOR_NOTHING)
+#define DECLARE_TYPE_LAYOUT_EXPLICIT_BASES(T, Interface, ...) INTERNAL_DECLARE_LAYOUT_EXPLICIT_BASES(T, __VA_ARGS__); INTERNAL_DECLARE_TYPE_LAYOUT(T, Interface, UE_EMPTY)
 #define DECLARE_INLINE_TYPE_LAYOUT_EXPLICIT_BASES(T, Interface, ...) INTERNAL_DECLARE_LAYOUT_EXPLICIT_BASES(T, __VA_ARGS__); INTERNAL_DECLARE_INLINE_TYPE_LAYOUT(T, Interface)
 #define DECLARE_EXPORTED_TYPE_LAYOUT_EXPLICIT_BASES(T, RequiredAPI, Interface, ...) INTERNAL_DECLARE_LAYOUT_EXPLICIT_BASES(T, __VA_ARGS__); INTERNAL_DECLARE_TYPE_LAYOUT(T, Interface, RequiredAPI)
 
diff --git a/Engine/Source/Runtime/CoreUObject/Public/UObject/UObjectGlobals.h b/Engine/Source/Runtime/CoreUObject/Public/UObject/UObjectGlobals.h
--- a/Engine/Source/Runtime/CoreUObject/Public/UObject/UObjectGlobals.h
+++ b/Engine/Source/Runtime/CoreUObject/Public/UObject/UObjectGlobals.h
@@ -3540,7 +3540,7 @@ struct FAssetMsg
 		static_assert((ELogVerbosity::Verbosity & ELogVerbosity::VerbosityMask) < ELogVerbosity::NumVerbosity && ELogVerbosity::Verbosity > 0, "Verbosity must be constant and in range."); \
 		CA_CONSTANT_IF((ELogVerbosity::Verbosity & ELogVerbosity::VerbosityMask) <= ELogVerbosity::COMPILED_IN_MINIMUM_VERBOSITY && (ELogVerbosity::Warning & ELogVerbosity::VerbosityMask) <= FLogCategory##CategoryName::CompileTimeVerbosity) \
 		{ \
-			UE_LOG_EXPAND_IS_FATAL(Verbosity, PREPROCESSOR_NOTHING, if (!CategoryName.IsSuppressed(ELogVerbosity::Verbosity))) \
+			UE_LOG_EXPAND_IS_FATAL(Verbosity, UE_EMPTY, if (!CategoryName.IsSuppressed(ELogVerbosity::Verbosity))) \
 			{ \
 				::UE::Core::Private::RecordAssetLog(CategoryName.GetCategoryName(), ELogVerbosity::Verbosity, ::UE::FAssetLog(Asset), FString::Printf(Format, ##__VA_ARGS__), __FILE__, __LINE__); \
 				UE_LOG_EXPAND_IS_FATAL(Verbosity, \
@@ -3549,7 +3549,7 @@ struct FAssetMsg
 						FDebug::AssertFailed("", __FILE__, __LINE__, TEXT("%s: %s"), *FAssetMsg::FormatPathForAssetLog(Asset), *FString::Printf(Format, ##__VA_ARGS__)); \
 						CA_ASSUME(false); \
 					}, \
-					PREPROCESSOR_NOTHING \
+					UE_EMPTY \
 				) \
 			} \
 		} \
