[UBT]
* Backed out part of change where multiple actions with different command line write to same file

diff --git a/Engine/Source/Programs/UnrealBuildTool/ToolChain/AppleToolChain.cs b/Engine/Source/Programs/UnrealBuildTool/ToolChain/AppleToolChain.cs
--- a/Engine/Source/Programs/UnrealBuildTool/ToolChain/AppleToolChain.cs
+++ b/Engine/Source/Programs/UnrealBuildTool/ToolChain/AppleToolChain.cs
@@ -426,8 +426,10 @@ namespace UnrealBuildTool
 				ExtractScript.AppendLine(String.Format("unzip -q -o {0} -d {1}", Utils.MakePathSafeToUseWithCommandLine(Framework.ZipFile.FullName), Utils.MakePathSafeToUseWithCommandLine(Framework.ZipOutputDirectory.ParentDirectory!.FullName))); // Zip contains folder with the same name, hence ParentDirectory
 				ExtractScript.AppendLine(String.Format("touch {0}", Utils.MakePathSafeToUseWithCommandLine(Framework.ExtractedTokenFile!.AbsolutePath)));
 
-				FileReference ScriptFile = FileReference.Combine(IntermediateDirectory, Framework.ZipOutputDirectory.GetDirectoryName() + ".sh");
-				FileItem ExtractScriptFileItem = Graph.CreateIntermediateTextFile(ScriptFile, ExtractScript.ToString());
+				// This is writing to folder shared by all configs.
+				//FileReference ScriptFile = FileReference.Combine(IntermediateDirectory, Framework.ZipOutputDirectory.GetDirectoryName() + ".sh");
+				//FileItem ExtractScriptFileItem = Graph.CreateIntermediateTextFile(ScriptFile, ExtractScript.ToString());
+				FileItem ExtractScriptFileItem = Graph.CreateIntermediateTextFile(new FileReference(Framework.ZipOutputDirectory.FullName + ".sh"), ExtractScript.ToString());
 
 				Action UnzipAction = Graph.CreateAction(ActionType.BuildProject);
 				UnzipAction.CommandPath = new FileReference("/bin/sh");
