[UBT]
* Various optimizations based on profiling code

diff --git a/Engine/Source/Programs/UnrealBuildTool/Configuration/UEBuildModule.cs b/Engine/Source/Programs/UnrealBuildTool/Configuration/UEBuildModule.cs
--- a/Engine/Source/Programs/UnrealBuildTool/Configuration/UEBuildModule.cs
+++ b/Engine/Source/Programs/UnrealBuildTool/Configuration/UEBuildModule.cs
@@ -1118,25 +1118,25 @@ namespace UnrealBuildTool
 		/// <param name="bIncludePrivateModules">Whether to include private modules.</param>
 		private void InternalGetAllDependencyModulesForPCH(HashSet<UEBuildModule> ReferencedModules, HashSet<UEBuildModule> IgnoreReferencedModules, bool bIncludePrivateModules)
 		{
-			List<UEBuildModule> AllDependencyModules = new List<UEBuildModule>(
+			HashSet<UEBuildModule> AllDependencyModules = new(
 				((bIncludePrivateModules && PrivateDependencyModules != null) ? PrivateDependencyModules.Count : 0) +
 				(PublicDependencyModules != null ? PublicDependencyModules.Count : 0) +
 				(PublicIncludePathModules != null ? PublicIncludePathModules!.Count : 0)
 				);
 			if (bIncludePrivateModules && PrivateDependencyModules != null)
 			{
-				AllDependencyModules.AddRange(PrivateDependencyModules);
+				AllDependencyModules.UnionWith(PrivateDependencyModules);
 			}
 			if (PublicDependencyModules != null)
 			{
-				AllDependencyModules.AddRange(PublicDependencyModules!);
+				AllDependencyModules.UnionWith(PublicDependencyModules!);
 			}
 			if (PublicIncludePathModules != null)
 			{
-				AllDependencyModules.AddRange(PublicIncludePathModules);
+				AllDependencyModules.UnionWith(PublicIncludePathModules);
 			}
 
-			foreach (UEBuildModule DependencyModule in AllDependencyModules.Distinct())
+			foreach (UEBuildModule DependencyModule in AllDependencyModules)
 			{
 				// Don't follow circular back-references!
 				if (!HasCircularDependencyOn(DependencyModule))
diff --git a/Engine/Source/Programs/UnrealBuildTool/Configuration/UEBuildModuleCPP.cs b/Engine/Source/Programs/UnrealBuildTool/Configuration/UEBuildModuleCPP.cs
--- a/Engine/Source/Programs/UnrealBuildTool/Configuration/UEBuildModuleCPP.cs
+++ b/Engine/Source/Programs/UnrealBuildTool/Configuration/UEBuildModuleCPP.cs
@@ -83,6 +83,11 @@ namespace UnrealBuildTool
 		/// </summary>
 		public readonly DirectoryReference? GeneratedCodeDirectoryVNI;
 
+		/// <summary>
+		/// Module directory. This is just here to reduce number of string creation since it is used in 
+		/// </summary>
+		private readonly DirectoryReference ModuleInterfaceDir;
+
 		/// <summary>
 		/// Global override to force all include paths to be always added
 		/// </summary>
@@ -263,6 +268,8 @@ namespace UnrealBuildTool
 				GeneratedCodeDirectoryVNI = DirectoryReference.Combine(GeneratedCodeDirectory!, "VNI");
 			}
 
+			ModuleInterfaceDir = UEToolChain.GetModuleInterfaceDir(IntermediateDirectory);
+
 			// Check for a Verse directory next to the rules file
 			DirectoryItem RulesDirectory = DirectoryItem.GetItemByDirectoryReference(Rules.File.Directory);
 			if (RulesDirectory.TryGetDirectory("Verse", out DirectoryItem? MaybeVerseDirectory))
@@ -468,7 +475,7 @@ namespace UnrealBuildTool
 				}
 			}
 
-			ModuleInterfacePaths.Add(UEToolChain.GetModuleInterfaceDir(IntermediateDirectory));
+			ModuleInterfacePaths.Add(ModuleInterfaceDir);
 
 			base.AddModuleToCompileEnvironment(SourceModule, SourceBinary, IncludePaths, SystemIncludePaths, ModuleInterfacePaths, Definitions, AdditionalFrameworks, AutoRTFMExternalMappingFiles, AdditionalPrerequisites, bLegacyPublicIncludePaths, bLegacyParentIncludePaths);
 		}
diff --git a/Engine/Source/Programs/UnrealBuildTool/System/CppCompileEnvironment.cs b/Engine/Source/Programs/UnrealBuildTool/System/CppCompileEnvironment.cs
--- a/Engine/Source/Programs/UnrealBuildTool/System/CppCompileEnvironment.cs
+++ b/Engine/Source/Programs/UnrealBuildTool/System/CppCompileEnvironment.cs
@@ -349,7 +349,7 @@ namespace UnrealBuildTool
 		/// <summary>
 		/// AutoRTFM external open -> closed mapping files
 		/// </summary>
-		public HashSet<FileItem> AutoRTFMExternalMappingFiles = [];
+		public HashSet<FileItem> AutoRTFMExternalMappingFiles;
 
 		/// <summary>
 		/// Whether we should run an extra LLVM verification pass after the AutoRTFM compiler pass.
@@ -436,27 +436,27 @@ namespace UnrealBuildTool
 		/// The static analyzer rulesets that should be used to filter warnings. This is only supported for MSVC.
 		/// See https://learn.microsoft.com/en-us/cpp/code-quality/using-rule-sets-to-specify-the-cpp-rules-to-run
 		/// </summary>
-		public HashSet<FileReference> StaticAnalyzerRulesets = [];
+		public HashSet<FileReference> StaticAnalyzerRulesets;
 
 		/// <summary>
 		/// The static analyzer checkers that should be enabled rather than the defaults. This is only supported for Clang.
 		/// </summary>
-		public HashSet<string> StaticAnalyzerCheckers = [];
+		public HashSet<string> StaticAnalyzerCheckers;
 
 		/// <summary>
 		/// The static analyzer default checkers that should be disabled. Unused if StaticAnalyzerCheckers is populated. This is only supported for Clang.
 		/// </summary>
-		public HashSet<string> StaticAnalyzerDisabledCheckers = [];
+		public HashSet<string> StaticAnalyzerDisabledCheckers;
 
 		/// <summary>
 		/// The static analyzer non-default checkers that should be enabled. Unused if StaticAnalyzerCheckers is populated. This is only supported for Clang.
 		/// </summary>
-		public HashSet<string> StaticAnalyzerAdditionalCheckers = [];
+		public HashSet<string> StaticAnalyzerAdditionalCheckers;
 
 		/// <summary>
 		/// The PVS Studio analysis warnings that should be disabled.
 		/// </summary>
-		public HashSet<string> StaticAnalyzerPVSDisabledErrors = [];
+		public HashSet<string> StaticAnalyzerPVSDisabledErrors;
 
 		/// <summary>
 		/// Whether to disable all dynamic analysis (such as Address Sanitizer)
@@ -636,22 +636,22 @@ namespace UnrealBuildTool
 		/// <summary>
 		/// List of header files to force include
 		/// </summary>
-		public List<FileItem> ForceIncludeFiles = [];
+		public List<FileItem> ForceIncludeFiles;
 
 		/// <summary>
 		/// List of files that need to be up to date before compile can proceed
 		/// </summary>
-		public List<FileItem> AdditionalPrerequisites = [];
+		public List<FileItem> AdditionalPrerequisites;
 
 		/// <summary>
 		/// A dictionary of the source file items and the inlined gen.cpp files contained in it
 		/// </summary>
-		public Dictionary<FileItem, HashSet<FileItem>> FileInlineSourceMap = [];
+		public Dictionary<FileItem, HashSet<FileItem>> FileInlineSourceMap;
 		
 		/// <summary>
 		/// Non-default naming conventions that generated CPP files can have (this uses a filter to check the file name for matches) 
 		/// </summary>
-		public List<string> ExtraGeneratedCPPFileTypes = [];
+		public List<string> ExtraGeneratedCPPFileTypes;
 
 		/// <summary>
 		/// FileItems with colliding names. (Which means they would overwrite each other in intermediate folder
@@ -661,12 +661,12 @@ namespace UnrealBuildTool
 		/// <summary>
 		/// The C++ preprocessor definitions to use.
 		/// </summary>
-		public List<string> Definitions = [];
+		public List<string> Definitions;
 
 		/// <summary>
 		/// Additional response files that will be used by main response file
 		/// </summary>
-		public List<FileItem> AdditionalResponseFiles = [];
+		public List<FileItem> AdditionalResponseFiles;
 
 		/// <summary>
 		/// Whether the compile environment has a response file in AdditionalResponseFiles that contains global compiler arguments.
@@ -681,7 +681,7 @@ namespace UnrealBuildTool
 		/// <summary>
 		/// A list of additional frameworks whose include paths are needed.
 		/// </summary>
-		public List<UEBuildFramework> AdditionalFrameworks = [];
+		public List<UEBuildFramework> AdditionalFrameworks;
 
 		/// <summary>
 		/// A dictionary of PCH files for multiple architectures
@@ -802,7 +802,7 @@ namespace UnrealBuildTool
 		/// File to use as an LLVM FPass Plugin:
 		///    'clang -fpass-plugin=[FPassPlugin]' 
 		/// </summary>
-		public List<string> FPassPlugins = [];
+		public List<string> FPassPlugins;
 
 		/// <summary>
 		/// Default constructor.
@@ -820,6 +820,20 @@ namespace UnrealBuildTool
 			SharedUserIncludePaths = [];
 			SharedSystemIncludePaths = [];
 			ModuleInterfacePaths = [];
+			ForceIncludeFiles = [];
+			AdditionalPrerequisites = [];
+			FileInlineSourceMap = [];
+			Definitions = [];
+			AdditionalResponseFiles = [];
+			AdditionalFrameworks = [];
+			AutoRTFMExternalMappingFiles = [];
+			ExtraGeneratedCPPFileTypes = [];
+			StaticAnalyzerRulesets = [];
+			StaticAnalyzerCheckers = [];
+			StaticAnalyzerDisabledCheckers = [];
+			StaticAnalyzerAdditionalCheckers = [];
+			StaticAnalyzerPVSDisabledErrors = [];
+			FPassPlugins = [];
 			CppCompileWarnings = new CppCompileWarnings();
 		}
 
@@ -856,17 +870,17 @@ namespace UnrealBuildTool
 			CppCompileWarnings = CppCompileWarnings.CreateShallowCopy(Other.CppCompileWarnings);
 			bWarningsAsErrors = Other.bWarningsAsErrors;
 			bDisableStaticAnalysis = Other.bDisableStaticAnalysis;
-			StaticAnalyzerCheckers = [.. Other.StaticAnalyzerCheckers];
-			StaticAnalyzerDisabledCheckers = [.. Other.StaticAnalyzerDisabledCheckers];
-			StaticAnalyzerAdditionalCheckers = [.. Other.StaticAnalyzerAdditionalCheckers];
-			StaticAnalyzerPVSDisabledErrors = [.. Other.StaticAnalyzerPVSDisabledErrors];
+			StaticAnalyzerCheckers = new(Other.StaticAnalyzerCheckers);
+			StaticAnalyzerDisabledCheckers = new(Other.StaticAnalyzerDisabledCheckers);
+			StaticAnalyzerAdditionalCheckers = new(Other.StaticAnalyzerAdditionalCheckers);
+			StaticAnalyzerPVSDisabledErrors = new(Other.StaticAnalyzerPVSDisabledErrors);
 			bStaticAnalyzerExtensions = Other.bStaticAnalyzerExtensions;
-			StaticAnalyzerRulesets = [.. Other.StaticAnalyzerRulesets];
+			StaticAnalyzerRulesets = new(Other.StaticAnalyzerRulesets);
 			bOptimizeCode = Other.bOptimizeCode;
 			bUseAutoRTFMCompiler = Other.bUseAutoRTFMCompiler;
 			bDisableAutoRTFMInstrumentation = Other.bDisableAutoRTFMInstrumentation;
 			bEnableAutoRTFMVerification = Other.bEnableAutoRTFMVerification;
-			AutoRTFMExternalMappingFiles = [.. Other.AutoRTFMExternalMappingFiles];
+			AutoRTFMExternalMappingFiles = new(Other.AutoRTFMExternalMappingFiles);
 			bAutoRTFMVerify = Other.bAutoRTFMVerify;
 			bAutoRTFMClosedStaticLinkage = Other.bAutoRTFMClosedStaticLinkage;
 			bCodeCoverage = Other.bCodeCoverage;
@@ -896,21 +910,21 @@ namespace UnrealBuildTool
 			bAllowRemotelyCompiledPCHs = Other.bAllowRemotelyCompiledPCHs;
 			bHasSharedResponseFile = Other.bHasSharedResponseFile;
 			bUseHeaderUnitsForPch = Other.bUseHeaderUnitsForPch;
-			UserIncludePaths = [.. Other.UserIncludePaths];
-			SystemIncludePaths = [.. Other.SystemIncludePaths];
-			SharedUserIncludePaths = [.. Other.SharedUserIncludePaths];
-			SharedSystemIncludePaths = [.. Other.SharedSystemIncludePaths];
-			ModuleInterfacePaths = [.. Other.ModuleInterfacePaths];
+			UserIncludePaths = new(Other.UserIncludePaths);
+			SystemIncludePaths = new(Other.SystemIncludePaths);
+			SharedUserIncludePaths = new(Other.SharedUserIncludePaths);
+			SharedSystemIncludePaths = new(Other.SharedSystemIncludePaths);
+			ModuleInterfacePaths = new(Other.ModuleInterfacePaths);
 			bCheckSystemHeadersForModification = Other.bCheckSystemHeadersForModification;
-			ForceIncludeFiles.AddRange(Other.ForceIncludeFiles);
-			AdditionalPrerequisites.AddRange(Other.AdditionalPrerequisites);
+			ForceIncludeFiles = new(Other.ForceIncludeFiles);
+			AdditionalPrerequisites = new(Other.AdditionalPrerequisites);
 			ExtraGeneratedCPPFileTypes = Other.ExtraGeneratedCPPFileTypes;
 			CollidingNames = Other.CollidingNames;
 			FileInlineSourceMap = new(Other.FileInlineSourceMap);
-			Definitions.AddRange(Other.Definitions);
-			AdditionalResponseFiles.AddRange(Other.AdditionalResponseFiles);
+			Definitions = new(Other.Definitions);
+			AdditionalResponseFiles = new(Other.AdditionalResponseFiles);
 			AdditionalArguments = Other.AdditionalArguments;
-			AdditionalFrameworks.AddRange(Other.AdditionalFrameworks);
+			AdditionalFrameworks = new(Other.AdditionalFrameworks);
 			PCHInstance = Other.PCHInstance;
 			ParentPCHInstance = Other.ParentPCHInstance;
 			bHackHeaderGenerator = Other.bHackHeaderGenerator;
