[UBT]
* Added retry in UBAExecutor if remote helper gets a linker error saying it can't overwrite original file. We have seen this once over millions of runs on the farm and can't figure out how this can happen (could very well be a linux/wine bug)

diff --git a/Engine/Source/Programs/UnrealBuildTool/Executors/UnrealBuildAccelerator/UBAExecutor.cs b/Engine/Source/Programs/UnrealBuildTool/Executors/UnrealBuildAccelerator/UBAExecutor.cs
--- a/Engine/Source/Programs/UnrealBuildTool/Executors/UnrealBuildAccelerator/UBAExecutor.cs
+++ b/Engine/Source/Programs/UnrealBuildTool/Executors/UnrealBuildAccelerator/UBAExecutor.cs
@@ -696,7 +696,7 @@ namespace UnrealBuildTool
 
 							string executingHost = info.ExecutingHost ?? "Unknown";
 
-							var rerun = (string error) =>
+							ProcessFinishedResponse rerun(string error)
 							{
 								logger.LogInformation("{Description} {StatusDescription} [RemoteExecutor: {ExecutingHost}]: Exited with error code {ExitCode} ({Error}). This action will retry locally", action.CommandDescription, action.StatusDescription, executingHost, exitCode, error);
 								foreach (string line in logLines)
@@ -728,6 +728,11 @@ namespace UnrealBuildTool
 							{
 								return rerun("C1001");
 							}
+							else if (exitCode != 0 && logLines.Any(x => x.Contains("cannot overwrite the original file", StringComparison.Ordinal)))
+							{
+								// We have seen this once on the farm and have no idea how it can happen.
+								return rerun("Cannot overwrite original file");
+							}
 							else if (exitCode >= 9000 && exitCode < 10000)
 							{
 								return rerun("UBA error");
