[UBT]
* Moved some noisy logging into log file only so user don't have to see it all the time

diff --git a/Engine/Source/Programs/UnrealBuildTool/Modes/BuildMode.cs b/Engine/Source/Programs/UnrealBuildTool/Modes/BuildMode.cs
--- a/Engine/Source/Programs/UnrealBuildTool/Modes/BuildMode.cs
+++ b/Engine/Source/Programs/UnrealBuildTool/Modes/BuildMode.cs
@@ -1033,7 +1033,7 @@ namespace UnrealBuildTool
 			{
 				foreach (string Diagnostic in Makefiles[OutputIndices[0]].Diagnostics)
 				{
-					Logger.LogInformation("{Diagnostic}", Diagnostic);
+					Logger.LogDebug("{Diagnostic}", Diagnostic);
 				}
 			}
 			else
@@ -1042,7 +1042,7 @@ namespace UnrealBuildTool
 				{
 					foreach (string Diagnostic in Makefiles[OutputIndex].Diagnostics)
 					{
-						Logger.LogInformation("{Name}: {Diagnostic}", TargetDescriptors[OutputIndex].Name, Diagnostic);
+						Logger.LogDebug("{Name}: {Diagnostic}", TargetDescriptors[OutputIndex].Name, Diagnostic);
 					}
 				}
 			}
diff --git a/Engine/Source/Programs/UnrealBuildTool/Modes/UnrealHeaderToolMode.cs b/Engine/Source/Programs/UnrealBuildTool/Modes/UnrealHeaderToolMode.cs
--- a/Engine/Source/Programs/UnrealBuildTool/Modes/UnrealHeaderToolMode.cs
+++ b/Engine/Source/Programs/UnrealBuildTool/Modes/UnrealHeaderToolMode.cs
@@ -796,8 +796,10 @@ namespace UnrealBuildTool.Modes
 				stopwatch.Stop();
 				if (!Session.HasErrors || Session.TotalFilesWritten != 0)
 				{
+					double elapsed = stopwatch.Elapsed.TotalSeconds;
 					string ActualTargetName = Session.Manifest != null ? Session.Manifest.TargetName : "UnknownTarget";
-					Logger.LogInformation("UHT processed {TargetName} in {Time} seconds ({TotalFilesWritten} files written)", ActualTargetName, stopwatch.Elapsed.TotalSeconds, Session.TotalFilesWritten);
+					LogLevel level = elapsed > 5 ? LogLevel.Information : LogLevel.Debug;
+					Logger.Log(level, "UHT processed {TargetName} in {Time} seconds ({TotalFilesWritten} files written)", ActualTargetName, elapsed, Session.TotalFilesWritten);
 				}
 
 				return (int)(Session.HasErrors ? CompilationResult.OtherCompilationError : CompilationResult.Succeeded);
diff --git a/Engine/Source/Programs/UnrealBuildTool/System/TargetMakefile.cs b/Engine/Source/Programs/UnrealBuildTool/System/TargetMakefile.cs
--- a/Engine/Source/Programs/UnrealBuildTool/System/TargetMakefile.cs
+++ b/Engine/Source/Programs/UnrealBuildTool/System/TargetMakefile.cs
@@ -223,7 +223,7 @@ namespace UnrealBuildTool
 		/// <summary>
 		/// Additional diagnostic output to print before building this target (toolchain version, etc...)
 		/// </summary>
-		public List<string> Diagnostics = new List<string>();
+		public List<string> Diagnostics;
 
 		/// <summary>
 		/// Any additional information about the build environment which the platform can use to invalidate the makefile
diff --git a/Engine/Source/Programs/UnrealBuildTool/System/UHTExecution.cs b/Engine/Source/Programs/UnrealBuildTool/System/UHTExecution.cs
--- a/Engine/Source/Programs/UnrealBuildTool/System/UHTExecution.cs
+++ b/Engine/Source/Programs/UnrealBuildTool/System/UHTExecution.cs
@@ -1359,7 +1359,7 @@ namespace UnrealBuildTool
 
 						WriteUHTManifest(Makefile, TargetName, ModuleInfoFileName, ExternalDependenciesFile);
 						string ActualTargetName = String.IsNullOrEmpty(TargetName) ? "UE5" : TargetName;
-						Logger.LogInformation("Parsing headers for {ActualTargetName}", ActualTargetName);
+						Logger.LogDebug("Parsing headers for {ActualTargetName}", ActualTargetName);
 
 						// Generate the command line
 						List<string> CmdArgs = [
@@ -1402,7 +1402,7 @@ namespace UnrealBuildTool
 
 						uhtManifestScope.Finish();
 
-						Logger.LogInformation("  Running Internal UnrealHeaderTool {CmdLine}", CmdLine);
+						Logger.LogDebug("  Running Internal UnrealHeaderTool {CmdLine}", CmdLine);
 
 						using var uhtExecute = Timeline.ScopeEvent("UHT");
 
