[UBT]
* Removed timeline scope logging that was accidentally added

diff --git a/Engine/Source/Programs/UnrealBuildTool/System/RulesCompiler.cs b/Engine/Source/Programs/UnrealBuildTool/System/RulesCompiler.cs
--- a/Engine/Source/Programs/UnrealBuildTool/System/RulesCompiler.cs
+++ b/Engine/Source/Programs/UnrealBuildTool/System/RulesCompiler.cs
@@ -79,7 +79,7 @@ namespace UnrealBuildTool
 						List<PluginInfo> EnginePlugins = [];
 						List<PluginInfo> MarketplacePlugins = [];
 						{
-							using var ite = Timeline.ScopeEvent("EngineRulesPlugins", Logger);
+							using var ite = Timeline.ScopeEvent("EngineRulesPlugins");
 							DirectoryReference MarketplaceDirectory = DirectoryReference.Combine(Unreal.EngineDirectory, "Plugins", "Marketplace");
 							foreach (PluginInfo PluginInfo in Plugins.ReadEnginePlugins(Unreal.EngineDirectory))
 							{
@@ -131,7 +131,7 @@ namespace UnrealBuildTool
 
 			// We want to prefetch the directories in one call to make it go as wide as possible against filesystem
 			{
-				using var fdt = Timeline.ScopeEvent("EngineRulesFetchDirs", Logger);
+				using var fdt = Timeline.ScopeEvent("EngineRulesFetchDirs");
 
 				List<DirectoryReference> populateDirs = [];
 				foreach (PluginInfo info in Plugins)
@@ -193,13 +193,13 @@ namespace UnrealBuildTool
 			}
 
 			// Create the assembly
-			using var erat = Timeline.ScopeEvent("EngineRulesAssembly", Logger);
+			using var erat = Timeline.ScopeEvent("EngineRulesAssembly");
 			DirectoryReference AssemblyDir = RootDirectories[0];
 			FileReference EngineAssemblyFileName = FileReference.Combine(AssemblyDir, "Intermediate", "Build", "BuildRules", AssemblyPrefix + "Rules" + FrameworkAssemblyExtension);
 			RulesAssembly EngineAssembly = new RulesAssembly(Scope, RootDirectories, Plugins, ModuleFileToContext, new List<FileReference>(), EngineAssemblyFileName, bContainsEngineModules: true, DefaultBuildSettings: DefaultEngineBuildSettingsVersion, bReadOnly: bReadOnly, bSkipCompile: bSkipCompile, bForceCompile: bForceCompile, Parent: Parent, Logger: Logger);
 			erat.Finish();
 
-			using var prat = Timeline.ScopeEvent("ProgramRulesAssembly", Logger);
+			using var prat = Timeline.ScopeEvent("ProgramRulesAssembly");
 			Dictionary<FileReference, ModuleRulesContext> ProgramModuleFiles = new Dictionary<FileReference, ModuleRulesContext>();
 			foreach (DirectoryReference RootDirectory in RootDirectories)
 			{
