Simplifiying the new ubt flag getter that controls if the asio module is compiled with ssl support or not

diff --git a/Engine/Source/Programs/UnrealBuildTool/Configuration/TargetRules.cs b/Engine/Source/Programs/UnrealBuildTool/Configuration/TargetRules.cs
--- a/Engine/Source/Programs/UnrealBuildTool/Configuration/TargetRules.cs
+++ b/Engine/Source/Programs/UnrealBuildTool/Configuration/TargetRules.cs
@@ -1018,12 +1018,7 @@ namespace UnrealBuildTool
 		[XmlConfigFile(Name = "bCompileAsioSSLSupport")]
 		[CommandLine("-bCompileAsioSSLSupport")]
 		[RequiresUniqueBuildEnvironment]
-		public bool bCompileAsioSSLSupport 
-		{
-			get => bEnableCompileAsioSSLSupport.HasValue && bEnableCompileAsioSSLSupport.Value;
-			set => bEnableCompileAsioSSLSupport = value;
-		}
-		private bool? bEnableCompileAsioSSLSupport;
+		public bool bCompileAsioSSLSupport { get; set; } = false;
 
 		/// <summary>
 		/// Whether to include ICU unicode/i18n support in Core.
