[UBT]
* Fixed race condition when invoking DisallowedPathRegistered, CrossArchitecturePathRegistered or PathHashRegistered in parallel

diff --git a/Engine/Source/Programs/UnrealBuildTool/System/TargetMakefile.cs b/Engine/Source/Programs/UnrealBuildTool/System/TargetMakefile.cs
--- a/Engine/Source/Programs/UnrealBuildTool/System/TargetMakefile.cs
+++ b/Engine/Source/Programs/UnrealBuildTool/System/TargetMakefile.cs
@@ -572,9 +572,29 @@ namespace UnrealBuildTool
 
 		void RegisterUBAHandlers()
 		{
-			EpicGames.UBA.Utils.DisallowedPathRegistered += (sender, e) => UbaDisallowedPaths.UnionWith(e.Paths.Select(DirectoryItem.GetItemByPath));
-			EpicGames.UBA.Utils.CrossArchitecturePathRegistered += (sender, e) => UbaCrossArchitecturePaths.Add(DirectoryItem.GetItemByPath(e.Path), DirectoryItem.GetItemByPath(e.CrossPath));
-			EpicGames.UBA.Utils.PathHashRegistered += (sender, e) => UbaPathHashes.Add(DirectoryItem.GetItemByPath(e.Path), e.Hash);
+			EpicGames.UBA.Utils.DisallowedPathRegistered += (sender, e) =>
+			{
+				lock (UbaDisallowedPaths)
+				{
+					UbaDisallowedPaths.UnionWith(e.Paths.Select(DirectoryItem.GetItemByPath));
+				}
+			};
+
+			EpicGames.UBA.Utils.CrossArchitecturePathRegistered += (sender, e) =>
+			{
+				lock (UbaCrossArchitecturePaths)
+				{
+					UbaCrossArchitecturePaths.Add(DirectoryItem.GetItemByPath(e.Path), DirectoryItem.GetItemByPath(e.CrossPath));
+				}
+			};
+
+			EpicGames.UBA.Utils.PathHashRegistered += (sender, e) =>
+			{
+				lock (UbaPathHashes)
+				{
+					UbaPathHashes.Add(DirectoryItem.GetItemByPath(e.Path), e.Hash);
+				}
+			};
 		}
 
 		/// <summary>
