[UBA]
* Added code for UBAExecutor to send UbaAgent.toml to helpers (toml file not submitted yet)
* Added preload of casdb in parallel with connecting to cache servers
* Fixed some spelling

diff --git a/Engine/Source/Programs/Shared/EpicGames.UBA/IServer.cs b/Engine/Source/Programs/Shared/EpicGames.UBA/IServer.cs
--- a/Engine/Source/Programs/Shared/EpicGames.UBA/IServer.cs
+++ b/Engine/Source/Programs/Shared/EpicGames.UBA/IServer.cs
@@ -39,6 +39,12 @@ namespace EpicGames.UBA
 		/// <returns>Success</returns>
 		public abstract bool AddClient(string ip, int port, string crypto = "");
 
+		/// <summary>
+		/// Set config (.toml) that should be used by all clients.
+		/// </summary>
+		/// <param name="path"></param>
+		public abstract void SetClientsConfig(string path);
+
 		/// <summary>
 		/// Create a IServer object
 		/// </summary>
diff --git a/Engine/Source/Programs/Shared/EpicGames.UBA/IStorageServer.cs b/Engine/Source/Programs/Shared/EpicGames.UBA/IStorageServer.cs
--- a/Engine/Source/Programs/Shared/EpicGames.UBA/IStorageServer.cs
+++ b/Engine/Source/Programs/Shared/EpicGames.UBA/IStorageServer.cs
@@ -49,6 +49,11 @@ namespace EpicGames.UBA
 		/// </summary>
 		public abstract void RegisterDisallowedPath(string path);
 
+		/// <summary>
+		/// Preload cas table (otherwise it will be loaded on demand which might stall)
+		/// </summary>
+		public abstract void PreloadCasTable();
+
 		/// <summary>
 		/// Create a IStorageServer object
 		/// </summary>
diff --git a/Engine/Source/Programs/Shared/EpicGames.UBA/Impl/ServerImpl.cs b/Engine/Source/Programs/Shared/EpicGames.UBA/Impl/ServerImpl.cs
--- a/Engine/Source/Programs/Shared/EpicGames.UBA/Impl/ServerImpl.cs
+++ b/Engine/Source/Programs/Shared/EpicGames.UBA/Impl/ServerImpl.cs
@@ -27,6 +27,9 @@ namespace EpicGames.UBA.Impl
 		[DllImport("UbaHost", CharSet = CharSet.Auto)]
 		static extern byte NetworkServer_AddNamedConnection(nint server, string name);
 
+		[DllImport("UbaHost", CharSet = CharSet.Auto)]
+		static extern byte NetworkServer_SetClientsConfig(nint server, nint config);
+
 		[DllImport("UbaHost", CharSet = CharSet.Auto)]
 		static extern byte NetworkServer_Stop(nint server);
 
@@ -79,6 +82,14 @@ namespace EpicGames.UBA.Impl
 
 		/// <inheritdoc/>
 		public bool AddNamedConnection(string name) => NetworkServer_AddNamedConnection(_handle, name) != 0;
+
+		/// <inheritdoc/>
+		public void SetClientsConfig(string path)
+		{
+			using IConfig config = IConfig.LoadConfig(path);
+			NetworkServer_SetClientsConfig(_handle, config.GetHandle());
+		}
+
 		#endregion
 	}
 }
diff --git a/Engine/Source/Programs/Shared/EpicGames.UBA/Impl/StorageServerImpl.cs b/Engine/Source/Programs/Shared/EpicGames.UBA/Impl/StorageServerImpl.cs
--- a/Engine/Source/Programs/Shared/EpicGames.UBA/Impl/StorageServerImpl.cs
+++ b/Engine/Source/Programs/Shared/EpicGames.UBA/Impl/StorageServerImpl.cs
@@ -24,6 +24,9 @@ namespace EpicGames.UBA.Impl
 
 		[DllImport("UbaHost", CharSet = CharSet.Auto)]
 		static extern void StorageServer_RegisterDisallowedPath(nint server, string path);
+
+		[DllImport("UbaHost", CharSet = CharSet.Auto)]
+		static extern void StorageServer_PreloadCasTable(nint server);
 		#endregion
 
 		public StorageServerImpl(IServer server, ILogger logger, StorageServerCreateInfo info)
@@ -63,6 +66,8 @@ namespace EpicGames.UBA.Impl
 		public void SaveCasTable() => StorageServer_SaveCasTable(_handle);
 
 		public void RegisterDisallowedPath(string path) => StorageServer_RegisterDisallowedPath(_handle, path);
+
+		public void PreloadCasTable() => StorageServer_PreloadCasTable(_handle);
 		#endregion
 	}
 }
diff --git a/Engine/Source/Programs/UnrealBuildTool/Executors/UnrealBuildAccelerator/UBAExecutor.cs b/Engine/Source/Programs/UnrealBuildTool/Executors/UnrealBuildAccelerator/UBAExecutor.cs
--- a/Engine/Source/Programs/UnrealBuildTool/Executors/UnrealBuildAccelerator/UBAExecutor.cs
+++ b/Engine/Source/Programs/UnrealBuildTool/Executors/UnrealBuildAccelerator/UBAExecutor.cs
@@ -284,6 +284,17 @@ namespace UnrealBuildTool
 				_session.AddInfo(_sessionHyperlink);
 			}
 
+			Task? preloadTask = null;
+
+			if (!UBAConfig.bDisableRemote)
+			{
+				preloadTask = Task.Run(() =>
+				{
+					using TraceTaskScope s = new(_trace, "LoadCasTable", "");
+					_storageServer.PreloadCasTable();
+				});
+			}
+
 			List<Task> coordinatorInitTasks = new();
 			if (!UBAConfig.bDisableRemote)
 			{
@@ -339,8 +350,12 @@ namespace UnrealBuildTool
 				_trace.UpdateStatus(1, 6, _cacheClients.Any() ? "Ready" : "Failed to connect", LogEntryType.Info);
 			}
 
+			preloadTask?.Wait();
+
 			if (!UBAConfig.bDisableRemote)
 			{
+				_server.SetClientsConfig(FileReference.Combine(UbaBinariesDir, "UbaAgent.toml").FullName);
+
 				_server.StartServer(UBAConfig.Host, UBAConfig.Port, Crypto);
 			}
 			else
@@ -528,9 +543,9 @@ namespace UnrealBuildTool
 		[SupportedOSPlatform("windows")]
 		static void LaunchVisualizer()
 		{
-			FileReference visaulizerPath = FileReference.Combine(Unreal.EngineDirectory, "Binaries", "Win64", "UnrealBuildAccelerator", RuntimeInformation.ProcessArchitecture.ToString(), "UbaVisualizer.exe");
-			FileReference tempPath = FileReference.Combine(new DirectoryReference(System.IO.Path.GetTempPath()), visaulizerPath.GetFileName());
-			if (!FileReference.Exists(visaulizerPath))
+			FileReference visualizerPath = FileReference.Combine(UbaBinariesDir, "UbaVisualizer.exe");
+			FileReference tempPath = FileReference.Combine(new DirectoryReference(System.IO.Path.GetTempPath()), visualizerPath.GetFileName());
+			if (!FileReference.Exists(visualizerPath))
 			{
 				return;
 			}
@@ -538,7 +553,7 @@ namespace UnrealBuildTool
 			try
 			{
 				// Check if a listening visualizer is already running
-				foreach (System.Diagnostics.Process process in System.Diagnostics.Process.GetProcessesByName(visaulizerPath.GetFileNameWithoutAnyExtensions()))
+				foreach (System.Diagnostics.Process process in System.Diagnostics.Process.GetProcessesByName(visualizerPath.GetFileNameWithoutAnyExtensions()))
 				{
 					using ManagementObjectSearcher searcher = new($"SELECT CommandLine FROM Win32_Process WHERE ProcessId = {process.Id}");
 					using ManagementObjectCollection objects = searcher.Get();
@@ -548,9 +563,9 @@ namespace UnrealBuildTool
 						return;
 					}
 				}
-				if (!FileReference.Exists(tempPath) || tempPath.ToFileInfo().LastWriteTime < visaulizerPath.ToFileInfo().LastWriteTime)
+				if (!FileReference.Exists(tempPath) || tempPath.ToFileInfo().LastWriteTime < visualizerPath.ToFileInfo().LastWriteTime)
 				{
-					FileReference.Copy(visaulizerPath, tempPath, true);
+					FileReference.Copy(visualizerPath, tempPath, true);
 				}
 				if (FileReference.Exists(tempPath))
 				{
