[UBT]
* Added lock around ios toolchain code that is writing plistfile since it is written in the same folder for all configs

diff --git a/Engine/Source/Programs/UnrealBuildTool/Platform/IOS/UEDeployIOS.cs b/Engine/Source/Programs/UnrealBuildTool/Platform/IOS/UEDeployIOS.cs
--- a/Engine/Source/Programs/UnrealBuildTool/Platform/IOS/UEDeployIOS.cs
+++ b/Engine/Source/Programs/UnrealBuildTool/Platform/IOS/UEDeployIOS.cs
@@ -507,45 +507,50 @@ namespace UnrealBuildTool
 			File.WriteAllText(PlistFile.FullName, Text.ToString());
 		}
 
+		private static object GenerateIOSPListLock = new();
+
 		public static bool GenerateIOSPList(FileReference? ProjectFile, UnrealTargetConfiguration Config, string ProjectDirectory, bool bIsUnrealGame, string GameName, bool bIsClient, string ProjectName, string InEngineDir, string AppDirectory, UnrealPluginLanguage? UPL, string? BundleID, bool bBuildAsFramework, ILogger Logger)
 		{
-			// get the settings from the ini file
-			// plist replacements
-			DirectoryReference? DirRef = bIsUnrealGame ? (!String.IsNullOrEmpty(UnrealBuildTool.GetRemoteIniPath()) ? new DirectoryReference(UnrealBuildTool.GetRemoteIniPath()!) : null) : new DirectoryReference(ProjectDirectory);
-
-			if (!AppleExports.UseModernXcode(ProjectFile))
+			lock (GenerateIOSPListLock) // Needs to be protected because multiple configs are built in parallel
 			{
-				return GenerateLegacyIOSPList(ProjectFile, Config, ProjectDirectory, bIsUnrealGame, GameName, bIsClient, ProjectName, InEngineDir, AppDirectory, UPL, BundleID, bBuildAsFramework, Logger);
-			}
+				// get the settings from the ini file
+				// plist replacements
+				DirectoryReference? DirRef = bIsUnrealGame ? (!String.IsNullOrEmpty(UnrealBuildTool.GetRemoteIniPath()) ? new DirectoryReference(UnrealBuildTool.GetRemoteIniPath()!) : null) : new DirectoryReference(ProjectDirectory);
 
-			// generate the Info.plist for future use
-			string BuildDirectory = ProjectDirectory + "/Build/IOS";
-			string IntermediateDirectory = ProjectDirectory + "/Intermediate/IOS";
-			string PListFile = IntermediateDirectory + "/" + GameName + "-Info.plist";
-			;
-			ProjectName = !String.IsNullOrEmpty(ProjectName) ? ProjectName : GameName;
-			VersionUtilities.BuildDirectory = BuildDirectory;
-			VersionUtilities.GameName = GameName;
+				if (!AppleExports.UseModernXcode(ProjectFile))
+				{
+					return GenerateLegacyIOSPList(ProjectFile, Config, ProjectDirectory, bIsUnrealGame, GameName, bIsClient, ProjectName, InEngineDir, AppDirectory, UPL, BundleID, bBuildAsFramework, Logger);
+				}
 
-			WritePlistFile(new FileReference(PListFile), DirRef, UPL, Config, GameName, bIsUnrealGame, ProjectName, Logger);
+				// generate the Info.plist for future use
+				string BuildDirectory = ProjectDirectory + "/Build/IOS";
+				string IntermediateDirectory = ProjectDirectory + "/Intermediate/IOS";
+				string PListFile = IntermediateDirectory + "/" + GameName + "-Info.plist";
+				;
+				ProjectName = !String.IsNullOrEmpty(ProjectName) ? ProjectName : GameName;
+				VersionUtilities.BuildDirectory = BuildDirectory;
+				VersionUtilities.GameName = GameName;
 
-			if (BuildHostPlatform.Current.Platform == UnrealTargetPlatform.Mac && !bBuildAsFramework)
-			{
-				// need to make sure touching this plist is in the same lock as FinalizeAppWithXcode()
-				string MutexName = GlobalSingleInstanceMutex.GetUniqueMutexForPath("UnrealBuildTool_XcodeBuild", Unreal.RootDirectory);
-				using (new GlobalSingleInstanceMutex(MutexName, true))
+				WritePlistFile(new FileReference(PListFile), DirRef, UPL, Config, GameName, bIsUnrealGame, ProjectName, Logger);
+
+				if (BuildHostPlatform.Current.Platform == UnrealTargetPlatform.Mac && !bBuildAsFramework)
 				{
-					FileReference FinalPlistFile;
-					FinalPlistFile = new FileReference($"{ProjectDirectory}/Build/IOS/UBTGenerated/Info.Template.plist");
+					// need to make sure touching this plist is in the same lock as FinalizeAppWithXcode()
+					string MutexName = GlobalSingleInstanceMutex.GetUniqueMutexForPath("UnrealBuildTool_XcodeBuild", Unreal.RootDirectory);
+					using (new GlobalSingleInstanceMutex(MutexName, true))
+					{
+						FileReference FinalPlistFile;
+						FinalPlistFile = new FileReference($"{ProjectDirectory}/Build/IOS/UBTGenerated/Info.Template.plist");
 
-					DirectoryReference.CreateDirectory(FinalPlistFile.Directory);
-					// @todo: writeifdifferent is better
-					FileReference.Delete(FinalPlistFile);
-					File.Copy(PListFile, FinalPlistFile.FullName);
+						DirectoryReference.CreateDirectory(FinalPlistFile.Directory);
+						// @todo: writeifdifferent is better
+						FileReference.Delete(FinalPlistFile);
+						File.Copy(PListFile, FinalPlistFile.FullName);
+					}
 				}
-			}
 
-			return true;
+				return true;
+			}
 		}
 
 		public static bool GenerateLegacyIOSPList(FileReference? ProjectFile, UnrealTargetConfiguration Config, string ProjectDirectory, bool bIsUnrealGame, string GameName, bool bIsClient, string ProjectName, string InEngineDir, string AppDirectory, UnrealPluginLanguage? UPL, string? BundleID, bool bBuildAsFramework, ILogger Logger)
