Fix remote codesigning by ensuring that intemediate certificates (Specifically Apple WWDR G3) are installed to the remote Mac's temporary keychain that's setup for remote codesigning.
Add */.vs/* directory wildcard to Rsync filters to ignore thos and allow debugging UBT from Windows while running a remote build.
Remove deletion options when Rsync'ing remote engine folder as this was causing the Windows machines Engine folder to get wiped.

diff --git a/Engine/Build/Rsync/RsyncEngine.txt b/Engine/Build/Rsync/RsyncEngine.txt
--- a/Engine/Build/Rsync/RsyncEngine.txt
+++ b/Engine/Build/Rsync/RsyncEngine.txt
@@ -214,6 +214,7 @@ protect /Source/Programs/Shared/EpicGames.*/***
 + /Source/ThirdParty/**/IOS/**.xcframework/**
 + /Source/ThirdParty/**/TVOS/**.xcframework/**
 - /Source/ThirdParty/**
+- /Source/**/.vs/**
 
 # Source: Include everything else
 + /Source/**/
diff --git a/Engine/Source/Programs/UnrealBuildTool/Configuration/UEBuildTarget.cs b/Engine/Source/Programs/UnrealBuildTool/Configuration/UEBuildTarget.cs
--- a/Engine/Source/Programs/UnrealBuildTool/Configuration/UEBuildTarget.cs
+++ b/Engine/Source/Programs/UnrealBuildTool/Configuration/UEBuildTarget.cs
@@ -2807,10 +2807,6 @@ namespace UnrealBuildTool
 			CodeGenContext? lastRunContext = null;
 			Dictionary<string, UHTModuleInfo>? lastRunUObjectModules = null;
 
-#if __VPROJECT_AVAILABLE__
-			HashSet<string>? lastRunVNIModules = null;
-#endif
-
 			string GetTargetName(TargetDescriptor desc) => $"{desc.Name}-{desc.Platform}-{desc.Configuration}";
 
 			foreach (CodeGenContext context in contexts)
@@ -2852,6 +2848,8 @@ namespace UnrealBuildTool
 					}
 
 #if __VPROJECT_AVAILABLE__
+					HashSet<string>? lastRunVNIModules = null;
+
 					if (lastRunVNIModules == null)
 					{
 						lastRunVNIModules = new(lastRunContext.Makefile.VNIModules.Select(m => m.ModuleName));
diff --git a/Engine/Source/Programs/UnrealBuildTool/Platform/IOS/IOSToolChain.cs b/Engine/Source/Programs/UnrealBuildTool/Platform/IOS/IOSToolChain.cs
--- a/Engine/Source/Programs/UnrealBuildTool/Platform/IOS/IOSToolChain.cs
+++ b/Engine/Source/Programs/UnrealBuildTool/Platform/IOS/IOSToolChain.cs
@@ -1448,7 +1448,7 @@ namespace UnrealBuildTool
 					// Copy the mobile provision into the system store
 					if (Target.ImportProvision != null)
 					{
-						Writer.WriteLine("cp -f {0} ~/Library/MobileDevice/Provisioning\\ Profiles/", Utils.EscapeShellArgument(Target.ImportProvision));
+						Writer.WriteLine("cp -f {0} {1}}", Utils.EscapeShellArgument(Target.ImportProvision), Utils.EscapeShellArgument(AppleExports.GetProvisionDirectory().FullName));
 					}
 
 					// Get the signing certificate to use
diff --git a/Engine/Source/Programs/UnrealBuildTool/Platform/IOS/UEBuildIOS.cs b/Engine/Source/Programs/UnrealBuildTool/Platform/IOS/UEBuildIOS.cs
--- a/Engine/Source/Programs/UnrealBuildTool/Platform/IOS/UEBuildIOS.cs
+++ b/Engine/Source/Programs/UnrealBuildTool/Platform/IOS/UEBuildIOS.cs
@@ -472,17 +472,7 @@ namespace UnrealBuildTool
 
 			if (!String.IsNullOrEmpty(MobileProvision))
 			{
-				DirectoryReference MobileProvisionDir;
-				if (BuildHostPlatform.Current.Platform == UnrealTargetPlatform.Mac)
-				{
-					MobileProvisionDir = DirectoryReference.Combine(new DirectoryReference(Environment.GetEnvironmentVariable("HOME")!), "Library", "MobileDevice", "Provisioning Profiles");
-				}
-				else
-				{
-					MobileProvisionDir = DirectoryReference.Combine(DirectoryReference.GetSpecialFolder(Environment.SpecialFolder.LocalApplicationData)!, "Apple Computer", "MobileDevice", "Provisioning Profiles");
-				}
-
-				FileReference PossibleMobileProvisionFile = FileReference.Combine(MobileProvisionDir, MobileProvision);
+				FileReference PossibleMobileProvisionFile = FileReference.Combine(AppleExports.GetProvisionDirectory(), MobileProvision);
 				if (FileReference.Exists(PossibleMobileProvisionFile))
 				{
 					MobileProvisionFile = PossibleMobileProvisionFile;
diff --git a/Engine/Source/Programs/UnrealBuildTool/ToolChain/AppleToolChain.cs b/Engine/Source/Programs/UnrealBuildTool/ToolChain/AppleToolChain.cs
--- a/Engine/Source/Programs/UnrealBuildTool/ToolChain/AppleToolChain.cs
+++ b/Engine/Source/Programs/UnrealBuildTool/ToolChain/AppleToolChain.cs
@@ -1242,7 +1242,7 @@ namespace UnrealBuildTool
 				}
 
 				// copy the provision into standard location
-				Writer.WriteLine("cp -f {0} ~/Library/MobileDevice/Provisioning\\ Profiles/", Utils.EscapeShellArgument(Target.RemoteImportProvision));
+				Writer.WriteLine("cp -f {0} {1}", Utils.EscapeShellArgument(Target.RemoteImportProvision), Utils.EscapeShellArgument(AppleExports.GetProvisionDirectory().FullName));
 				MobileProvisionContents MobileProvision = MobileProvisionContents.Read(new FileReference(Target.RemoteImportProvision));
 				MobileProvisionUUID = MobileProvision.GetUniqueId();
 
@@ -1266,6 +1266,8 @@ namespace UnrealBuildTool
 				Writer.WriteLine("security list-keychains");
 				Writer.WriteLine("security set-keychain-settings -t 3600 -l  \"{0}\"", TempKeychain);
 				Writer.WriteLine("security -v unlock-keychain -p \"A\" \"{0}\"", TempKeychain);
+				Writer.WriteLine("[ ! -f {0}/AppleWWDRCAG3.cer ] && curl --output {0}/AppleWWDRCAG3.cer https://www.apple.com/certificateauthority/AppleWWDRCAG3.cer", Utils.EscapeShellArgument(Target.ProjectIntermediateDirectory.FullName));
+				Writer.WriteLine("security import {0}/AppleWWDRCAG3.cer -k \"{1}\"", Utils.EscapeShellArgument(Target.ProjectIntermediateDirectory.FullName), TempKeychain);
 				Writer.WriteLine("security import {0} -P {1} -k \"{2}\" -T /usr/bin/codesign -T /usr/bin/security -t agg", Utils.EscapeShellArgument(Target.RemoteImportCertificate), Utils.EscapeShellArgument(Target.RemoteImportCertificatePassword!), TempKeychain);
 				Writer.WriteLine("security set-key-partition-list -S apple-tool:,apple:,codesign: -s -k \"A\" -D '{0}' -t private {1}", SigningCertificate, TempKeychain);
 			}
diff --git a/Engine/Source/Programs/UnrealBuildTool/ToolChain/RemoteMac.cs b/Engine/Source/Programs/UnrealBuildTool/ToolChain/RemoteMac.cs
--- a/Engine/Source/Programs/UnrealBuildTool/ToolChain/RemoteMac.cs
+++ b/Engine/Source/Programs/UnrealBuildTool/ToolChain/RemoteMac.cs
@@ -13,22 +13,6 @@ using UnrealBuildBase;
 
 namespace UnrealBuildTool
 {
-	/// <summary>
-	/// Stores information about how a local directory maps to a remote directory
-	/// </summary>
-	[DebuggerDisplay("{LocalDirectory}")]
-	class RemoteMapping
-	{
-		public DirectoryReference LocalDirectory;
-		public string RemoteDirectory;
-
-		public RemoteMapping(DirectoryReference LocalDirectory, string RemoteDirectory)
-		{
-			this.LocalDirectory = LocalDirectory;
-			this.RemoteDirectory = RemoteDirectory;
-		}
-	}
-
 	/// <summary>
 	/// Handles uploading and building on a remote Mac
 	/// </summary>
@@ -99,6 +83,29 @@ namespace UnrealBuildTool
 		/// </summary>
 		private string RemoteBaseDir;
 
+		/// <summary>
+		/// Stores information about how a local directory maps to a remote directory
+		/// </summary>
+		[DebuggerDisplay("{LocalDirectory}")]
+		private class RemoteMapping
+		{
+			public DirectoryReference LocalDirectory;
+			public string RemoteDirectory;
+			public List<string> DownloadFilesRsyncArgs;
+
+			public RemoteMapping(DirectoryReference LocalDirectory, string RemoteDirectory, List<string> DownloadFilesRsyncArgs)
+			{
+				this.LocalDirectory = LocalDirectory;
+				this.RemoteDirectory = RemoteDirectory;
+				this.DownloadFilesRsyncArgs = DownloadFilesRsyncArgs;
+			}
+		}
+
+		private RemoteMapping CreateRemoteMapping(DirectoryReference LocalDirectory, List<string> DownloadFilesRsyncArgs)
+		{
+			return new RemoteMapping(LocalDirectory, GetRemotePath(LocalDirectory), DownloadFilesRsyncArgs);
+		}
+
 		/// <summary>
 		/// Mappings from local directories to remote directories
 		/// </summary>
@@ -114,6 +121,11 @@ namespace UnrealBuildTool
 		/// </summary>
 		private List<string> BasicRsyncArguments;
 
+		/// <summary>
+		/// Arguments that are used by directory Rsync call without deletions
+		/// </summary>
+		private List<string> NoDeleteRsyncArguments;
+
 		/// <summary>
 		/// Arguments that are used by directory Rsync call
 		/// </summary>
@@ -266,18 +278,21 @@ namespace UnrealBuildTool
 			BasicRsyncArguments.Add(String.Format("--rsh=\"{0} -p {1}\"", RsyncAuthentication, ServerPort));
 			BasicRsyncArguments.Add("--chmod=ugo=rwx");
 
+			// Build a list of arguments for Rsync filters without destination deletions
+			NoDeleteRsyncArguments = new List<string>(BasicRsyncArguments);
+			NoDeleteRsyncArguments.Add("--copy-links");
+			NoDeleteRsyncArguments.Add("--recursive");
+			NoDeleteRsyncArguments.Add("--times"); // Preserve modification times
+			NoDeleteRsyncArguments.Add("--omit-dir-times"); // Ignore modification times for directories
+			NoDeleteRsyncArguments.Add("--prune-empty-dirs"); // Remove empty directories from the file list
+
 			// Build a list of arguments for Rsync filters
-			CommonRsyncArguments = new List<string>(BasicRsyncArguments);
-			CommonRsyncArguments.Add("--copy-links");
-			CommonRsyncArguments.Add("--recursive");
+			CommonRsyncArguments = new List<string>(NoDeleteRsyncArguments);
 			if (!bPrepareForSecondaryMac)
 			{
 				CommonRsyncArguments.Add("--delete"); // Delete anything not in the source directory
 				CommonRsyncArguments.Add("--delete-excluded"); // Delete anything not in the source directory
 			}
-			CommonRsyncArguments.Add("--times"); // Preserve modification times
-			CommonRsyncArguments.Add("--omit-dir-times"); // Ignore modification times for directories
-			CommonRsyncArguments.Add("--prune-empty-dirs"); // Remove empty directories from the file list
 
 			// Get the remote base directory
 			string RemoteServerOverrideBuildPath;
@@ -299,10 +314,10 @@ namespace UnrealBuildTool
 
 			// Build the list of directory mappings between the local and remote machines
 			Mappings = new List<RemoteMapping>();
-			Mappings.Add(new RemoteMapping(Unreal.EngineDirectory, GetRemotePath(Unreal.EngineDirectory)));
+			Mappings.Add(CreateRemoteMapping(Unreal.EngineDirectory, NoDeleteRsyncArguments));
 			if (ProjectFile != null && !ProjectFile.IsUnderDirectory(Unreal.EngineDirectory))
 			{
-				Mappings.Add(new RemoteMapping(ProjectFile.Directory, GetRemotePath(ProjectFile.Directory)));
+				Mappings.Add(CreateRemoteMapping(ProjectFile.Directory, CommonRsyncArguments));
 			}
 			if (AdditionalPaths != null && ProjectFile != null)
 			{
@@ -311,7 +326,7 @@ namespace UnrealBuildTool
 					if (!AdditionalPath.IsUnderDirectory(Unreal.EngineDirectory) &&
 						!AdditionalPath.IsUnderDirectory(ProjectFile.Directory))
 					{
-						Mappings.Add(new RemoteMapping(AdditionalPath, GetRemotePath(AdditionalPath)));
+						Mappings.Add(CreateRemoteMapping(AdditionalPath, CommonRsyncArguments));
 					}
 				}
 			}
@@ -465,7 +480,7 @@ namespace UnrealBuildTool
 			}
 			if (!bLogIsMapped)
 			{
-				Mappings.Add(new RemoteMapping(RemoteLogFile.Directory, GetRemotePath(RemoteLogFile.Directory)));
+				Mappings.Add(CreateRemoteMapping(RemoteLogFile.Directory, CommonRsyncArguments));
 			}
 
 			// Path to the local manifest file. This has to be translated from the remote format after the build is complete.
@@ -1134,7 +1149,7 @@ namespace UnrealBuildTool
 					DirectoryReference.CreateDirectory(DownloadListLocation.Directory);
 					FileReference.WriteAllLines(DownloadListLocation, FileGroups[Idx].Select(x => x.MakeRelativeTo(Mappings[Idx].LocalDirectory).Replace('\\', '/')));
 
-					List<string> Arguments = new List<string>(CommonRsyncArguments);
+					List<string> Arguments = new List<string>(Mappings[Idx].DownloadFilesRsyncArgs);
 					Arguments.Add(String.Format("--files-from=\"{0}\"", GetLocalCygwinPath(DownloadListLocation)));
 					Arguments.Add(String.Format("\"{0}@{1}\":'{2}/'", UserName, ServerName, Mappings[Idx].RemoteDirectory));
 					Arguments.Add(String.Format("\"{0}/\"", GetLocalCygwinPath(Mappings[Idx].LocalDirectory)));
