Erroneous format string syntax fix.

diff --git a/Engine/Source/Programs/UnrealBuildTool/Platform/IOS/IOSToolChain.cs b/Engine/Source/Programs/UnrealBuildTool/Platform/IOS/IOSToolChain.cs
--- a/Engine/Source/Programs/UnrealBuildTool/Platform/IOS/IOSToolChain.cs
+++ b/Engine/Source/Programs/UnrealBuildTool/Platform/IOS/IOSToolChain.cs
@@ -1448,7 +1448,7 @@ namespace UnrealBuildTool
 					// Copy the mobile provision into the system store
 					if (Target.ImportProvision != null)
 					{
-						Writer.WriteLine("cp -f {0} {1}}", Utils.EscapeShellArgument(Target.ImportProvision), Utils.EscapeShellArgument(AppleExports.GetProvisionDirectory().FullName));
+						Writer.WriteLine("cp -f {0} {1}", Utils.EscapeShellArgument(Target.ImportProvision), Utils.EscapeShellArgument(AppleExports.GetProvisionDirectory().FullName));
 					}
 
 					// Get the signing certificate to use
