Further fixes to remoteMac - Don't delete files from game directory as this was causing game modules to get deleted on Windows and subsequentl cook would fail due to missing DLL.
Additionally, update to current Apple Worldwide Developer Relaions codesigning certificate. Use local certificate instead of downloading when remote codesigning on remote Mac.

diff --git a/Engine/Source/Programs/UnrealBuildTool/ToolChain/AppleToolChain.cs b/Engine/Source/Programs/UnrealBuildTool/ToolChain/AppleToolChain.cs
--- a/Engine/Source/Programs/UnrealBuildTool/ToolChain/AppleToolChain.cs
+++ b/Engine/Source/Programs/UnrealBuildTool/ToolChain/AppleToolChain.cs
@@ -1266,8 +1266,7 @@ namespace UnrealBuildTool
 				Writer.WriteLine("security list-keychains");
 				Writer.WriteLine("security set-keychain-settings -t 3600 -l  \"{0}\"", TempKeychain);
 				Writer.WriteLine("security -v unlock-keychain -p \"A\" \"{0}\"", TempKeychain);
-				Writer.WriteLine("[ ! -f {0}/AppleWWDRCAG3.cer ] && curl --output {0}/AppleWWDRCAG3.cer https://www.apple.com/certificateauthority/AppleWWDRCAG3.cer", Utils.EscapeShellArgument(Target.ProjectIntermediateDirectory.FullName));
-				Writer.WriteLine("security import {0}/AppleWWDRCAG3.cer -k \"{1}\"", Utils.EscapeShellArgument(Target.ProjectIntermediateDirectory.FullName), TempKeychain);
+				Writer.WriteLine("security import {0}/Build/IOS/AppleWorldwideDeveloperRelationsCA.pem -k \"{1}\"", Utils.EscapeShellArgument(Unreal.EngineDirectory.FullName), TempKeychain);
 				Writer.WriteLine("security import {0} -P {1} -k \"{2}\" -T /usr/bin/codesign -T /usr/bin/security -t agg", Utils.EscapeShellArgument(Target.RemoteImportCertificate), Utils.EscapeShellArgument(Target.RemoteImportCertificatePassword!), TempKeychain);
 				Writer.WriteLine("security set-key-partition-list -S apple-tool:,apple:,codesign: -s -k \"A\" -D '{0}' -t private {1}", SigningCertificate, TempKeychain);
 			}
diff --git a/Engine/Source/Programs/UnrealBuildTool/ToolChain/RemoteMac.cs b/Engine/Source/Programs/UnrealBuildTool/ToolChain/RemoteMac.cs
--- a/Engine/Source/Programs/UnrealBuildTool/ToolChain/RemoteMac.cs
+++ b/Engine/Source/Programs/UnrealBuildTool/ToolChain/RemoteMac.cs
@@ -91,19 +91,17 @@ namespace UnrealBuildTool
 		{
 			public DirectoryReference LocalDirectory;
 			public string RemoteDirectory;
-			public List<string> DownloadFilesRsyncArgs;
 
-			public RemoteMapping(DirectoryReference LocalDirectory, string RemoteDirectory, List<string> DownloadFilesRsyncArgs)
+			public RemoteMapping(DirectoryReference LocalDirectory, string RemoteDirectory)
 			{
 				this.LocalDirectory = LocalDirectory;
 				this.RemoteDirectory = RemoteDirectory;
-				this.DownloadFilesRsyncArgs = DownloadFilesRsyncArgs;
 			}
 		}
 
-		private RemoteMapping CreateRemoteMapping(DirectoryReference LocalDirectory, List<string> DownloadFilesRsyncArgs)
+		private RemoteMapping CreateRemoteMapping(DirectoryReference LocalDirectory)
 		{
-			return new RemoteMapping(LocalDirectory, GetRemotePath(LocalDirectory), DownloadFilesRsyncArgs);
+			return new RemoteMapping(LocalDirectory, GetRemotePath(LocalDirectory));
 		}
 
 		/// <summary>
@@ -314,10 +312,10 @@ namespace UnrealBuildTool
 
 			// Build the list of directory mappings between the local and remote machines
 			Mappings = new List<RemoteMapping>();
-			Mappings.Add(CreateRemoteMapping(Unreal.EngineDirectory, NoDeleteRsyncArguments));
+			Mappings.Add(CreateRemoteMapping(Unreal.EngineDirectory));
 			if (ProjectFile != null && !ProjectFile.IsUnderDirectory(Unreal.EngineDirectory))
 			{
-				Mappings.Add(CreateRemoteMapping(ProjectFile.Directory, CommonRsyncArguments));
+				Mappings.Add(CreateRemoteMapping(ProjectFile.Directory));
 			}
 			if (AdditionalPaths != null && ProjectFile != null)
 			{
@@ -326,7 +324,7 @@ namespace UnrealBuildTool
 					if (!AdditionalPath.IsUnderDirectory(Unreal.EngineDirectory) &&
 						!AdditionalPath.IsUnderDirectory(ProjectFile.Directory))
 					{
-						Mappings.Add(CreateRemoteMapping(AdditionalPath, CommonRsyncArguments));
+						Mappings.Add(CreateRemoteMapping(AdditionalPath));
 					}
 				}
 			}
@@ -480,7 +478,7 @@ namespace UnrealBuildTool
 			}
 			if (!bLogIsMapped)
 			{
-				Mappings.Add(CreateRemoteMapping(RemoteLogFile.Directory, CommonRsyncArguments));
+				Mappings.Add(CreateRemoteMapping(RemoteLogFile.Directory));
 			}
 
 			// Path to the local manifest file. This has to be translated from the remote format after the build is complete.
@@ -1149,7 +1147,7 @@ namespace UnrealBuildTool
 					DirectoryReference.CreateDirectory(DownloadListLocation.Directory);
 					FileReference.WriteAllLines(DownloadListLocation, FileGroups[Idx].Select(x => x.MakeRelativeTo(Mappings[Idx].LocalDirectory).Replace('\\', '/')));
 
-					List<string> Arguments = new List<string>(Mappings[Idx].DownloadFilesRsyncArgs);
+					List<string> Arguments = new List<string>(NoDeleteRsyncArguments);
 					Arguments.Add(String.Format("--files-from=\"{0}\"", GetLocalCygwinPath(DownloadListLocation)));
 					Arguments.Add(String.Format("\"{0}@{1}\":'{2}/'", UserName, ServerName, Mappings[Idx].RemoteDirectory));
 					Arguments.Add(String.Format("\"{0}/\"", GetLocalCygwinPath(Mappings[Idx].LocalDirectory)));
