Removed PLATFORM_64BITS def from Android toolchain.

diff --git a/Engine/Source/Programs/UnrealBuildTool/Platform/Android/AndroidToolChain.cs b/Engine/Source/Programs/UnrealBuildTool/Platform/Android/AndroidToolChain.cs
--- a/Engine/Source/Programs/UnrealBuildTool/Platform/Android/AndroidToolChain.cs
+++ b/Engine/Source/Programs/UnrealBuildTool/Platform/Android/AndroidToolChain.cs
@@ -694,7 +694,6 @@ namespace UnrealBuildTool
 			Arguments.Add("-fforce-emit-vtables");      // Helps with devirtualization
 			Arguments.Add("-D_FORTIFY_SOURCE=2");       // FORTIFY default
 			Arguments.Add($"-DPLATFORM_USED_NDK_VERSION_INTEGER={NDKApiLevel64Int}");       // NDK version
-			Arguments.Add("-DPLATFORM_64BITS=1");       // NDK version
 		}
 
 		/// <inheritdoc/>
