[UBT]
* Added so UbaExecutor can be created without using init task
* Fixed so compress mode code is not using init task

diff --git a/Engine/Source/Programs/UnrealBuildTool/Executors/UnrealBuildAccelerator/UBAExecutor.cs b/Engine/Source/Programs/UnrealBuildTool/Executors/UnrealBuildAccelerator/UBAExecutor.cs
--- a/Engine/Source/Programs/UnrealBuildTool/Executors/UnrealBuildAccelerator/UBAExecutor.cs
+++ b/Engine/Source/Programs/UnrealBuildTool/Executors/UnrealBuildAccelerator/UBAExecutor.cs
@@ -25,7 +25,7 @@ namespace UnrealBuildTool
 		public string Crypto { get; private set; } = String.Empty;
 
 		readonly IEnumerable<TargetDescriptor> _targetDescriptors;
-		readonly Task _initTask;
+		readonly Task? _initTask;
 		readonly CancellationTokenSource _cancellationSource;
 
 		EpicGames.UBA.ILogger? _ubaLogger;
@@ -59,7 +59,7 @@ namespace UnrealBuildTool
 		{
 			if (disposing)
 			{
-				_initTask.Wait();
+				_initTask?.Wait();
 
 				{
 					using ITimelineEvent a2 = Timeline.ScopeEvent("DisconnectCache");
@@ -133,13 +133,16 @@ namespace UnrealBuildTool
 			}
 		}
 
-		public UBAExecutor(int maxLocalActions, bool bAllCores, bool bCompactOutput, ILogger logger, IEnumerable<TargetDescriptor> targetDescriptors)
+		public UBAExecutor(int maxLocalActions, bool bAllCores, bool bCompactOutput, ILogger logger, IEnumerable<TargetDescriptor> targetDescriptors, bool useInitTask)
 			: base(maxLocalActions, bAllCores, bCompactOutput, logger)
 		{
 			_targetDescriptors = targetDescriptors;
 			_cancellationSource = new();
-			_initTask = new(() => { Init(targetDescriptors, logger); }, TaskCreationOptions.LongRunning);
-			_initTask.Start();
+			if (useInitTask)
+			{
+				_initTask = new(() => { Init(targetDescriptors, logger); }, TaskCreationOptions.LongRunning);
+				_initTask.Start();
+			}
 		}
 
 		private void Init(IEnumerable<TargetDescriptor> targetDescriptors, ILogger logger)
@@ -476,7 +479,14 @@ namespace UnrealBuildTool
 				_cancellationSource.Cancel();
 			}
 
-			_initTask.Wait();
+			if (_initTask == null)
+			{
+				Init(_targetDescriptors, logger);
+			}
+			else
+			{
+				_initTask?.Wait();
+			}
 
 			PrintConfiguration(logger);
 
@@ -490,6 +500,12 @@ namespace UnrealBuildTool
 			if (UBAConfig.bForceNoCache)
 			{
 				DisposeCacheClients();
+
+				// TOOD: We want to do this but it cause okta re-logins which are annoying
+				//if (_agentCoordinators.Count > 0 && inputActions.Count() < NumParallelProcesses)
+				//{
+				//	_cancellationSource.Cancel();
+				//}
 			}
 			else
 			{
diff --git a/Engine/Source/Programs/UnrealBuildTool/System/ActionGraph.cs b/Engine/Source/Programs/UnrealBuildTool/System/ActionGraph.cs
--- a/Engine/Source/Programs/UnrealBuildTool/System/ActionGraph.cs
+++ b/Engine/Source/Programs/UnrealBuildTool/System/ActionGraph.cs
@@ -394,7 +394,7 @@ namespace UnrealBuildTool
 			}
 			if (BuildConfiguration.bAllowUBAExecutor && UBAExecutor.IsAvailable())
 			{
-				return new UBAExecutor(BuildConfiguration.MaxParallelActions, BuildConfiguration.bAllCores, BuildConfiguration.bCompactOutput, Logger, TargetDescriptors);
+				return new UBAExecutor(BuildConfiguration.MaxParallelActions, BuildConfiguration.bAllCores, BuildConfiguration.bCompactOutput, Logger, TargetDescriptors, true);
 			}
 
 			return null;
@@ -433,7 +433,7 @@ namespace UnrealBuildTool
 						// Intentionally not checking MinActionsForRemote
 						if (BuildConfiguration.bAllowUBAExecutor && UBAExecutor.IsAvailable())
 						{
-							return new UBAExecutor(BuildConfiguration.MaxParallelActions, BuildConfiguration.bAllCores, BuildConfiguration.bCompactOutput, Logger, TargetDescriptors);
+							return new UBAExecutor(BuildConfiguration.MaxParallelActions, BuildConfiguration.bAllCores, BuildConfiguration.bCompactOutput, Logger, TargetDescriptors, true);
 						}
 						return null;
 					}
