UnrealBuildTool: Avoid printing entire exception for agent allocation errors

diff --git a/Engine/Source/Programs/UnrealBuildTool/Executors/UnrealBuildAccelerator/UBAAgentCoordinatorHorde.cs b/Engine/Source/Programs/UnrealBuildTool/Executors/UnrealBuildAccelerator/UBAAgentCoordinatorHorde.cs
--- a/Engine/Source/Programs/UnrealBuildTool/Executors/UnrealBuildAccelerator/UBAAgentCoordinatorHorde.cs
+++ b/Engine/Source/Programs/UnrealBuildTool/Executors/UnrealBuildAccelerator/UBAAgentCoordinatorHorde.cs
@@ -980,7 +980,7 @@ namespace UnrealBuildTool
 					if (!_cancellationToken.IsCancellationRequested)
 					{
 						hordeSession.UpdateHordeStatus(" - " + ex.Message);
-						_logger.Log(logLevel, KnownLogEvents.Systemic_Horde_Compute, ex, "Unable to get worker: {Ex}", ex.ToString());
+						_logger.Log(logLevel, KnownLogEvents.Systemic_Horde_Compute, ex, "Unable to get worker: {Ex}", ex.Message);
 					}
 				}
 
