[UBT]
* Added option to control if retry can be done natively or not
* Changed so compress mode allow retry but not natively since it is using virtual files

diff --git a/Engine/Source/Programs/UnrealBuildTool/Executors/UnrealBuildAccelerator/UBAExecutor.cs b/Engine/Source/Programs/UnrealBuildTool/Executors/UnrealBuildAccelerator/UBAExecutor.cs
--- a/Engine/Source/Programs/UnrealBuildTool/Executors/UnrealBuildAccelerator/UBAExecutor.cs
+++ b/Engine/Source/Programs/UnrealBuildTool/Executors/UnrealBuildAccelerator/UBAExecutor.cs
@@ -24,6 +24,8 @@ namespace UnrealBuildTool
 
 		public string Crypto { get; private set; } = String.Empty;
 
+		public bool CanRetryNatively { get; set; } = true;
+
 		readonly IEnumerable<TargetDescriptor> _targetDescriptors;
 		readonly Task? _initTask;
 		readonly CancellationTokenSource _cancellationSource;
@@ -94,10 +96,7 @@ namespace UnrealBuildTool
 
 		public override string Name => !UBAConfig.bDisableRemote && _agentCoordinators.Count > 0 ? "Unreal Build Accelerator" : "Unreal Build Accelerator local";
 
-		public static new bool IsAvailable()
-		{
-			return EpicGames.UBA.Utils.IsAvailable();
-		}
+		public static new bool IsAvailable() => EpicGames.UBA.Utils.IsAvailable();
 
 		public delegate void SessionCreatedDelegate(ISessionServer Session);
 
@@ -781,7 +780,7 @@ namespace UnrealBuildTool
 						{
 							if (UBAConfig.bAllowRetry && (exitCode != 0 && UBAConfig.bForcedRetry || (exitCode >= 9000 && exitCode < 10000)))
 							{
-								if (UBAConfig.bStoreObjFilesCompressed)
+								if (!CanRetryNatively || UBAConfig.bStoreObjFilesCompressed)
 								{
 									logger.LogInformation("{Description} {StatusDescription}: Exited with error code {ExitCode}. This action will retry", action.CommandDescription, action.StatusDescription, exitCode);
 									++forcedRetryActions;
