[UBT]
* Minor optimization in VCEnvironment. Use DirectoryItem instead of File.Exists since it is cached (and called a ton of times during project generation)

diff --git a/Engine/Source/Programs/UnrealBuildTool/Platform/Windows/VCEnvironment.cs b/Engine/Source/Programs/UnrealBuildTool/Platform/Windows/VCEnvironment.cs
--- a/Engine/Source/Programs/UnrealBuildTool/Platform/Windows/VCEnvironment.cs
+++ b/Engine/Source/Programs/UnrealBuildTool/Platform/Windows/VCEnvironment.cs
@@ -190,7 +190,8 @@ namespace UnrealBuildTool
 		public static DirectoryReference GetVCToolPath(DirectoryReference VCToolChainDir, UnrealArch Architecture)
 		{
 			FileReference CompilerPath = FileReference.Combine(VCToolChainDir, "bin", MicrosoftPlatformSDK.MSVCHostDirectoryName, Architecture.WindowsToolChain, "cl.exe");
-			if (FileReference.Exists(CompilerPath))
+			FileItem FileItem = FileItem.GetItemByFileReference(CompilerPath); // FileItem is cached
+			if (FileItem.Exists)
 			{
 				return CompilerPath.Directory;
 			}
@@ -282,7 +283,8 @@ namespace UnrealBuildTool
 		{
 			string BuildHostArch = RuntimeInformation.ProcessArchitecture == System.Runtime.InteropServices.Architecture.Arm64 ? "arm64" : "x64";
 			FileReference ResourceCompilerPath = FileReference.Combine(WindowsSdkDir, "bin", WindowsSdkVersion.ToString(), BuildHostArch, "rc.exe");
-			if (FileReference.Exists(ResourceCompilerPath))
+			FileItem FileItem = FileItem.GetItemByFileReference(ResourceCompilerPath);
+			if (FileItem.Exists)
 			{
 				return ResourceCompilerPath;
 			}
