[UBT]
* Optimized UEBuildModule.ValidateCircularReferences. Before opts (using biggest project we have): 35sec thread time, after 4sec thread time.

diff --git a/Engine/Source/Programs/UnrealBuildTool/Configuration/UEBuildModule.cs b/Engine/Source/Programs/UnrealBuildTool/Configuration/UEBuildModule.cs
--- a/Engine/Source/Programs/UnrealBuildTool/Configuration/UEBuildModule.cs
+++ b/Engine/Source/Programs/UnrealBuildTool/Configuration/UEBuildModule.cs
@@ -173,7 +173,7 @@ namespace UnrealBuildTool
 		/// <summary>
 		/// Set of modules that have referenced this module
 		/// </summary>
-		protected HashSet<UEBuildModule>? ReferenceStackParentModules;
+		protected SortedSet<UEBuildModule>? ReferenceStackParentModules;
 
 		/// <summary>
 		/// Set of all allowed restricted folder references
@@ -1231,14 +1231,18 @@ namespace UnrealBuildTool
 					LogDependencyNameList("Private Include Paths:", Rules.PrivateIncludePathModuleNames, Logger);
 				}
 
-				ReferenceStackParentModules ??= new();
+				ReferenceStackParentModules ??= new(Comparer<UEBuildModule>.Create(static (a, b) => a.Name.CompareTo(b.Name)));
 
 				// Create all the dependency modules - pass through the reference stack so we can check for cycles
 				RecursivelyCreateModulesByName(Rules.PublicDependencyModuleNames, ref PublicDependencyModules, ref bDependsOnVerse, CreateModule, NextReferenceChain, ReferenceStack, Logger);
 				if (Rules.Target.IsTestTarget)
 				{
 					// Move the test runner dependency to last position to give it the opportunity to build its special dependencies such as CoreUObject, ApplicationCore etc.
-					MoveTestsRunnerDependencyToLastPosition();
+					if (Rules.PrivateDependencyModuleNames.Contains("LowLevelTestsRunner"))
+					{
+						Rules.PrivateDependencyModuleNames.Remove("LowLevelTestsRunner");
+						Rules.PrivateDependencyModuleNames.Add("LowLevelTestsRunner");
+					}
 				}
 				RecursivelyCreateModulesByName(Rules.PrivateDependencyModuleNames, ref PrivateDependencyModules, ref bDependsOnVerse, CreateModule, NextReferenceChain, ReferenceStack, Logger);
 
@@ -1392,7 +1396,7 @@ namespace UnrealBuildTool
 			parentStack.Add(this);
 
 			// Check very module that references this module until all are processed
-			foreach (UEBuildModule ParentModule in ReferenceStackParentModules.OrderBy(x => x.Name))
+			foreach (UEBuildModule ParentModule in ReferenceStackParentModules)
 			{
 				ParentModule.CheckForCircularReferences(module, processed, parentStack, circularReferences);
 			}
@@ -1401,15 +1405,6 @@ namespace UnrealBuildTool
 			parentStack.RemoveAt(parentStack.Count - 1);
 		}
 
-		private void MoveTestsRunnerDependencyToLastPosition()
-		{
-			if (Rules.PrivateDependencyModuleNames.Contains("LowLevelTestsRunner"))
-			{
-				Rules.PrivateDependencyModuleNames.Remove("LowLevelTestsRunner");
-				Rules.PrivateDependencyModuleNames.Add("LowLevelTestsRunner");
-			}
-		}
-
 		private static void LogDependencyNameList(string Title, List<string> DependencyNameList, ILogger logger)
 		{
 			logger.LogTrace("  {Title}", Title);
