[UBT]
* Changed List to HashSet for circular module dependencies in module rules since this container is only used to call contains

diff --git a/Engine/Source/Programs/UnrealBuildTool/Configuration/ModuleRules.cs b/Engine/Source/Programs/UnrealBuildTool/Configuration/ModuleRules.cs
--- a/Engine/Source/Programs/UnrealBuildTool/Configuration/ModuleRules.cs
+++ b/Engine/Source/Programs/UnrealBuildTool/Configuration/ModuleRules.cs
@@ -1199,7 +1199,7 @@ namespace UnrealBuildTool
 		/// Only for legacy reason, should not be used in new code. List of module dependencies that should be treated as circular references.  This modules must have already been added to
 		/// either the public or private dependent module list.
 		/// </summary>
-		public List<string> CircularlyReferencedDependentModules = [];
+		public HashSet<string> CircularlyReferencedDependentModules = [];
 
 		/// <summary>
 		/// List of system/library include paths - typically used for External (third party) modules.  These are public stable header file directories that are not checked when resolving header dependencies.
@@ -2087,4 +2087,33 @@ namespace UnrealBuildTool
 		})).Value;
 		static readonly ConcurrentDictionary<Type, Lazy<HashSet<UnrealTargetPlatform>?>> s_supportedPlatformsAttribute = new();
 	}
+
+	/// <summary>
+	/// AddRange added to HashSet
+	/// </summary>
+	public static class HashSetExtensions
+	{
+		/// <summary>
+		/// Add a range of objects
+		/// </summary>
+		/// <typeparam name="T"></typeparam>
+		/// <param name="set"></param>
+		/// <param name="items"></param>
+		/// <exception cref="ArgumentNullException"></exception>
+		public static void AddRange<T>(this HashSet<T> set, IEnumerable<T> items)
+		{
+			set.UnionWith(items);
+		}
+
+		/// <summary>
+		/// Add a range of objects
+		/// </summary>
+		/// <typeparam name="T"></typeparam>
+		/// <param name="set"></param>
+		/// <param name="items"></param>
+		public static void AddRange<T>(this HashSet<T> set, params T[] items)
+		{
+			foreach (var item in items) set.Add(item);
+		}
+	}
 }
\ No newline at end of file
