[UBT]
* Removed unused parameter when creating UEBuildModules

#FailedValidations Preflight

diff --git a/Engine/Source/Programs/UnrealBuildTool/Configuration/UEBuildModule.cs b/Engine/Source/Programs/UnrealBuildTool/Configuration/UEBuildModule.cs
--- a/Engine/Source/Programs/UnrealBuildTool/Configuration/UEBuildModule.cs
+++ b/Engine/Source/Programs/UnrealBuildTool/Configuration/UEBuildModule.cs
@@ -1154,26 +1154,13 @@ namespace UnrealBuildTool
 
 		public delegate UEBuildModule CreateModuleDelegate(string Name, string ReferenceChain);
 
-		/// <summary>
-		/// Public entry point to recursively create a module and all its dependencies
-		/// </summary>
-		/// <param name="CreateModule"></param>
-		/// <param name="ReferenceChain"></param>
-		/// <param name="Logger"></param>
-		public void RecursivelyCreateModules(CreateModuleDelegate CreateModule, string ReferenceChain, ILogger Logger)
-		{
-			List<UEBuildModule> ReferenceStack = new List<UEBuildModule>();
-			RecursivelyCreateModules(CreateModule, ReferenceChain, ReferenceStack, Logger);
-		}
-
 		/// <summary>
 		/// Creates all the modules required for this target
 		/// </summary>
 		/// <param name="CreateModule">Delegate to create a module with a given name</param>
 		/// <param name="ReferenceChain">Chain of references before reaching this module</param>
-		/// <param name="ReferenceStack">ParentStack of module dependencies that led to this module</param>
 		/// <param name="Logger">Logger for output</param>
-		protected void RecursivelyCreateModules(CreateModuleDelegate CreateModule, string ReferenceChain, List<UEBuildModule> ReferenceStack, ILogger Logger)
+		public void RecursivelyCreateModules(CreateModuleDelegate CreateModule, string ReferenceChain, ILogger Logger)
 		{
 			// Name of this reference
 			string ThisRefName = (RulesFile == null) ? Name : RulesFile.GetFileName();
@@ -1181,9 +1168,6 @@ namespace UnrealBuildTool
 			// Set the reference chain for anything referenced by this module
 			string NextReferenceChain = String.Format("{0} -> {1}", ReferenceChain, ThisRefName);
 
-			// Add us to the reference stack
-			ReferenceStack.Add(this);
-
 			// Check if this module is invalid for the current target, based on Module attributes
 			if (Rules.CppCompileWarningSettings.ModuleUnsupportedWarningLevel != WarningLevel.Off)
 			{
@@ -1234,7 +1218,7 @@ namespace UnrealBuildTool
 				ReferenceStackParentModules ??= new(Comparer<UEBuildModule>.Create(static (a, b) => a.Name.CompareTo(b.Name)));
 
 				// Create all the dependency modules - pass through the reference stack so we can check for cycles
-				RecursivelyCreateModulesByName(Rules.PublicDependencyModuleNames, ref PublicDependencyModules, ref bDependsOnVerse, CreateModule, NextReferenceChain, ReferenceStack, Logger);
+				RecursivelyCreateModulesByName(Rules.PublicDependencyModuleNames, ref PublicDependencyModules, ref bDependsOnVerse, CreateModule, NextReferenceChain, Logger);
 				if (Rules.Target.IsTestTarget)
 				{
 					// Move the test runner dependency to last position to give it the opportunity to build its special dependencies such as CoreUObject, ApplicationCore etc.
@@ -1244,7 +1228,7 @@ namespace UnrealBuildTool
 						Rules.PrivateDependencyModuleNames.Add("LowLevelTestsRunner");
 					}
 				}
-				RecursivelyCreateModulesByName(Rules.PrivateDependencyModuleNames, ref PrivateDependencyModules, ref bDependsOnVerse, CreateModule, NextReferenceChain, ReferenceStack, Logger);
+				RecursivelyCreateModulesByName(Rules.PrivateDependencyModuleNames, ref PrivateDependencyModules, ref bDependsOnVerse, CreateModule, NextReferenceChain, Logger);
 
 				// Recursively create all the public include path modules
 				HashSet<string> PublicIncludePathModuleNames = new();
@@ -1256,14 +1240,11 @@ namespace UnrealBuildTool
 				RecursivelyCreateIncludePathModulesByName(Rules.PrivateIncludePathModuleNames, ref PrivateIncludePathModules, ref bDependsOnVerse, CreateModule, NextReferenceChain);
 
 				// Dynamic loads aren't considered a reference chain so start with an empty stack
-				RecursivelyCreateModulesByName(Rules.DynamicallyLoadedModuleNames, ref DynamicallyLoadedModules, ref bDependsOnVerse, CreateModule, NextReferenceChain, new List<UEBuildModule>(), Logger);
+				RecursivelyCreateModulesByName(Rules.DynamicallyLoadedModuleNames, ref DynamicallyLoadedModules, ref bDependsOnVerse, CreateModule, NextReferenceChain, Logger);
 
 				PublicDependencyModules?.ForEach(x => x.ReferenceStackParentModules?.Add(this));
 				PrivateDependencyModules?.ForEach(x => x.ReferenceStackParentModules?.Add(this));
 			}
-
-			// pop us off the current stack
-			ReferenceStack.RemoveAt(ReferenceStack.Count - 1);
 		}
 
 		record CircularReference
@@ -1414,7 +1395,7 @@ namespace UnrealBuildTool
 			}
 		}
 
-		private static void RecursivelyCreateModulesByName(List<string> ModuleNames, ref List<UEBuildModule>? Modules, ref bool bDependsOnVerse, CreateModuleDelegate CreateModule, string ReferenceChain, List<UEBuildModule> ReferenceStack, ILogger Logger)
+		private static void RecursivelyCreateModulesByName(List<string> ModuleNames, ref List<UEBuildModule>? Modules, ref bool bDependsOnVerse, CreateModuleDelegate CreateModule, string ReferenceChain, ILogger Logger)
 		{
 			// Check whether the module list is already set. We set this immediately (via the ref) to avoid infinite recursion.
 			if (Modules == null)
@@ -1426,7 +1407,7 @@ namespace UnrealBuildTool
 					UEBuildModule Module = CreateModule(ModuleName, ReferenceChain);
 					if (!Modules.Contains(Module))
 					{
-						Module.RecursivelyCreateModules(CreateModule, ReferenceChain, ReferenceStack, Logger);
+						Module.RecursivelyCreateModules(CreateModule, ReferenceChain, Logger);
 						Modules.Add(Module);
 						bDependsOnVerse |= Module.bDependsOnVerse;
 					}
