[UBT]
* Implemented file hash cache that can be used to fetch hashes of files (will be recalculated if file's last written time has changed)
* Replaced makefile logic to use file hash cache. this way files are not read to be compared if they already exist in the hash cache

#FailedValidations Preflight

diff --git a/Engine/Source/Programs/UnrealBuildTool/Configuration/UEBuildTarget.cs b/Engine/Source/Programs/UnrealBuildTool/Configuration/UEBuildTarget.cs
--- a/Engine/Source/Programs/UnrealBuildTool/Configuration/UEBuildTarget.cs
+++ b/Engine/Source/Programs/UnrealBuildTool/Configuration/UEBuildTarget.cs
@@ -996,7 +996,7 @@ namespace UnrealBuildTool
 			{
 				if (Architectures.Count > 1)
 				{
-					throw new BuildException($"Asking CppCompileEnvironment for a single Architecture, but it has multiple Architectures ({String.Join(", ", Architectures)}). This indicates logic needs to be updated");
+					throw new Exception($"Asking CppCompileEnvironment for a single Architecture, but it has multiple Architectures ({String.Join(", ", Architectures)}). This indicates logic needs to be updated");
 				}
 				return Architectures[0];
 			}
@@ -2611,7 +2611,7 @@ namespace UnrealBuildTool
 
 			RunCodeGen(BuildConfiguration, [codeGenContext], Logger);
 
-			BuildMakefile2(WorkingSet, Logger, context);
+			BuildMakefile2(WorkingSet, Logger, context, new(false, true));
 
 			codeGenContext.uhtRun?.Wait();
 
@@ -2932,7 +2932,7 @@ namespace UnrealBuildTool
 			return timestampTasks;
 		}
 
-		internal void BuildMakefile2(ISourceFileWorkingSet WorkingSet, ILogger Logger, BuildContext context)
+		internal void BuildMakefile2(ISourceFileWorkingSet WorkingSet, ILogger Logger, BuildContext context, FileHashCache fileHashCache)
 		{
 			SourceFileMetadataCache MetadataCache = context.MetadataCache;
 			TargetDescriptor TargetDescriptor = context.Descriptor;
@@ -2943,7 +2943,7 @@ namespace UnrealBuildTool
 			UEToolChain TargetToolChain = context.TargetToolChain;
 			LinkEnvironment GlobalLinkEnvironment = context.GlobalLinkEnvironment;
 
-			TargetMakefileBuilder MakefileBuilder = new TargetMakefileBuilder(Makefile, Logger);
+			TargetMakefileBuilder MakefileBuilder = new TargetMakefileBuilder(Makefile, fileHashCache, Logger);
 
 			if (Rules.bUseSharedPCHs)
 			{
@@ -4374,9 +4374,13 @@ namespace UnrealBuildTool
 			AddExtraModules(Logger);
 
 			// Create all the modules referenced by the existing binaries
+			UEBuildModule CreateModuleDelegate(string Name, string RefChain) => FindOrCreateModuleByName(Name, RefChain, Logger);
 			foreach (UEBuildBinary Binary in Binaries)
 			{
-				Binary.CreateAllDependentModules((Name, RefChain) => FindOrCreateModuleByName(Name, RefChain, Logger), Logger);
+				foreach (UEBuildModule Module in Binary.Modules)
+				{
+					Module.RecursivelyCreateModules(CreateModuleDelegate, "Target", Logger);
+				}
 			}
 
 			// Cache dependencies since we can do this in parallel for all modules
@@ -6601,135 +6605,146 @@ namespace UnrealBuildTool
 		/// <param name="ModuleName">Name of the module</param>
 		/// <param name="ReferenceChain">Chain of references causing this module to be instantiated, for display in error messages</param>
 		/// <param name="Logger">Logger for output</param>
-		public UEBuildModule FindOrCreateModuleByName(string ModuleName, string ReferenceChain, ILogger Logger)
+		public UEBuildModule CreateModuleByName(string ModuleName, string ReferenceChain, ILogger Logger)
 		{
-			UEBuildModule? Module;
-			if (!Modules.TryGetValue(ModuleName, out Module))
+			// @todo projectfiles: Cross-platform modules can appear here during project generation, but they may have already
+			//   been filtered out by the project generator.  This causes the projects to not be added to directories properly.
+			ModuleRules RulesObject = CreateModuleRulesAndSetDefaults(ModuleName, ReferenceChain, Logger);
+			DirectoryReference ModuleDirectory = RulesObject.File.Directory;
+
+			// Clear the bUsePrecompiled flag if we're compiling a foreign plugin; since it's treated like an engine module, it will default to true in an installed build.
+			if (RulesObject.Plugin != null && RulesObject.Plugin.File == ForeignPlugin)
 			{
-				// @todo projectfiles: Cross-platform modules can appear here during project generation, but they may have already
-				//   been filtered out by the project generator.  This causes the projects to not be added to directories properly.
-				ModuleRules RulesObject = CreateModuleRulesAndSetDefaults(ModuleName, ReferenceChain, Logger);
-				DirectoryReference ModuleDirectory = RulesObject.File.Directory;
+				RulesObject.bPrecompile = true;
+				RulesObject.bUsePrecompiled = false;
+			}
 
-				// Clear the bUsePrecompiled flag if we're compiling a foreign plugin; since it's treated like an engine module, it will default to true in an installed build.
-				if (RulesObject.Plugin != null && RulesObject.Plugin.File == ForeignPlugin)
+			// Get the base directory for paths referenced by the module. If the module's under the UProject source directory use that, otherwise leave it relative to the Engine source directory.
+			if (ProjectFile != null)
+			{
+				DirectoryReference ProjectSourceDirectoryName = DirectoryReference.Combine(ProjectFile.Directory, "Source");
+				if (RulesObject.File.IsUnderDirectory(ProjectSourceDirectoryName))
 				{
-					RulesObject.bPrecompile = true;
-					RulesObject.bUsePrecompiled = false;
+					RulesObject.PublicIncludePaths = CombinePathList(ProjectSourceDirectoryName, RulesObject.PublicIncludePaths);
+					RulesObject.InternalIncludePaths = CombinePathList(ProjectSourceDirectoryName, RulesObject.InternalIncludePaths);
+					RulesObject.PrivateIncludePaths = CombinePathList(ProjectSourceDirectoryName, RulesObject.PrivateIncludePaths);
+					RulesObject.PublicSystemLibraryPaths = CombinePathList(ProjectSourceDirectoryName, RulesObject.PublicSystemLibraryPaths);
 				}
+			}
 
-				// Get the base directory for paths referenced by the module. If the module's under the UProject source directory use that, otherwise leave it relative to the Engine source directory.
-				if (ProjectFile != null)
+			// Get the generated code directory. Plugins always write to their own intermediate directory so they can be copied between projects, shared engine
+			// intermediates go in the engine intermediate folder, and anything else goes in the project folder.
+			DirectoryReference? GeneratedCodeDirectory = null;
+			if (RulesObject.Type != ModuleRules.ModuleType.External)
+			{
+				// Get the base directory
+				if (bUseSharedBuildEnvironment)
 				{
-					DirectoryReference ProjectSourceDirectoryName = DirectoryReference.Combine(ProjectFile.Directory, "Source");
-					if (RulesObject.File.IsUnderDirectory(ProjectSourceDirectoryName))
-					{
-						RulesObject.PublicIncludePaths = CombinePathList(ProjectSourceDirectoryName, RulesObject.PublicIncludePaths);
-						RulesObject.InternalIncludePaths = CombinePathList(ProjectSourceDirectoryName, RulesObject.InternalIncludePaths);
-						RulesObject.PrivateIncludePaths = CombinePathList(ProjectSourceDirectoryName, RulesObject.PrivateIncludePaths);
-						RulesObject.PublicSystemLibraryPaths = CombinePathList(ProjectSourceDirectoryName, RulesObject.PublicSystemLibraryPaths);
-					}
+					GeneratedCodeDirectory = RulesObject.Context.DefaultOutputBaseDir;
 				}
-
-				// Get the generated code directory. Plugins always write to their own intermediate directory so they can be copied between projects, shared engine
-				// intermediates go in the engine intermediate folder, and anything else goes in the project folder.
-				DirectoryReference? GeneratedCodeDirectory = null;
-				if (RulesObject.Type != ModuleRules.ModuleType.External)
+				else
 				{
-					// Get the base directory
-					if (bUseSharedBuildEnvironment)
-					{
-						GeneratedCodeDirectory = RulesObject.Context.DefaultOutputBaseDir;
-					}
-					else
-					{
-						GeneratedCodeDirectory = ProjectDirectory;
-					}
-
-					// Get the subfolder containing generated code - we don't need architecture information since these are shared between all arches for a platform, as well as shared between all intermediate environment variants
-					GeneratedCodeDirectory = DirectoryReference.Combine(GeneratedCodeDirectory, PlatformIntermediateFolderNoArch, AppName, "Inc");
+					GeneratedCodeDirectory = ProjectDirectory;
+				}
 
-					// Append the binaries subfolder, if present. We rely on this to ensure that build products can be filtered correctly.
-					if (RulesObject.BinariesSubFolder != null)
-					{
-						GeneratedCodeDirectory = DirectoryReference.Combine(GeneratedCodeDirectory, RulesObject.BinariesSubFolder);
-					}
+				// Get the subfolder containing generated code - we don't need architecture information since these are shared between all arches for a platform, as well as shared between all intermediate environment variants
+				GeneratedCodeDirectory = DirectoryReference.Combine(GeneratedCodeDirectory, PlatformIntermediateFolderNoArch, AppName, "Inc");
 
-					// Finally, append the module name (using the ShortName if it has been set)
-					GeneratedCodeDirectory = DirectoryReference.Combine(GeneratedCodeDirectory, RulesObject.ShortName ?? ModuleName);
+				// Append the binaries subfolder, if present. We rely on this to ensure that build products can be filtered correctly.
+				if (RulesObject.BinariesSubFolder != null)
+				{
+					GeneratedCodeDirectory = DirectoryReference.Combine(GeneratedCodeDirectory, RulesObject.BinariesSubFolder);
 				}
 
-				// For legacy modules, add a bunch of default include paths.
-				if (RulesObject.Type == ModuleRules.ModuleType.CPlusPlus && RulesObject.bAddDefaultIncludePaths && (RulesObject.Plugin != null || (ProjectFile != null && RulesObject.File.IsUnderDirectory(ProjectFile.Directory))))
+				// Finally, append the module name (using the ShortName if it has been set)
+				GeneratedCodeDirectory = DirectoryReference.Combine(GeneratedCodeDirectory, RulesObject.ShortName ?? ModuleName);
+			}
+
+			// For legacy modules, add a bunch of default include paths.
+			if (RulesObject.Type == ModuleRules.ModuleType.CPlusPlus && RulesObject.bAddDefaultIncludePaths && (RulesObject.Plugin != null || (ProjectFile != null && RulesObject.File.IsUnderDirectory(ProjectFile.Directory))))
+			{
+				// Add the module source directory
+				DirectoryReference BaseSourceDirectory;
+				if (RulesObject.Plugin != null)
 				{
-					// Add the module source directory
-					DirectoryReference BaseSourceDirectory;
-					if (RulesObject.Plugin != null)
-					{
-						BaseSourceDirectory = DirectoryReference.Combine(RulesObject.Plugin.Directory, "Source");
-					}
-					else
-					{
-						BaseSourceDirectory = DirectoryReference.Combine(ProjectFile!.Directory, "Source");
-					}
+					BaseSourceDirectory = DirectoryReference.Combine(RulesObject.Plugin.Directory, "Source");
+				}
+				else
+				{
+					BaseSourceDirectory = DirectoryReference.Combine(ProjectFile!.Directory, "Source");
+				}
 
-					// If it's a game module (plugin or otherwise), add the root source directory to the include paths.
-					if (RulesObject.File.IsUnderDirectory(BaseSourceDirectory) || (RulesObject.Plugin != null && RulesObject.Plugin.LoadedFrom == PluginLoadedFrom.Project))
+				// If it's a game module (plugin or otherwise), add the root source directory to the include paths.
+				if (RulesObject.File.IsUnderDirectory(BaseSourceDirectory) || (RulesObject.Plugin != null && RulesObject.Plugin.LoadedFrom == PluginLoadedFrom.Project))
+				{
+					if (DirectoryItem.GetItemByDirectoryReference(BaseSourceDirectory).Exists)
 					{
-						if (DirectoryItem.GetItemByDirectoryReference(BaseSourceDirectory).Exists)
-						{
-							RulesObject.PublicIncludePaths.Add(NormalizeIncludePath(BaseSourceDirectory));
-						}
+						RulesObject.PublicIncludePaths.Add(NormalizeIncludePath(BaseSourceDirectory));
 					}
+				}
 
-					// Resolve private include paths against the project source root
-					for (int Idx = 0; Idx < RulesObject.PrivateIncludePaths.Count; Idx++)
+				// Resolve private include paths against the project source root
+				for (int Idx = 0; Idx < RulesObject.PrivateIncludePaths.Count; Idx++)
+				{
+					string PrivateIncludePath = RulesObject.PrivateIncludePaths[Idx];
+					if (!Path.IsPathRooted(PrivateIncludePath))
 					{
-						string PrivateIncludePath = RulesObject.PrivateIncludePaths[Idx];
-						if (!Path.IsPathRooted(PrivateIncludePath))
-						{
-							PrivateIncludePath = DirectoryReference.Combine(BaseSourceDirectory, PrivateIncludePath).FullName;
-						}
-						RulesObject.PrivateIncludePaths[Idx] = PrivateIncludePath;
+						PrivateIncludePath = DirectoryReference.Combine(BaseSourceDirectory, PrivateIncludePath).FullName;
 					}
+					RulesObject.PrivateIncludePaths[Idx] = PrivateIncludePath;
 				}
+			}
 
-				// Allow the current target to modify the module rules
-				ModifyModuleRulesForTarget(ModuleName, RulesObject);
+			// Allow the current target to modify the module rules
+			ModifyModuleRulesForTarget(ModuleName, RulesObject);
 
-				// Allow the current platform to modify the module rules
-				UEBuildPlatform.GetBuildPlatform(Platform).ModifyModuleRulesForActivePlatform(ModuleName, RulesObject, Rules);
+			// Allow the current platform to modify the module rules
+			UEBuildPlatform.GetBuildPlatform(Platform).ModifyModuleRulesForActivePlatform(ModuleName, RulesObject, Rules);
 
-				// Allow all build platforms to 'adjust' the module setting.
-				// This will allow undisclosed platforms to make changes without
-				// exposing information about the platform in publicly accessible
-				// locations.
-				UEBuildPlatform.PlatformModifyHostModuleRules(ModuleName, RulesObject, Rules);
+			// Allow all build platforms to 'adjust' the module setting.
+			// This will allow undisclosed platforms to make changes without
+			// exposing information about the platform in publicly accessible
+			// locations.
+			UEBuildPlatform.PlatformModifyHostModuleRules(ModuleName, RulesObject, Rules);
 
-				// If there are files to generate for module, then we have to create Gen directory and add it as public include
-				if (RulesObject.FilesToGenerate.Count != 0)
-				{
-					DirectoryReference GenDir = DirectoryReference.Combine(GetModuleIntermediateDirectory(RulesObject, Architectures), "Gen");
-					DirectoryReference.CreateDirectory(GenDir);
-					RulesObject.PublicIncludePaths.Add(NormalizeIncludePath(GenDir));
-				}
+			// If there are files to generate for module, then we have to create Gen directory and add it as public include
+			if (RulesObject.FilesToGenerate.Count != 0)
+			{
+				DirectoryReference GenDir = DirectoryReference.Combine(GetModuleIntermediateDirectory(RulesObject, Architectures), "Gen");
+				DirectoryReference.CreateDirectory(GenDir);
+				RulesObject.PublicIncludePaths.Add(NormalizeIncludePath(GenDir));
+			}
 
-				// Now, go ahead and create the module builder instance
-				Module = InstantiateModule(RulesObject, GeneratedCodeDirectory, Logger);
+			// Now, go ahead and create the module builder instance
+			UEBuildModule Module = InstantiateModule(RulesObject, GeneratedCodeDirectory, Logger);
 
-				// Check if this module conflicts with the ShortName of any other existing modules
-				if (!UsedShortNames.TryAdd(Module.Rules.ShortName ?? Module.Name, Module))
-				{
-					UEBuildModule? ConflictingModule = UsedShortNames[Module.Rules.ShortName ?? Module.Name];
-					Logger.LogError("Conflicting module short names found, unable to create module {Name}", Module.Name);
-					Logger.LogInformation(" * Name={Name} ShortName={ShortName} Path={Path}", ConflictingModule.Name, ConflictingModule.Rules.ShortName ?? ConflictingModule.Name, ConflictingModule.RulesFile);
-					Logger.LogInformation(" * Name={Name} ShortName={ShortName} Path={Path}", Module.Name, Module.Rules.ShortName ?? Module.Name, Module.RulesFile);
+			// Check if this module conflicts with the ShortName of any other existing modules
+			if (!UsedShortNames.TryAdd(Module.Rules.ShortName ?? Module.Name, Module))
+			{
+				UEBuildModule? ConflictingModule = UsedShortNames[Module.Rules.ShortName ?? Module.Name];
+				Logger.LogError("Conflicting module short names found, unable to create module {Name}", Module.Name);
+				Logger.LogInformation(" * Name={Name} ShortName={ShortName} Path={Path}", ConflictingModule.Name, ConflictingModule.Rules.ShortName ?? ConflictingModule.Name, ConflictingModule.RulesFile);
+				Logger.LogInformation(" * Name={Name} ShortName={ShortName} Path={Path}", Module.Name, Module.Rules.ShortName ?? Module.Name, Module.RulesFile);
 
-					throw new BuildException("Unable to create module {0} (ShortName={1}) as it conflicts with existing module {2} (ShortName={3})",
-						Module.Name, Module.Rules.ShortName ?? Module.Name,
-						ConflictingModule.Name, ConflictingModule.Rules.ShortName ?? ConflictingModule.Name);
-				}
+				throw new BuildException("Unable to create module {0} (ShortName={1}) as it conflicts with existing module {2} (ShortName={3})",
+					Module.Name, Module.Rules.ShortName ?? Module.Name,
+					ConflictingModule.Name, ConflictingModule.Rules.ShortName ?? ConflictingModule.Name);
+			}
+			return Module;
+		}
 
+		/// <summary>
+		/// Finds a module given its name.  Throws an exception if the module couldn't be found.
+		/// </summary>
+		/// <param name="ModuleName">Name of the module</param>
+		/// <param name="ReferenceChain">Chain of references causing this module to be instantiated, for display in error messages</param>
+		/// <param name="Logger">Logger for output</param>
+		public UEBuildModule FindOrCreateModuleByName(string ModuleName, string ReferenceChain, ILogger Logger)
+		{
+			UEBuildModule? Module;
+			if (!Modules.TryGetValue(ModuleName, out Module))
+			{
+				Module = CreateModuleByName(ModuleName, ReferenceChain, Logger);
 				Modules.Add(Module.Name, Module);
 			}
 
diff --git a/Engine/Source/Programs/UnrealBuildTool/Modes/BuildMode.cs b/Engine/Source/Programs/UnrealBuildTool/Modes/BuildMode.cs
--- a/Engine/Source/Programs/UnrealBuildTool/Modes/BuildMode.cs
+++ b/Engine/Source/Programs/UnrealBuildTool/Modes/BuildMode.cs
@@ -165,6 +165,7 @@ namespace UnrealBuildTool
 			}
 
 			CppDependencyCache cppDependencyCache = new();
+			FileHashCache fileHashCache = new(true, true);
 			List<Task> tasks = [];
 
 			// Parse and build the targets
@@ -256,7 +257,11 @@ namespace UnrealBuildTool
 					// Create the working set provider per group.
 					using (ISourceFileWorkingSet WorkingSet = SourceFileWorkingSet.Create(Unreal.RootDirectory, ProjectDirs, Logger))
 					{
-						List<(TargetDescriptor, TargetMakefile)> Targets = CreateMakefiles(BuildConfiguration, TargetDescriptors, WorkingSet, Logger, true, true, tasks, cppDependencyCache);
+						List<(TargetDescriptor, TargetMakefile)> Targets = CreateMakefiles(BuildConfiguration, TargetDescriptors, WorkingSet, Logger, true, true, tasks, cppDependencyCache, fileHashCache);
+						if (fileHashCache.HasErrors)
+						{
+							throw new BuildException("FileHashCache had errors");
+						}
 						await BuildAsync(Targets, BuildConfiguration, Options, WriteOutdatedActionsFile, cppDependencyCache, Logger, ActionTypeFilter, executor);
 					}
 				}
@@ -271,6 +276,7 @@ namespace UnrealBuildTool
 				// Save all the caches
 				using (Timeline.ScopeEvent("Saving caches"))
 				{
+					fileHashCache.SaveAll();
 					SourceFileMetadataCache.SaveAll();
 					cppDependencyCache.SaveAll();
 				}
@@ -400,7 +406,7 @@ namespace UnrealBuildTool
 			List<Task> tasks = [];
 			try
 			{
-				List<(TargetDescriptor, TargetMakefile)> Targets = CreateMakefiles(BuildConfiguration, TargetDescriptors, WorkingSet, Logger, true, true, tasks, cppDependencyCache);
+				List<(TargetDescriptor, TargetMakefile)> Targets = CreateMakefiles(BuildConfiguration, TargetDescriptors, WorkingSet, Logger, true, true, tasks, cppDependencyCache, new(false, true));
 				await BuildAsync(Targets, BuildConfiguration, Options, WriteOutdatedActionsFile, cppDependencyCache, Logger, ActionTypeFilter, null);
 			}
 			finally
@@ -1069,7 +1075,7 @@ namespace UnrealBuildTool
 			List<Task> tasks = [];
 			try
 			{
-				List<(TargetDescriptor, TargetMakefile Makefile)> Targets = CreateMakefiles(BuildConfiguration, [TargetDescriptor], WorkingSet, Logger, false, false, tasks, null);
+				List<(TargetDescriptor, TargetMakefile Makefile)> Targets = CreateMakefiles(BuildConfiguration, [TargetDescriptor], WorkingSet, Logger, false, false, tasks, null, new(false, true));
 				return Targets[0].Makefile;
 			}
 			finally
@@ -1103,7 +1109,7 @@ namespace UnrealBuildTool
 			return (makefile, true);
 		}
 
-		private static List<(TargetDescriptor, TargetMakefile)> CreateMakefiles(BuildConfiguration buildConfiguration, List<TargetDescriptor> descriptors, ISourceFileWorkingSet workingSet, ILogger logger, bool runInitScripts, bool skipPreBuildTargets, List<Task> outTasks, CppDependencyCache? cppDependencyCache)
+		private static List<(TargetDescriptor, TargetMakefile)> CreateMakefiles(BuildConfiguration buildConfiguration, List<TargetDescriptor> descriptors, ISourceFileWorkingSet workingSet, ILogger logger, bool runInitScripts, bool skipPreBuildTargets, List<Task> outTasks, CppDependencyCache? cppDependencyCache, FileHashCache fileHashCache)
 		{
 			List<(TargetDescriptor descriptor, TargetMakefile makefile)> outTargets = new();
 
@@ -1111,6 +1117,12 @@ namespace UnrealBuildTool
 
 			Dictionary<FileReference, Task<SourceFileMetadataCache>> sourceFileMetadataCaches = [];
 
+			// Mac toolchain unzips frameworks to a folder outside of Intermediate/Build and also write a .sh file. So we need to mount this directory
+			if (OperatingSystem.IsMacOS())
+			{
+				fileHashCache.Mount(DirectoryReference.Combine(Unreal.WritableEngineDirectory, "Intermediate", "UnzippedFrameworks"), logger);
+			}
+
 			while (descriptors.Count != 0) // descriptors list will be reallocated for PrebuildTargets. hence the while-loop
 			{
 				int descriptorCount = descriptors.Count;
@@ -1173,6 +1185,8 @@ namespace UnrealBuildTool
 						cppDependencyCache.MountAsync(descriptor, buildTarget.AppName, buildTarget.IntermediateEnvironment, logger);
 					}
 
+					fileHashCache.Mount(descriptor, buildTarget.AppName, buildTarget.IntermediateEnvironment, logger);
+
 					// Add target to proper app name group
 					string appName = $"{buildTarget.AppName}_{buildTarget.Platform}";
 					lock (appNameTargetGroups)
@@ -1379,7 +1393,7 @@ namespace UnrealBuildTool
 								{
 									UEBuildTarget buildTarget = buildTargets[index];
 
-									buildTarget.BuildMakefile2(workingSet, logger, buildContext);
+									buildTarget.BuildMakefile2(workingSet, logger, buildContext, fileHashCache);
 
 									// Save the makefile for next time
 									FileReference? makefileLocation = makefileLocations[index];
diff --git a/Engine/Source/Programs/UnrealBuildTool/System/CppDependencyCache.cs b/Engine/Source/Programs/UnrealBuildTool/System/CppDependencyCache.cs
--- a/Engine/Source/Programs/UnrealBuildTool/System/CppDependencyCache.cs
+++ b/Engine/Source/Programs/UnrealBuildTool/System/CppDependencyCache.cs
@@ -477,6 +477,12 @@ namespace UnrealBuildTool
 				if (projectFile != null)
 				{
 					AddPartition(projectFile.Directory, GetFilterPath(arch, configuration));
+
+					// Only needed for targets that has unique build environment
+					if (configuration == UnrealTargetConfiguration.DebugGame)
+					{
+						AddPartition(projectFile.Directory, GetFilterPath(arch, UnrealTargetConfiguration.Development));
+					}
 				}
 			}
 		}
diff --git a/Engine/Source/Programs/UnrealBuildTool/System/FileHashCache.cs b/Engine/Source/Programs/UnrealBuildTool/System/FileHashCache.cs
new file mode 100644
--- /dev/null
+++ b/Engine/Source/Programs/UnrealBuildTool/System/FileHashCache.cs
@@ -0,0 +1,589 @@
+// Copyright Epic Games, Inc. All Rights Reserved.
+
+using EpicGames.Core;
+using Microsoft.Extensions.Logging;
+using System;
+using System.Buffers.Binary;
+using System.Collections.Concurrent;
+using System.Collections.Generic;
+using System.IO;
+using System.IO.Hashing;
+using System.Linq;
+using System.Runtime.InteropServices;
+using System.Text;
+using System.Threading.Tasks;
+using UnrealBuildBase;
+
+namespace UnrealBuildTool
+{
+	/// <summary>
+	/// A database storing hashes of files together with last written time.
+	/// This is used to reduce number of reads just to check if file is up-to-date before writing it.
+	/// </summary>
+	class FileHashCache
+	{
+		/// <summary>
+		/// Ctor
+		/// </summary>
+		/// <param name="warnOnNoPartition">Warn if file does not exist in a mounted partition</param>
+		/// <param name="backupModifiedFiles">Backup old file before writing new changed file</param>
+		/// <param name="errorOnRewrite">Report error if file is rewritten with different content, otherwise just warning</param>
+		public FileHashCache(bool warnOnNoPartition, bool backupModifiedFiles, bool errorOnRewrite = false)
+		{
+			_warnOnNoPartition = warnOnNoPartition;
+			_backupModifiedFiles = backupModifiedFiles;
+			_errorOnRewrite = errorOnRewrite;
+		}
+
+		/// <summary>
+		/// FileHashCache has reported at least one error
+		/// </summary>
+		public bool HasErrors => _hasErrors;
+
+		/// <summary>
+		/// Mount partition based on below parameters. Will filter out files into this partition based on root path and intermediate sub paths
+		/// </summary>
+		/// <param name="descriptor"></param>
+		/// <param name="appName"></param>
+		/// <param name="intermediateEnvironment"></param>
+		/// <param name="logger"></param>
+		public void Mount(TargetDescriptor descriptor, string appName, UnrealIntermediateEnvironment intermediateEnvironment, ILogger logger)
+		{
+			Mount(descriptor.ProjectFile, appName, descriptor.Platform, descriptor.Configuration, descriptor.Architectures, intermediateEnvironment, logger);
+		}
+
+		/// <summary>
+		/// Mount partition based on below parameters. Will filter out files into this partition based on root path and intermediate sub paths
+		/// </summary>
+		public void Mount(FileReference? projectFile, string appName, UnrealTargetPlatform platform, UnrealTargetConfiguration? configuration, UnrealArchitectures? architectures, UnrealIntermediateEnvironment intermediateEnvironment, ILogger logger)
+		{
+			appName = UEBuildTarget.GetTargetIntermediateFolderName(appName, intermediateEnvironment);
+
+			string GetFilterPath(UnrealArch? arch, UnrealTargetConfiguration? configuration)
+			{
+				StringBuilder sb = new();
+				sb.Append(platform).Append(Path.DirectorySeparatorChar);
+				if (arch != null)
+				{
+					sb.Append(UnrealArchitectureConfig.ForPlatform(platform).GetFolderNameForArchitecture(arch.Value)).Append(Path.DirectorySeparatorChar);
+				}
+				sb.Append(appName);
+				if (configuration != null)
+				{
+					sb.Append(Path.DirectorySeparatorChar).Append(configuration);
+				}
+				return sb.ToString();
+			}
+
+			void AddPartitionByArch(UnrealArch? arch, ILogger logger)
+			{
+				AddPartition(Unreal.WritableEngineDirectory, GetFilterPath(arch, configuration), logger);
+				if (configuration == UnrealTargetConfiguration.DebugGame)
+				{
+					AddPartition(Unreal.WritableEngineDirectory, GetFilterPath(arch, UnrealTargetConfiguration.Development), logger);
+				}
+
+				if (projectFile != null)
+				{
+					AddPartition(projectFile.Directory, GetFilterPath(arch, configuration), logger);
+					if (configuration == UnrealTargetConfiguration.DebugGame)
+					{
+						AddPartition(projectFile.Directory, GetFilterPath(arch, UnrealTargetConfiguration.Development), logger);
+					}
+				}
+			}
+
+			if (architectures != null)
+			{
+				foreach (UnrealArch arch in architectures.Architectures)
+				{
+					AddPartitionByArch(arch, logger);
+				}
+			}
+			else
+			{
+				AddPartitionByArch(null, logger);
+			}
+		}
+
+		/// <summary>
+		/// Mount a root directory directly
+		/// </summary>
+		public void Mount(DirectoryReference root, ILogger logger)
+		{
+			AddPartition(root, "", logger);
+		}
+
+		/// <summary>
+		/// Get hash of MemoryStream
+		/// </summary>
+		/// <param name="stream"></param>
+		/// <returns></returns>
+		/// <exception cref="Exception"></exception>
+		public static UInt128 GetFileHash(MemoryStream stream)
+		{
+			var xx = new XxHash128();
+
+			stream.Position = 0;
+			xx.Append(stream);
+			UInt128 newHash = ToUInt128LittleEndian(xx.GetHashAndReset());
+			if (newHash == 0)
+			{
+				throw new Exception("Failed to calculate hash");
+			}
+			return newHash;
+		}
+
+		/// <summary>
+		/// Write file if changed. 
+		/// </summary>
+		public Task? WriteFileIfChanged(FileItem fileItem, string content, ILogger logger, bool allowAsync)
+		{
+			UInt128 newHash = XxHash128.HashToUInt128(MemoryMarshal.AsBytes(content.AsSpan()));
+			return InternalWriteFileIfChanged(fileItem, newHash, content, logger, allowAsync, false);
+		}
+
+		/// <summary>
+		/// Write file if changed. 
+		/// </summary>
+		public Task? WriteFileIfChanged(FileItem fileItem, IEnumerable<string> content, ILogger logger, bool allowAsync)
+		{
+			var xx = new XxHash128();
+			foreach (string str in content)
+			{
+				xx.Append(MemoryMarshal.AsBytes(str.AsSpan()));
+			}
+			UInt128 newHash = ToUInt128LittleEndian(xx.GetHashAndReset());
+			return InternalWriteFileIfChanged(fileItem, newHash, content, logger, allowAsync, false);
+		}
+
+		/// <summary>
+		/// Write file if changed. 
+		/// </summary>
+		public Task? WriteFileIfChanged(FileItem fileItem, MemoryStream data, ILogger logger, bool alwaysUpdateTimestamp = false)
+		{
+			UInt128 newHash = GetFileHash(data);
+			return InternalWriteFileIfChanged(fileItem, newHash, data, logger, false, alwaysUpdateTimestamp);
+		}
+
+		/// <summary>
+		/// Save all partitions
+		/// </summary>
+		public void SaveAll()
+		{
+			foreach ((DirectoryReference root, List<(string filter, Partition partition)>) root in _roots)
+			{
+				foreach ((string filter, Partition partition) in root.Item2)
+				{
+					if (!partition.Modified)
+					{
+						continue;
+					}
+
+					List<KeyValuePair<FileItem, FileEntry>> files = new(partition.Files);
+					files.Sort(static (a, b) => a.Key.Location.CompareTo(b.Key.Location));
+
+					//DirectoryReference.CreateDirectory(partition.Location.Directory);
+					using FileStream stream = File.Open(partition.Location.FullName, FileMode.Create, FileAccess.Write, FileShare.Read);
+					using BinaryWriter writer = new(stream, new UTF8Encoding(false), leaveOpen: false);
+					writer.Write(Partition.CurrentVersion);
+					writer.Write(partition.Files.Count);
+
+					string lastFullName = "";
+
+					foreach (KeyValuePair<FileItem, FileEntry> pair in files)
+					{
+						ReadOnlySpan<char> last = lastFullName;
+						ReadOnlySpan<char> curr = pair.Key.FullName.AsSpan();
+						UInt16 matchLength = (UInt16)MemoryExtensions.CommonPrefixLength(last, curr);
+						writer.Write(matchLength);
+						writer.Write(curr.Slice(matchLength).ToString());
+						lastFullName = pair.Key.FullName;
+
+						writer.Write((ulong)pair.Value.hash);
+						writer.Write((ulong)(pair.Value.hash >> 64));
+						writer.Write(pair.Value.lastWriteTimeUtc);
+
+					}
+				}
+			}
+		}
+
+		/// <summary>
+		/// bytes array to 128 byte key
+		/// </summary>
+		public static UInt128 ToUInt128LittleEndian(ReadOnlySpan<byte> hash16)
+		{
+			ulong lo = BinaryPrimitives.ReadUInt64LittleEndian(hash16[..8]);
+			ulong hi = BinaryPrimitives.ReadUInt64LittleEndian(hash16[8..]);
+			return ((UInt128)hi << 64) | lo;
+		}
+
+		/// <summary>
+		/// Get hash of file if it exists
+		/// </summary>
+		public bool TryGetHash(FileItem fileItem, ILogger logger, out UInt128 hash)
+		{
+			Partition? partition = GetPartition(fileItem.Location, logger);
+			if (partition != null && partition.Files.TryGetValue(fileItem, out FileEntry fileEntry))
+			{
+				hash = fileEntry.hash;
+				return true;
+			}
+			hash = 0;
+			return false;
+		}
+
+		private Task? InternalWriteFileIfChanged(FileItem fileItem, UInt128 newHash, object data, ILogger logger, bool allowAsync, bool alwaysUpdateTimestamp)
+		{
+			Partition? partition = GetPartition(fileItem.Location, logger);
+
+			bool inDatabase = false;
+			bool writtenBefore = false;
+			if (partition != null)
+			{
+				if (partition.Files.TryGetValue(fileItem, out FileEntry fileEntry))
+				{
+					inDatabase = true;
+					writtenBefore = fileEntry.written;
+
+					// If hash is the same we check last write time
+					if (fileEntry.hash == newHash)
+					{
+						if (fileItem.Exists)
+						{
+							// if match we can early out
+							if (fileItem.LastWriteTimeUtc.Ticks == fileEntry.lastWriteTimeUtc)
+							{
+								if (alwaysUpdateTimestamp)
+								{
+									DateTime now = DateTime.UtcNow;
+									if (partition.Files.TryUpdate(fileItem, new() { hash = fileEntry.hash, lastWriteTimeUtc = now.Ticks, written = true }, fileEntry))
+									{
+										partition.Modified = true;
+									}
+									File.SetLastWriteTimeUtc(fileItem.FullName, now);
+									fileItem.ResetCachedInfo();
+								}
+								else if (!writtenBefore)
+								{
+									partition.Files.TryUpdate(fileItem, new() { hash = fileEntry.hash, lastWriteTimeUtc = fileEntry.lastWriteTimeUtc, written = true }, fileEntry);
+								}
+								return null;
+							}
+						}
+					}
+				}
+			}
+			if (allowAsync)
+			{
+				return Task.Run(() => InternalWriteFileIfChanged2(inDatabase, writtenBefore, partition, fileItem, newHash, data, logger, allowAsync, alwaysUpdateTimestamp));
+			}
+
+			InternalWriteFileIfChanged2(inDatabase, writtenBefore, partition, fileItem, newHash, data, logger, allowAsync, alwaysUpdateTimestamp);
+			return null;
+		}
+
+		private void InternalWriteFileIfChanged2(bool inDatabase, bool writtenBefore, Partition? partition, FileItem fileItem, UInt128 newHash, object data, ILogger logger, bool allowAsync, bool alwaysUpdateTimestamp)
+		{
+			if (!inDatabase)
+			{
+				if (fileItem.Exists)
+				{
+					try
+					{
+						bool isEqual = false;
+						if (data is MemoryStream)
+						{
+							const int BufferSize = 1024 * 1024; // 1 MB
+							byte[] buffer = new byte[BufferSize];
+							using var fs = new FileStream(fileItem.FullName, FileMode.Open, FileAccess.Read, FileShare.Read, BufferSize, FileOptions.SequentialScan);
+							int read;
+
+							var xx = new XxHash128();
+							while ((read = fs.Read(buffer, 0, buffer.Length)) > 0)
+								xx.Append(buffer.AsSpan(0, read));
+							isEqual = newHash == ToUInt128LittleEndian(xx.GetHashAndReset());
+						}
+						else if (data is string content)
+						{
+							string oldContent = FileReference.ReadAllText(fileItem.Location, new UTF8Encoding(false));
+							isEqual = String.Equals(oldContent, content, StringComparison.Ordinal);
+						}
+						else if (data is IEnumerable<string> contentLines)
+						{
+							string[] oldLines = FileReference.ReadAllLines(fileItem.Location, new UTF8Encoding(false));
+							int index = 0;
+							isEqual = true;
+							foreach (string str in contentLines)
+							{
+								if (index >= oldLines.Length || str != oldLines[index++])
+								{
+									isEqual = false;
+									break;
+								}
+							}
+						}
+
+						if (isEqual)
+						{
+							DateTime lastWrite = fileItem.LastWriteTimeUtc;
+							if (alwaysUpdateTimestamp)
+							{
+								lastWrite = DateTime.UtcNow;
+								File.SetLastWriteTimeUtc(fileItem.FullName, lastWrite);
+								fileItem.ResetCachedInfo();
+							}
+							if (partition != null)
+							{
+								if (partition.Files.TryAdd(fileItem, new() { hash = newHash, lastWriteTimeUtc = lastWrite.Ticks }))
+								{
+									partition.Modified = true;
+								}
+							}
+							return;
+						}
+					}
+					catch
+					{
+					}
+				}
+			}
+
+			// Same file is written with different content multiple times
+			if (writtenBefore)
+			{
+
+				logger.Log(_errorOnRewrite ? LogLevel.Error : LogLevel.Warning, "Re-writing a file that was previously written with different content: \"{File}\"", fileItem.Location);
+				_hasErrors = _errorOnRewrite;
+			}
+
+			if (!fileItem.Exists)
+			{
+				if (!fileItem.Directory.Exists)
+				{
+					DirectoryReference.CreateDirectory(fileItem.Directory.Location);
+				}
+			}
+			else if (_backupModifiedFiles)
+			{
+				// Make a copy of old file
+				FileReference backupFile = new FileReference(fileItem.FullName + ".old");
+				try
+				{
+					logger.LogDebug("Updating {File}: contents have changed. Saving previous version to {backupFile}.", fileItem.Location, backupFile);
+					FileReference.Delete(backupFile);
+					FileReference.Move(fileItem.Location, backupFile);
+				}
+				catch (Exception Ex)
+				{
+					logger.LogWarning("Unable to rename {fileItem} to {backupFile}", fileItem, backupFile);
+					logger.LogDebug(Ex, "{Ex}", ExceptionUtils.FormatExceptionDetails(Ex));
+				}
+			}
+
+
+			if (data is MemoryStream stream)
+			{
+				stream.Position = 0;
+				using var file = new FileStream(fileItem.FullName, FileMode.Create, FileAccess.Write, FileShare.None);
+				stream.CopyTo(file);
+			}
+			else if (data is string content)
+			{
+				using var fs = new FileStream(fileItem.FullName, FileMode.Create, FileAccess.Write, FileShare.Read, 1 << 20, FileOptions.SequentialScan);
+				using var writer = new StreamWriter(fs, new UTF8Encoding(false), 1 << 20, false);
+				writer.Write(content);
+			}
+			else if (data is IEnumerable<string> contentLines)
+			{
+				string contents = String.Join(Environment.NewLine, contentLines);
+				using var fs = new FileStream(fileItem.FullName, FileMode.Create, FileAccess.Write, FileShare.Read, 1 << 20, FileOptions.SequentialScan);
+				using var writer = new StreamWriter(fs, new UTF8Encoding(false), 1 << 20, false);
+				writer.Write(contents);
+			}
+
+			fileItem.ResetCachedInfo();
+
+			if (partition != null)
+			{
+				FileEntry fileEntry = new() { hash = newHash, lastWriteTimeUtc = fileItem.LastWriteTimeUtc.Ticks, written = true };
+				partition.Files.AddOrUpdate(fileItem, fileEntry, (k,v) => fileEntry);
+				partition.Modified = true;
+			}
+		}
+
+		public void AddPartition(DirectoryReference rootPath, string filterPath, ILogger logger)
+		{
+			Partition? partition = null;
+			lock (_roots)
+			{
+				var root = _roots.FirstOrDefault(e => e.root == rootPath);
+				if (root.root == null)
+				{
+					root = new(rootPath, new());
+					_roots.Add(root);
+				}
+				var filters = root.Item2;
+				var filter = filters.FirstOrDefault(f => f.filter == filterPath);
+				partition = filter.partition;
+				if (partition == null)
+				{
+					FileReference location;
+					if (filterPath != string.Empty)
+					{
+						location = FileReference.Combine(rootPath, "Intermediate", "Build", filterPath, "FileHashCache.bin");
+					}
+					else
+					{
+						location = FileReference.Combine(rootPath, "FileHashCache.bin");
+					}
+					partition = new Partition(location);
+					filters.Add((filterPath, partition));
+				}
+			}
+
+			lock (partition)
+			{
+				if (!partition.Loaded)
+				{
+					partition.Loaded = true;
+					if (FileReference.Exists(partition.Location))
+					{
+						ReadPartition(partition, logger);
+					}
+				}
+			}
+
+		}
+
+		private readonly string Intermediate = $"{Path.DirectorySeparatorChar}Intermediate";
+		private readonly string External = $"{Path.DirectorySeparatorChar}External";
+		private readonly string Build = $"{Path.DirectorySeparatorChar}Build{Path.DirectorySeparatorChar}";
+
+		private Partition? GetPartition(FileReference fileRef, ILogger logger)
+		{
+			ReadOnlySpan<char> fullName = fileRef.FullName;
+			int index = fullName.IndexOf(Intermediate, FileReference.Comparison);
+			if (index == -1)
+			{
+				if (_warnOnNoPartition)
+				{
+					logger.LogWarning($"{fullName} is not under intermediate folder");
+				}
+				return null;
+			}
+
+			ReadOnlySpan<char> filterPath = fullName.Slice(index + Intermediate.Length);
+
+			if (filterPath.StartsWith(External))
+			{
+				filterPath = filterPath.Slice(External.Length);
+			}
+
+			if (filterPath.StartsWith(Build))
+			{
+				filterPath = filterPath.Slice(Build.Length);
+			}
+
+			Partition? partition = null;
+
+			foreach (var root in _roots)
+			{
+				if (!fileRef.IsUnderDirectory(root.root))
+				{
+					continue;
+				}
+
+				foreach (var entry in root.Item2)
+				{
+					if (filterPath.StartsWith(entry.filter, FileReference.Comparison))
+					{
+						partition = entry.partition;
+						break;
+					}
+				}
+				break;
+
+			}
+
+			if (partition == null)
+			{
+				if (_warnOnNoPartition)
+				{
+					logger.LogWarning($"There is no partition covering path to {fileRef}!");
+				}
+				return null;
+			}
+
+			return partition;
+		}
+		private void ReadPartition(Partition partition, ILogger logger)
+		{
+			try
+			{
+				using FileStream stream = File.Open(partition.Location.FullName, FileMode.Open, FileAccess.Read, FileShare.Read);
+				using BinaryReader reader = new(stream, new UTF8Encoding(false), leaveOpen: false);
+
+				int version = reader.ReadInt32();
+				if (version != Partition.CurrentVersion)
+				{
+					logger.LogDebug($"Unable to read dependency cache from {partition.Location}; version {version} vs current {Partition.CurrentVersion}");
+					return;
+				}
+
+				int count = reader.ReadInt32();
+				partition.Files = new(-1, count);
+
+				string lastFilename = "";
+				StringBuilder sb = new(256);
+
+				for (int Idx = 0; Idx < count; Idx++)
+				{
+					UInt16 matchLength = reader.ReadUInt16();
+					string str = reader.ReadString();
+					sb.Clear().Append(lastFilename, 0, matchLength).Append(str);
+
+					ulong hashA = reader.ReadUInt64();
+					UInt128 hashB = reader.ReadUInt64();
+
+					UInt128 hash = (hashB << 64) + hashA;
+					long lastWriteTimeUtc = reader.ReadInt64();
+
+					string path = sb.ToString();
+					lastFilename = path;
+					partition.Files.TryAdd(FileItem.GetItemByPath(path), new() { hash = hash, lastWriteTimeUtc = lastWriteTimeUtc });
+				}
+			}
+			catch (Exception Ex)
+			{
+				logger.LogWarning("Unable to read {Location}. See log for additional information.", partition.Location);
+				logger.LogDebug("{Ex}", ExceptionUtils.FormatExceptionDetails(Ex));
+			}
+		}
+
+		private struct FileEntry
+		{
+			internal UInt128 hash;
+			internal long lastWriteTimeUtc;
+			internal bool written;
+		}
+
+		private class Partition
+		{
+			internal Partition(FileReference location) { Location = location; }
+			internal readonly FileReference Location;
+			internal bool Loaded;
+			internal bool Modified;
+			internal const int CurrentVersion = 1;
+
+			internal ConcurrentDictionary<FileItem, FileEntry> Files = [];
+		}
+
+		private List<(DirectoryReference root, List<(string filter, Partition partition)>)> _roots = [];
+		private bool _warnOnNoPartition;
+		private bool _backupModifiedFiles;
+		private bool _errorOnRewrite;
+		private bool _hasErrors;
+	}
+}
\ No newline at end of file
diff --git a/Engine/Source/Programs/UnrealBuildTool/System/TargetMakefile.cs b/Engine/Source/Programs/UnrealBuildTool/System/TargetMakefile.cs
--- a/Engine/Source/Programs/UnrealBuildTool/System/TargetMakefile.cs
+++ b/Engine/Source/Programs/UnrealBuildTool/System/TargetMakefile.cs
@@ -69,11 +69,14 @@ namespace UnrealBuildTool
 		private readonly ILogger Logger;
 		public TargetMakefile Makefile { get; }
 
+		public FileHashCache FileHashCache { get; }
+
 		private ConcurrentBag<Task> WriteTasks = new();
 
-		public TargetMakefileBuilder(TargetMakefile InMakefile, ILogger InLogger)
+		public TargetMakefileBuilder(TargetMakefile InMakefile, FileHashCache fileHashCache, ILogger InLogger)
 		{
 			Logger = InLogger;
+			FileHashCache = fileHashCache;
 			Makefile = InMakefile;
 		}
 
@@ -94,16 +97,10 @@ namespace UnrealBuildTool
 				Makefile.InternalDependencies.Add(FileItem);
 			}
 
-			if (AllowAsync)
-			{
-				WriteTasks.Add(Task.Run(() =>
-				{
-					Utils.WriteFileIfChanged(FileItem.Location, Contents, Logger);
-				}));
-			}
-			else
+			Task? task = FileHashCache.WriteFileIfChanged(FileItem, Contents, Logger, AllowAsync);
+			if (task != null)
 			{
-				Utils.WriteFileIfChanged(FileItem.Location, Contents, Logger);
+				WriteTasks.Add(task);
 			}
 		}
 
@@ -115,16 +112,10 @@ namespace UnrealBuildTool
 				Makefile.InternalDependencies.Add(FileItem);
 			}
 
-			if (AllowAsync)
-			{
-				WriteTasks.Add(Task.Run(() =>
-				{
-					Utils.WriteFileIfChanged(FileItem, ContentLines, Logger);
-				}));
-			}
-			else
+			Task? task = FileHashCache.WriteFileIfChanged(FileItem, ContentLines, Logger, AllowAsync);
+			if (task != null)
 			{
-				Utils.WriteFileIfChanged(FileItem.Location, ContentLines, Logger);
+				WriteTasks.Add(task);
 			}
 		}
 
diff --git a/Engine/Source/Programs/UnrealBuildTool/System/Utils.cs b/Engine/Source/Programs/UnrealBuildTool/System/Utils.cs
--- a/Engine/Source/Programs/UnrealBuildTool/System/Utils.cs
+++ b/Engine/Source/Programs/UnrealBuildTool/System/Utils.cs
@@ -1452,6 +1452,8 @@ namespace UnrealBuildTool
 		{
 			foreach (FileReference ScriptFile in ScriptFiles)
 			{
+				using var _ = Timeline.ScopeEvent(ScriptFile.GetFileName());
+
 				ProcessStartInfo StartInfo = new ProcessStartInfo();
 				StartInfo.FileName = BuildHostPlatform.Current.Shell.FullName;
 
@@ -1473,6 +1475,8 @@ namespace UnrealBuildTool
 
 			if (ScriptFiles.Length > 0)
 			{
+				using var _ = Timeline.ScopeEvent("ResetCachedInfo");
+
 				// We have to invalidate all cached file info after running the scripts, because we don't know what may have changed.
 				DirectoryItem.ResetAllCachedInfo_SLOW();
 			}
