[UBT]
* Fixed library path in mac rsp file not being virtualized

diff --git a/Engine/Source/Programs/UnrealBuildTool/Platform/Mac/MacToolChain.cs b/Engine/Source/Programs/UnrealBuildTool/Platform/Mac/MacToolChain.cs
--- a/Engine/Source/Programs/UnrealBuildTool/Platform/Mac/MacToolChain.cs
+++ b/Engine/Source/Programs/UnrealBuildTool/Platform/Mac/MacToolChain.cs
@@ -757,7 +757,7 @@ namespace UnrealBuildTool
 
 					HashSet<DirectoryReference> UsedDestDirs = new();
 					HashSet<FileReference> UsedSourceLibs = new();
-					var ProcessLib = (FileReference SourceLib, FileReference DestLib, FileReference ReferencingFile) => 
+					void ProcessLib(FileReference SourceLib, FileReference DestLib, FileReference ReferencingFile)
 					{
 						if (string.Equals(SourceLib.GetExtension(), ".dylib", StringComparison.InvariantCultureIgnoreCase))
 						{
@@ -826,7 +826,7 @@ namespace UnrealBuildTool
 					}
 					else if (AdditionalLibrary.Contains(".framework/"))
 					{
-						LinkArguments.Add(String.Format("\"{0}\"", AdditionalLibrary));
+						LinkArguments.Add($"\"{NormalizeCommandLinePath(FileReference.FromString(AdditionalLibrary), LinkEnvironment.RootPaths)}\"");
 					}
 					else if (String.IsNullOrEmpty(Path.GetDirectoryName(AdditionalLibrary)) && String.IsNullOrEmpty(Path.GetExtension(AdditionalLibrary)))
 					{
