UnrealBuildTool: Readd modified code I skipped to simplify a merge resolve

#rnx

diff --git a/Engine/Source/Programs/UnrealBuildTool/Configuration/UEBuildModuleCPP.cs b/Engine/Source/Programs/UnrealBuildTool/Configuration/UEBuildModuleCPP.cs
--- a/Engine/Source/Programs/UnrealBuildTool/Configuration/UEBuildModuleCPP.cs
+++ b/Engine/Source/Programs/UnrealBuildTool/Configuration/UEBuildModuleCPP.cs
@@ -130,14 +130,24 @@ namespace UnrealBuildTool
 
 			if (Rules.bPrecompile)
 			{
-				if (GeneratedCodeDirectory != null && DirectoryReference.Exists(GeneratedCodeDirectory))
+				DirectoryItem? GeneratedCodeDirectoryItem = GeneratedCodeDirectory != null ? DirectoryItem.GetItemByDirectoryReference(GeneratedCodeDirectory) : null;
+				if (GeneratedCodeDirectoryItem?.Exists == true)
 				{
-					foreach (FileReference GeneratedCodeFile in DirectoryReference.EnumerateFiles(GeneratedCodeDirectory, "*", SearchOption.AllDirectories))
+					Queue<DirectoryItem> DirQueue = new();
+					DirQueue.Enqueue(GeneratedCodeDirectoryItem);
+					while (DirQueue.TryDequeue(out DirectoryItem? dir))
 					{
-						// Exclude timestamp files, since they're always updated and cause collisions between builds
-						if (!GeneratedCodeFile.GetFileName().Equals("Timestamp", StringComparison.OrdinalIgnoreCase) && !GeneratedCodeFile.HasExtension(".cpp"))
+						foreach (FileItem GeneratedCodeFile in dir.EnumerateFiles())
 						{
-							BuildProducts.Add(GeneratedCodeFile, BuildProductType.BuildResource);
+							// Exclude timestamp files, since they're always updated and cause collisions between builds
+							if (!GeneratedCodeFile.Name.Equals("Timestamp", StringComparison.OrdinalIgnoreCase) && !GeneratedCodeFile.HasExtension(".cpp"))
+							{
+								BuildProducts.Add(GeneratedCodeFile.Location, BuildProductType.BuildResource);
+							}
+						}
+						foreach (DirectoryItem subdir in dir.EnumerateDirectories())
+						{
+							DirQueue.Enqueue(subdir);
 						}
 					}
 				}
