[UBT]
* Fixed so makefiles are not preloaded when bParallelMakefileGeneration is disabled since otherwise it would run in parallel
* Fixed so GitSourceFileWorkingSet can handle parallel calls of Contains method

diff --git a/Engine/Source/Programs/UnrealBuildTool/Modes/BuildMode.cs b/Engine/Source/Programs/UnrealBuildTool/Modes/BuildMode.cs
--- a/Engine/Source/Programs/UnrealBuildTool/Modes/BuildMode.cs
+++ b/Engine/Source/Programs/UnrealBuildTool/Modes/BuildMode.cs
@@ -1137,7 +1137,7 @@ namespace UnrealBuildTool
 				// Note, there are (very) rare instances of this being the wrong makefile
 				// and that is when target rules change name of target. so there is a tiny chance we will waste some time loading a makefile that is not being used
 				Dictionary<FileReference, Task<(TargetMakefile?, bool)>> preLoadedMakefiles = [];
-				if (buildConfiguration.bUseUBTMakefiles)
+				if (buildConfiguration.bUseUBTMakefiles && buildConfiguration.bParallelMakefileGeneration) // Only preload if parallel is enabled to prevent workingSet from being called from multiple threads
 				{
 					for (int index = 0; index != descriptorCount; ++index)
 					{
diff --git a/Engine/Source/Programs/UnrealBuildTool/System/SourceFileWorkingSet.cs b/Engine/Source/Programs/UnrealBuildTool/System/SourceFileWorkingSet.cs
--- a/Engine/Source/Programs/UnrealBuildTool/System/SourceFileWorkingSet.cs
+++ b/Engine/Source/Programs/UnrealBuildTool/System/SourceFileWorkingSet.cs
@@ -90,13 +90,14 @@ namespace UnrealBuildTool
 	/// </summary>
 	class GitSourceFileWorkingSet : ISourceFileWorkingSet
 	{
-		DirectoryReference RootDir;
+		readonly DirectoryReference RootDir;
 		Process? BackgroundProcess;
-		HashSet<FileReference> Files;
-		List<DirectoryReference> Directories;
-		List<string> ErrorOutput;
-		GitSourceFileWorkingSet? Inner;
-		ILogger Logger;
+		readonly object BackgroundProcessDisposeLock = new();
+		readonly Lazy<bool> BackgroundProcessDone;
+		readonly HashSet<FileReference> Files;
+		readonly List<DirectoryReference> Directories;
+		readonly List<string> ErrorOutput;
+		readonly GitSourceFileWorkingSet? Inner;
 
 		/// <summary>
 		/// Constructor
@@ -112,7 +113,6 @@ namespace UnrealBuildTool
 			Directories = [];
 			ErrorOutput = [];
 			this.Inner = Inner;
-			this.Logger = Logger;
 
 			Logger.LogInformation("Using 'git status' to determine working set for adaptive non-unity build ({RootDir}).", RootDir);
 
@@ -141,63 +141,58 @@ namespace UnrealBuildTool
 				BackgroundProcess.Dispose();
 				BackgroundProcess = null;
 			}
-		}
 
-		/// <summary>
-		/// Terminates the background process.
-		/// </summary>
-		private void TerminateBackgroundProcess()
-		{
-			if (BackgroundProcess != null)
+			BackgroundProcessDone = new Lazy<bool>(() =>
 			{
-				if (BackgroundProcess?.HasExited == false)
+				if (BackgroundProcess == null)
 				{
+					return true;
+				}
+
+				if (BackgroundProcess.WaitForExit(500) == false)
+				{
+					Logger.LogInformation("Waiting for 'git status' command to complete");
+				}
+				if (BackgroundProcess.WaitForExit(15000) == false)
+				{
+					Logger.LogInformation("Terminating git child process due to timeout");
 					try
 					{
-						BackgroundProcess?.Kill();
+						BackgroundProcess.Kill();
 					}
 					catch
 					{
 					}
 				}
-				WaitForBackgroundProcess();
-			}
+				BackgroundProcess.WaitForExit();
+				lock (BackgroundProcessDisposeLock)
+				{
+					BackgroundProcess.Dispose();
+					BackgroundProcess = null;
+				}
+				return true;
+			});
 		}
 
 		/// <summary>
-		/// Waits for the background to terminate.
+		/// Dispose of this object
 		/// </summary>
-		private void WaitForBackgroundProcess()
+		public void Dispose()
 		{
-			if (BackgroundProcess != null)
+			lock (BackgroundProcessDisposeLock)
 			{
-				if (BackgroundProcess?.WaitForExit(500) == false)
-				{
-					Logger.LogInformation("Waiting for 'git status' command to complete");
-				}
-				if (BackgroundProcess?.WaitForExit(15000) == false)
+				if (BackgroundProcess != null && !BackgroundProcess.HasExited)
 				{
-					Logger.LogInformation("Terminating git child process due to timeout");
 					try
 					{
-						BackgroundProcess?.Kill();
+						BackgroundProcess.Kill();
 					}
 					catch
 					{
 					}
 				}
-				BackgroundProcess?.WaitForExit();
-				BackgroundProcess?.Dispose();
-				BackgroundProcess = null;
 			}
-		}
-
-		/// <summary>
-		/// Dispose of this object
-		/// </summary>
-		public void Dispose()
-		{
-			TerminateBackgroundProcess();
+			bool _ = BackgroundProcessDone.Value;
 
 			Inner?.Dispose();
 		}
@@ -215,7 +210,8 @@ namespace UnrealBuildTool
 				return false;
 			}
 
-			WaitForBackgroundProcess();
+			bool _ = BackgroundProcessDone.Value;
+
 			if (Files.Contains(File.Location) || Directories.Any(x => File.Location.IsUnderDirectory(x)))
 			{
 				return true;
