[UBT]
* Fixed so filehashcache also handles "arm64+x64" architecture folder for macos

diff --git a/Engine/Source/Programs/UnrealBuildTool/System/FileHashCache.cs b/Engine/Source/Programs/UnrealBuildTool/System/FileHashCache.cs
--- a/Engine/Source/Programs/UnrealBuildTool/System/FileHashCache.cs
+++ b/Engine/Source/Programs/UnrealBuildTool/System/FileHashCache.cs
@@ -59,13 +59,13 @@ namespace UnrealBuildTool
 		{
 			appName = UEBuildTarget.GetTargetIntermediateFolderName(appName, intermediateEnvironment);
 
-			string GetFilterPath(UnrealArch? arch, UnrealTargetConfiguration? configuration)
+			string GetFilterPath(string? archFolder, UnrealTargetConfiguration? configuration)
 			{
 				StringBuilder sb = new();
 				sb.Append(platform).Append(Path.DirectorySeparatorChar);
-				if (arch != null)
+				if (archFolder != null)
 				{
-					sb.Append(UnrealArchitectureConfig.ForPlatform(platform).GetFolderNameForArchitecture(arch.Value)).Append(Path.DirectorySeparatorChar);
+					sb.Append(archFolder).Append(Path.DirectorySeparatorChar);
 				}
 				sb.Append(appName);
 				if (configuration != null)
@@ -75,7 +75,7 @@ namespace UnrealBuildTool
 				return sb.ToString();
 			}
 
-			void AddPartitionByArch(UnrealArch? arch, ILogger logger)
+			void AddPartitionByArch(string? arch, ILogger logger)
 			{
 				AddPartition(Unreal.WritableEngineDirectory, GetFilterPath(arch, configuration), logger);
 				if (configuration == UnrealTargetConfiguration.DebugGame)
@@ -97,7 +97,13 @@ namespace UnrealBuildTool
 			{
 				foreach (UnrealArch arch in architectures.Architectures)
 				{
-					AddPartitionByArch(arch, logger);
+					AddPartitionByArch(UnrealArchitectureConfig.ForPlatform(platform).GetFolderNameForArchitecture(arch), logger);
+				}
+
+				// If macos we also need to cover "arm64+x64" architecture
+				if (OperatingSystem.IsMacOS() && architectures.Architectures.Count > 1)
+				{
+					AddPartitionByArch(architectures.ToString(), logger);
 				}
 			}
 			else
