[UBT]
* Fixed stubs file location when building target that is a lib

diff --git a/Engine/Source/Programs/UnrealBuildTool/Platform/Mac/MacToolChain.cs b/Engine/Source/Programs/UnrealBuildTool/Platform/Mac/MacToolChain.cs
--- a/Engine/Source/Programs/UnrealBuildTool/Platform/Mac/MacToolChain.cs
+++ b/Engine/Source/Programs/UnrealBuildTool/Platform/Mac/MacToolChain.cs
@@ -653,7 +653,8 @@ namespace UnrealBuildTool
 			FileReference DylibPathRef = new FileReference(DylibPath);
 			bool bIsEngineDylib = DylibPathRef.IsUnderDirectory(Unreal.EngineDirectory);
 
-			string DylibFilename = Path.GetFileName(DylibPath);
+			string IntermediateDirectory = (ProjectFile == null || LinkEnvironment.OutputFilePaths.All(x => x.IsUnderDirectory(Unreal.EngineDirectory)) || bIsEngineDylib ? Unreal.EngineDirectory : ProjectFile.Directory) + "/Intermediate/Build/Mac";
+			string DylibFilename = Path.GetFileNameWithoutExtension(DylibPath);
 
 			string AppName = string.Empty;
 			int firstDash = DylibFilename.IndexOf('-');
@@ -661,7 +662,10 @@ namespace UnrealBuildTool
 			{
 				AppName = DylibFilename.Substring(0, firstDash);
 			}
-			string IntermediateDirectory = (ProjectFile == null || LinkEnvironment.OutputFilePaths.All(x => x.IsUnderDirectory(Unreal.EngineDirectory)) || bIsEngineDylib ? Unreal.EngineDirectory : ProjectFile.Directory) + "/Intermediate/Build/Mac";
+			else if (DylibFilename.StartsWith("lib")) // This is expected to exist. if there is no dash involved this is the target lib itself and name is part of the name
+			{
+				AppName = DylibFilename.Substring(3);
+			}
 			return FileItem.GetItemByPath(Path.Combine(IntermediateDirectory, LinkEnvironment.Architecture.ToString(), AppName, LinkEnvironment.Configuration.ToString(), "Stubs", DylibFilename));
 		}
 
