[UBT]
* Fixed bug where stub filenames was missing .dylib caused by previous submit to MakeStubItem

diff --git a/Engine/Source/Programs/UnrealBuildTool/Platform/Mac/MacToolChain.cs b/Engine/Source/Programs/UnrealBuildTool/Platform/Mac/MacToolChain.cs
--- a/Engine/Source/Programs/UnrealBuildTool/Platform/Mac/MacToolChain.cs
+++ b/Engine/Source/Programs/UnrealBuildTool/Platform/Mac/MacToolChain.cs
@@ -654,7 +654,7 @@ namespace UnrealBuildTool
 			bool bIsEngineDylib = DylibPathRef.IsUnderDirectory(Unreal.EngineDirectory);
 
 			string IntermediateDirectory = (ProjectFile == null || LinkEnvironment.OutputFilePaths.All(x => x.IsUnderDirectory(Unreal.EngineDirectory)) || bIsEngineDylib ? Unreal.EngineDirectory : ProjectFile.Directory) + "/Intermediate/Build/Mac";
-			string DylibFilename = Path.GetFileNameWithoutExtension(DylibPath);
+			string DylibFilename = Path.GetFileName(DylibPath);
 
 			string AppName = string.Empty;
 			int firstDash = DylibFilename.IndexOf('-');
@@ -664,7 +664,7 @@ namespace UnrealBuildTool
 			}
 			else if (DylibFilename.StartsWith("lib")) // This is expected to exist. if there is no dash involved this is the target lib itself and name is part of the name
 			{
-				AppName = DylibFilename.Substring(3);
+				AppName = Path.GetFileNameWithoutExtension(DylibFilename).Substring(3);
 			}
 			return FileItem.GetItemByPath(Path.Combine(IntermediateDirectory, LinkEnvironment.Architecture.ToString(), AppName, LinkEnvironment.Configuration.ToString(), "Stubs", DylibFilename));
 		}
