Add some safety checks for merged modules around what may or may not exist in the config

#rnx

diff --git a/Engine/Source/Programs/UnrealBuildTool/Configuration/UEBuildTarget.cs b/Engine/Source/Programs/UnrealBuildTool/Configuration/UEBuildTarget.cs
--- a/Engine/Source/Programs/UnrealBuildTool/Configuration/UEBuildTarget.cs
+++ b/Engine/Source/Programs/UnrealBuildTool/Configuration/UEBuildTarget.cs
@@ -4020,7 +4020,7 @@ namespace UnrealBuildTool
 			UEBuildModuleCPP launchModule = executableBinary.PrimaryModule;
 			HashSet<UEBuildModuleCPP> allModules = new(Binaries.SelectMany(x => x.Modules.OfType<UEBuildModuleCPP>().Where(x => x != launchModule)));
 
-			if (Rules.MergePlugins.Any() && BuildPlugins != null)
+			if ((Rules.MergePlugins.Any() || Rules.MergePluginsShared.Any()) && BuildPlugins != null)
 			{
 				// Dependency rules:
 				// - Regular (non-shared) merged libraries cannot contain plugins that are dependencies of other merged libraries.
@@ -4050,7 +4050,7 @@ namespace UnrealBuildTool
 				{
 					foreach (string name in sharedGroupRule.Value)
 					{
-						if (!Rules.MergePlugins.Contains(name))
+						if (BuildPlugins.Any(x => x.Name == name) && !Rules.MergePlugins.Contains(name))
 						{
 							rootPluginsToMerge.Add(name);
 							rootPluginRemappings.Add(name, sharedGroupRule.Key);
