Optimize merge modules usage reporting as it was extraordinarily slow

#rnx

diff --git a/Engine/Source/Programs/UnrealBuildTool/Configuration/UEBuildTarget.cs b/Engine/Source/Programs/UnrealBuildTool/Configuration/UEBuildTarget.cs
--- a/Engine/Source/Programs/UnrealBuildTool/Configuration/UEBuildTarget.cs
+++ b/Engine/Source/Programs/UnrealBuildTool/Configuration/UEBuildTarget.cs
@@ -4243,7 +4243,16 @@ namespace UnrealBuildTool
 								{
 									foreach (UEBuildPlugin user in globalPluginUserMap[plugin])
 									{
-										FindRootUsers(user, ref roots, newPath);
+										Dictionary<string, List<string>> cachedRoots = new();
+
+										if (rootPluginUsageMap.TryGetValue(user.Name, out cachedRoots))
+										{
+											roots = cachedRoots;
+										}
+										else
+										{
+											FindRootUsers(user, ref roots, newPath);
+										}
 									}
 								}
 								else
