[Backout] - CL46582595
#rnx
Original CL Desc
-----------------------------------------------------------------
Optimize merge modules usage reporting as it was extraordinarily slow

#rnx

diff --git a/Engine/Source/Programs/UnrealBuildTool/Configuration/UEBuildTarget.cs b/Engine/Source/Programs/UnrealBuildTool/Configuration/UEBuildTarget.cs
--- a/Engine/Source/Programs/UnrealBuildTool/Configuration/UEBuildTarget.cs
+++ b/Engine/Source/Programs/UnrealBuildTool/Configuration/UEBuildTarget.cs
@@ -4243,16 +4243,7 @@ namespace UnrealBuildTool
 								{
 									foreach (UEBuildPlugin user in globalPluginUserMap[plugin])
 									{
-										Dictionary<string, List<string>> cachedRoots = new();
-
-										if (rootPluginUsageMap.TryGetValue(user.Name, out cachedRoots))
-										{
-											roots = cachedRoots;
-										}
-										else
-										{
-											FindRootUsers(user, ref roots, newPath);
-										}
+										FindRootUsers(user, ref roots, newPath);
 									}
 								}
 								else
