Optimize merge modules usage reporting as it was extraordinarily slow

#rnx

diff --git a/Engine/Source/Programs/UnrealBuildTool/Configuration/UEBuildTarget.cs b/Engine/Source/Programs/UnrealBuildTool/Configuration/UEBuildTarget.cs
--- a/Engine/Source/Programs/UnrealBuildTool/Configuration/UEBuildTarget.cs
+++ b/Engine/Source/Programs/UnrealBuildTool/Configuration/UEBuildTarget.cs
@@ -4234,16 +4234,20 @@ namespace UnrealBuildTool
 							List<string> userNames = new();
 							UEBuildPlugin plugin = pluginCandidates.First();
 
-							void FindRootUsers(UEBuildPlugin plugin, ref Dictionary<string, List<string>> roots, List<string> path)
+							void FindRootUsers(UEBuildPlugin plugin, ref Dictionary<string, List<string>> roots, List<string> path, HashSet<string> seen)
 							{
 								List<string> newPath = new(path);
 								newPath.Add(plugin.Name);
+								seen.Add(plugin.Name);
 
 								if (globalPluginUserMap.ContainsKey(plugin))
 								{
 									foreach (UEBuildPlugin user in globalPluginUserMap[plugin])
 									{
-										FindRootUsers(user, ref roots, newPath);
+										if (!seen.Contains(user.Name))
+										{
+											FindRootUsers(user, ref roots, newPath, seen);
+										}
 									}
 								}
 								else
@@ -4260,7 +4264,7 @@ namespace UnrealBuildTool
 							}
 
 							Dictionary<string, List<string>> roots = new();
-							FindRootUsers(plugin, ref roots, new());
+							FindRootUsers(plugin, ref roots, new(), new());
 							rootPluginUsageMap.Add(module.Name, roots);
 						}
 					}
