Tweak modules merging rules to always follow every dependency for explicitly-excluded GFPs

We can't follow every single dependency when processing unmergeable because unreferenced roots would tend to exclude a lot of mergeable plugins.
This allows setting exclusion roots where a GFP is set up explicitly for the job of excluding other plugins.
Plus we just don't generally want an excluded GFP to have mergeable dependencies!

Also cleans up parameter names for GetPluginDependencies to clarify what they really do.

#rnx

diff --git a/Engine/Source/Programs/UnrealBuildTool/Configuration/UEBuildTarget.cs b/Engine/Source/Programs/UnrealBuildTool/Configuration/UEBuildTarget.cs
--- a/Engine/Source/Programs/UnrealBuildTool/Configuration/UEBuildTarget.cs
+++ b/Engine/Source/Programs/UnrealBuildTool/Configuration/UEBuildTarget.cs
@@ -3951,11 +3951,11 @@ namespace UnrealBuildTool
 
 			Dictionary<UEBuildPlugin, HashSet<UEBuildPlugin>> globalPluginUserMap = new();
 
-			void GetPluginDependencies(UEBuildPlugin plugin, ref HashSet<UEBuildPlugin> dependencies, bool strictFollowMode = false, bool hasCodeInHierarchy = false)
+			void GetPluginDependencies(UEBuildPlugin plugin, ref HashSet<UEBuildPlugin> dependencies, bool bShouldOnlyFollowCode = false, bool bHasUsersWithSourceCode = false)
 			{
 				if (plugin.Modules.Any())
 				{
-					hasCodeInHierarchy = true;
+					bHasUsersWithSourceCode = true;
 				}
 
 				foreach (PluginReferenceDescriptor dependency in plugin.Info.Descriptor.Plugins ?? new())
@@ -3967,9 +3967,9 @@ namespace UnrealBuildTool
 						{
 							UEBuildPlugin dependentPlugin = dependentPluginCandidates.First();
 
-							bool shouldFollowPlugin = (CanPluginBeMerged(dependentPlugin) && !strictFollowMode) || hasCodeInHierarchy;
+							bool bShouldFollowPlugin = (CanPluginBeMerged(dependentPlugin) && !bShouldOnlyFollowCode) || bHasUsersWithSourceCode;
 
-							if (!dependencies.Contains(dependentPlugin) && shouldFollowPlugin)
+							if (!dependencies.Contains(dependentPlugin) && bShouldFollowPlugin)
 							{
 								dependencies.Add(dependentPlugin);
 
@@ -3979,7 +3979,7 @@ namespace UnrealBuildTool
 								}
 								globalPluginUserMap[dependentPlugin].Add(plugin);
 
-								GetPluginDependencies(dependentPlugin, ref dependencies, strictFollowMode, hasCodeInHierarchy);
+								GetPluginDependencies(dependentPlugin, ref dependencies, bShouldOnlyFollowCode, bHasUsersWithSourceCode);
 							}
 						}
 					}
@@ -4169,7 +4169,8 @@ namespace UnrealBuildTool
 				foreach (UEBuildPlugin plugin in unmergablePlugins)
 				{
 					HashSet<UEBuildPlugin> dependencies = new();
-					GetPluginDependencies(plugin, ref dependencies, true);
+					bool bShouldFollowEveryDependency = Rules.MergePluginsExcluded.Contains(plugin.Name);
+					GetPluginDependencies(plugin, ref dependencies, !bShouldFollowEveryDependency);
 
 					foreach (UEBuildPlugin dependency in dependencies)
 					{
