UnrealBuildTool: Split editor and program target modules into more binaries when merged modules is enabled

#rnx

diff --git a/Engine/Source/Programs/UnrealBuildTool/Configuration/UEBuildTarget.cs b/Engine/Source/Programs/UnrealBuildTool/Configuration/UEBuildTarget.cs
--- a/Engine/Source/Programs/UnrealBuildTool/Configuration/UEBuildTarget.cs
+++ b/Engine/Source/Programs/UnrealBuildTool/Configuration/UEBuildTarget.cs
@@ -4331,10 +4331,55 @@ namespace UnrealBuildTool
 			}
 
 			// Merge remaining engine modules
-			ConsumeModules("Engine", engineModules);
+			IEnumerable<UEBuildModuleCPP> mergedEngineModules = engineModules;
+			if (Rules.Type == TargetType.Editor || Rules.Type == TargetType.Program)
+			{
+				IEnumerable<UEBuildModuleCPP> modules;
+				foreach (DirectoryReference reference in DirectoryReference.EnumerateDirectories(Unreal.EngineSourceDirectory))
+				{
+					modules = mergedEngineModules.Where(x => x.Rules.Directory.IsUnderDirectory(reference));
+					ConsumeModules($"Engine{reference.GetDirectoryName()}", modules);
+					mergedEngineModules = [.. mergedEngineModules.Except(modules)];
+
+					DirectoryReference pluginReference = DirectoryReference.Combine(Unreal.EngineDirectory, "Plugins", reference.GetDirectoryName());
+					modules = mergedEngineModules.Where(x => x.Rules.Directory.IsUnderDirectory(pluginReference));
+					ConsumeModules($"EnginePlugin{pluginReference.GetDirectoryName()}", modules);
+					mergedEngineModules = [.. mergedEngineModules.Except(modules)];
+				}
+
+				DirectoryReference experimentalPluginReference = DirectoryReference.Combine(Unreal.EngineDirectory, "Plugins", "Experimental");
+				modules = mergedEngineModules.Where(x => x.Rules.Directory.IsUnderDirectory(experimentalPluginReference));
+				ConsumeModules($"EnginePlugin{experimentalPluginReference.GetDirectoryName()}", modules);
+				mergedEngineModules = [.. mergedEngineModules.Except(modules)];
+
+				foreach (DirectoryReference reference in DirectoryReference.EnumerateDirectories(DirectoryReference.Combine(Unreal.EngineDirectory, "Platforms")))
+				{
+					modules = mergedEngineModules.Where(x => x.Rules.Directory.IsUnderDirectory(reference));
+					ConsumeModules($"Platform{reference.GetDirectoryName()}", modules);
+					mergedEngineModules = [.. mergedEngineModules.Except(modules)];
+				}
+
+				foreach (DirectoryReference reference in DirectoryReference.EnumerateDirectories(Unreal.EngineDirectory))
+				{
+					modules = mergedEngineModules.Where(x => x.Rules.Directory.IsUnderDirectory(reference));
+					ConsumeModules($"Engine{reference.GetDirectoryName()}", modules);
+					mergedEngineModules = [.. mergedEngineModules.Except(modules)];
+				}
+
+				modules = mergedEngineModules.Where(x => x.Rules.Plugin != null);
+				ConsumeModules("EnginePlugin", modules);
+				mergedEngineModules = [.. mergedEngineModules.Except(modules)];
+			}
+			ConsumeModules("Engine", mergedEngineModules);
 
 			// Merge remaining modules to common set
 			IEnumerable<UEBuildModuleCPP> commonModules = allModules.Except(engineModules);
+			if (Rules.Type == TargetType.Editor || Rules.Type == TargetType.Program)
+			{
+				IEnumerable<UEBuildModuleCPP> modules = commonModules.Where(x => x.Rules.Plugin != null);
+				ConsumeModules("CommonPlugin", modules);
+				commonModules = commonModules.Except(modules);
+			}
 			ConsumeModules("Common", commonModules);
 
 			// Remove any binaries that no longer have modules due to merging
