[UBT]
* Fixed so TestTargetRules.bCompileAgainstEngine is not set when assigning bCompileAgainstEngine since it affects all test targets.

diff --git a/Engine/Source/Developer/LowLevelTestsRunner/LowLevelTestsRunner.Build.cs b/Engine/Source/Developer/LowLevelTestsRunner/LowLevelTestsRunner.Build.cs
--- a/Engine/Source/Developer/LowLevelTestsRunner/LowLevelTestsRunner.Build.cs
+++ b/Engine/Source/Developer/LowLevelTestsRunner/LowLevelTestsRunner.Build.cs
@@ -30,7 +30,7 @@ namespace UnrealBuildTool.Rules
 			{
 				PrivateDependencyModuleNames.Add("UnrealEd");
 			}
-			if (TestTargetRules.bTestsRequireEngine)
+			if (TestTargetRules.bTestsRequireEngine == true)
 			{
 				PrivateDependencyModuleNames.Add("Engine");
 			}
diff --git a/Engine/Source/Programs/UnrealBuildTool/Configuration/TargetRules.cs b/Engine/Source/Programs/UnrealBuildTool/Configuration/TargetRules.cs
--- a/Engine/Source/Programs/UnrealBuildTool/Configuration/TargetRules.cs
+++ b/Engine/Source/Programs/UnrealBuildTool/Configuration/TargetRules.cs
@@ -1213,7 +1213,11 @@ namespace UnrealBuildTool
 			get => bCompileAgainstEnginePrivate;
 			set => bCompileAgainstEnginePrivate = value;
 		}
-		private bool bCompileAgainstEnginePrivate = true;
+
+		/// <summary>
+		/// variable backing bCompileAgainstEngine
+		/// </summary>
+		protected bool bCompileAgainstEnginePrivate = true;
 
 		/// <summary>
 		/// Enabled for all builds that include the CoreUObject project.  Disabled only when building standalone apps that only link with Core.
diff --git a/Engine/Source/Programs/UnrealBuildTool/Configuration/TestTargetRules.cs b/Engine/Source/Programs/UnrealBuildTool/Configuration/TestTargetRules.cs
--- a/Engine/Source/Programs/UnrealBuildTool/Configuration/TestTargetRules.cs
+++ b/Engine/Source/Programs/UnrealBuildTool/Configuration/TestTargetRules.cs
@@ -28,7 +28,7 @@ namespace UnrealBuildTool
 		/// <summary>
 		/// Keeps track if low level tests executable must build with the Engine.
 		/// </summary>
-		public static bool bTestsRequireEngine { get; set; }
+		public static bool? bTestsRequireEngine { get; set; }
 
 		/// <summary>
 		/// Keeps track if low level tests executable must build with the ApplicationCore.
@@ -81,8 +81,8 @@ namespace UnrealBuildTool
 		/// </summary>
 		public override bool bCompileAgainstEngine
 		{
-			get => bTestsRequireEngine;
-			set => bTestsRequireEngine = value;
+			get => bTestsRequireEngine??bCompileAgainstEnginePrivate;
+			set => bCompileAgainstEnginePrivate = value;
 		}
 
 		/// <summary>
@@ -295,7 +295,7 @@ namespace UnrealBuildTool
 			bForceBuildShaderFormats = false;
 
 			// Do not compile against the engine, editor etc
-			bCompileAgainstEngine = bTestsRequireEngine;
+			bCompileAgainstEnginePrivate = false;
 			bCompileAgainstEditor = bTestsRequireEditor;
 			bCompileAgainstCoreUObject = bTestsRequireCoreUObject;
 			bCompileAgainstApplicationCore = bTestsRequireApplicationCore;
diff --git a/Engine/Source/Programs/UnrealBuildTool/Modes/TestMode.cs b/Engine/Source/Programs/UnrealBuildTool/Modes/TestMode.cs
--- a/Engine/Source/Programs/UnrealBuildTool/Modes/TestMode.cs
+++ b/Engine/Source/Programs/UnrealBuildTool/Modes/TestMode.cs
@@ -261,7 +261,7 @@ namespace UnrealBuildTool
 			HashSet<string> moduleNames = [.. target.Binaries.SelectMany(x => x.Modules).Select(x => x.Name)];
 
 			bool updated = false;
-			if (moduleNames.Contains(ENGINE_MODULE) && !TestTargetRules.bTestsRequireEngine)
+			if (moduleNames.Contains(ENGINE_MODULE) && TestTargetRules.bTestsRequireEngine == false)
 			{
 				TestTargetRules.bTestsRequireEngine = true;
 				Logger.LogInformation("{Module} module found in dependency graph, will compile against it.", ENGINE_MODULE);
