Add support for Code coverage on Windows with Clang

#rnx

diff --git a/Engine/Source/Programs/UnrealBuildTool/Platform/Windows/VCToolChain.cs b/Engine/Source/Programs/UnrealBuildTool/Platform/Windows/VCToolChain.cs
--- a/Engine/Source/Programs/UnrealBuildTool/Platform/Windows/VCToolChain.cs
+++ b/Engine/Source/Programs/UnrealBuildTool/Platform/Windows/VCToolChain.cs
@@ -855,6 +855,19 @@ namespace UnrealBuildTool
 				}
 			}
 
+			// Code Coverage Support
+			if (CompileEnvironment.bCodeCoverage)
+			{
+				// Clang llvm-cov coverage flags
+				if(Target.WindowsPlatform.Compiler.IsClang())
+				{
+					Arguments.Add("/Od");
+					Arguments.Add("/clang:-fprofile-instr-generate");
+					Arguments.Add("/clang:-fcoverage-mapping");
+					Log.TraceInformationOnce("Enabling Clang code coverage flags");
+				}
+			}
+
 			//
 			//	Debug
 			//
