[UBT]
* Disabled warning of files being created outside filehashcache partition. pvs writes a couple files and we don't care about tracking if they are the same in the file hash cache db

diff --git a/Engine/Source/Programs/UnrealBuildTool/Modes/BuildMode.cs b/Engine/Source/Programs/UnrealBuildTool/Modes/BuildMode.cs
--- a/Engine/Source/Programs/UnrealBuildTool/Modes/BuildMode.cs
+++ b/Engine/Source/Programs/UnrealBuildTool/Modes/BuildMode.cs
@@ -164,8 +164,13 @@ namespace UnrealBuildTool
 				}
 			}
 
+			bool warnOnNoPartition = false;
+#if DEBUG
+			warnOnNoPartition = true; // Pvs target writes three files that is not covered by partition
+#endif
+
 			CppDependencyCache cppDependencyCache = new();
-			FileHashCache fileHashCache = new(true, true);
+			FileHashCache fileHashCache = new(warnOnNoPartition, true);
 			List<Task> tasks = [];
 
 			// Parse and build the targets
