[UBT]
* Added MountAll to file hash cache to catch everything in one partition
* Fixed so FileHashCache can retry write if two threads are writing the same file at the same time. This can happen when generating solutions

diff --git a/Engine/Source/Programs/UnrealBuildTool/System/FileHashCache.cs b/Engine/Source/Programs/UnrealBuildTool/System/FileHashCache.cs
--- a/Engine/Source/Programs/UnrealBuildTool/System/FileHashCache.cs
+++ b/Engine/Source/Programs/UnrealBuildTool/System/FileHashCache.cs
@@ -11,6 +11,7 @@ using System.IO.Hashing;
 using System.Linq;
 using System.Runtime.InteropServices;
 using System.Text;
+using System.Threading;
 using System.Threading.Tasks;
 using UnrealBuildBase;
 
@@ -120,6 +121,17 @@ namespace UnrealBuildTool
 			AddPartition(root, "", logger);
 		}
 
+		/// <summary>
+		/// Catch all mount
+		/// </summary>
+		public void MountAll(DirectoryReference intermediateDir, ILogger logger)
+		{
+			FileReference file = FileReference.Combine(intermediateDir, "FileHashCache.bin");
+			Partition partition = new(file);
+			_roots.Add((new DirectoryReference("", DirectoryReference.Sanitize.None), [("", partition)]));
+			ReadPartition(partition, logger);
+		}
+
 		/// <summary>
 		/// Get hash of MemoryStream
 		/// </summary>
@@ -388,25 +400,59 @@ namespace UnrealBuildTool
 				}
 			}
 
-
-			if (data is MemoryStream stream)
-			{
-				stream.Position = 0;
-				using var file = new FileStream(fileItem.FullName, FileMode.Create, FileAccess.Write, FileShare.None);
-				stream.CopyTo(file);
-			}
-			else if (data is string content)
+			// When generating solutions there is a tiny chance that we will write the same file at the same time
+			// .. so let's just handle that with retry.
+			int retryIndex = 0;
+			bool isLocked = false;
+			while (true)
 			{
-				using var fs = new FileStream(fileItem.FullName, FileMode.Create, FileAccess.Write, FileShare.Read, 1 << 20, FileOptions.SequentialScan);
-				using var writer = new StreamWriter(fs, new UTF8Encoding(false), 1 << 20, false);
-				writer.Write(content);
-			}
-			else if (data is IEnumerable<string> contentLines)
-			{
-				string contents = String.Join(Environment.NewLine, contentLines);
-				using var fs = new FileStream(fileItem.FullName, FileMode.Create, FileAccess.Write, FileShare.Read, 1 << 20, FileOptions.SequentialScan);
-				using var writer = new StreamWriter(fs, new UTF8Encoding(false), 1 << 20, false);
-				writer.Write(contents);
+				try
+				{
+					if (data is MemoryStream stream)
+					{
+						stream.Position = 0;
+						using var file = new FileStream(fileItem.FullName, FileMode.Create, FileAccess.Write, FileShare.None);
+						stream.CopyTo(file);
+					}
+					else if (data is string content)
+					{
+						using var fs = new FileStream(fileItem.FullName, FileMode.Create, FileAccess.Write, FileShare.Read, 1 << 20, FileOptions.SequentialScan);
+						using var writer = new StreamWriter(fs, new UTF8Encoding(false), 1 << 20, false);
+						writer.Write(content);
+					}
+					else if (data is IEnumerable<string> contentLines)
+					{
+						string contents = String.Join(Environment.NewLine, contentLines);
+						using var fs = new FileStream(fileItem.FullName, FileMode.Create, FileAccess.Write, FileShare.Read, 1 << 20, FileOptions.SequentialScan);
+						using var writer = new StreamWriter(fs, new UTF8Encoding(false), 1 << 20, false);
+						writer.Write(contents);
+					}
+					if (isLocked)
+					{
+						Monitor.Exit(fileItem);
+					}
+					break;
+				}
+				catch (Exception e)
+				{
+					if (retryIndex == 20)
+					{
+						if (isLocked)
+						{
+							Monitor.Exit(fileItem);
+						}
+						throw;
+					}
+
+					if (!isLocked)
+					{
+						logger.LogInformation($"Get exception trying to write {fileItem} ({e.Message}). Will wait 1 second and try again");
+						Monitor.Enter(fileItem);
+						isLocked = true;
+					}
+					Thread.Sleep(100);
+					++retryIndex;
+				}
 			}
 
 			fileItem.ResetCachedInfo();
@@ -473,9 +519,14 @@ namespace UnrealBuildTool
 			int index = fullName.IndexOf(Intermediate, FileReference.Comparison);
 			if (index == -1)
 			{
+				if (_roots.Count == 1 && _roots[0].root.FullName == string.Empty) // MountAll
+				{
+					return _roots[0].Item2[0].partition;
+				}
+
 				if (_warnOnNoPartition)
 				{
-					logger.LogWarning($"{fullName} is not under intermediate folder");
+					logger.LogWarning($"FileHashCache - {fullName} is not under intermediate folder");
 				}
 				return null;
 			}
@@ -496,7 +547,7 @@ namespace UnrealBuildTool
 
 			foreach (var root in _roots)
 			{
-				if (!fileRef.IsUnderDirectory(root.root))
+				if (root.root.FullName.Length != 0 && !fileRef.IsUnderDirectory(root.root))
 				{
 					continue;
 				}
@@ -517,7 +568,7 @@ namespace UnrealBuildTool
 			{
 				if (_warnOnNoPartition)
 				{
-					logger.LogWarning($"There is no partition covering path to {fileRef}!");
+					logger.LogWarning($"FileHashCache - There is no partition covering path to {fileRef}!");
 				}
 				return null;
 			}
