[UBT]
* Improved logging text

diff --git a/Engine/Source/Programs/UnrealBuildTool/System/CppDependencyCache.cs b/Engine/Source/Programs/UnrealBuildTool/System/CppDependencyCache.cs
--- a/Engine/Source/Programs/UnrealBuildTool/System/CppDependencyCache.cs
+++ b/Engine/Source/Programs/UnrealBuildTool/System/CppDependencyCache.cs
@@ -309,7 +309,7 @@ namespace UnrealBuildTool
 
 			if (partition == null)
 			{
-				Logger.LogWarning($"There is no partition covering path to {InputFile}!");
+				Logger.LogWarning($"CppDependencyCache - There is no partition covering path to {InputFile}!");
 				OutInfo = null;
 				return false;
 			}
