[UBT]
* Changed so update check is done on values instead of string.
* Made ProgressWriter thread safe by adding lock when actually updating string

diff --git a/Engine/Source/Programs/UnrealBuildTool/System/ProgressWriter.cs b/Engine/Source/Programs/UnrealBuildTool/System/ProgressWriter.cs
--- a/Engine/Source/Programs/UnrealBuildTool/System/ProgressWriter.cs
+++ b/Engine/Source/Programs/UnrealBuildTool/System/ProgressWriter.cs
@@ -34,7 +34,7 @@ namespace UnrealBuildTool
 		/// <summary>
 		/// The current progress message
 		/// </summary>
-		string? CurrentProgressString;
+		int CurrentProgressValue = 0;
 
 		/// <summary>
 		/// Constructor
@@ -72,16 +72,27 @@ namespace UnrealBuildTool
 		/// <param name="Denominator">Denominator for the progress fraction</param>
 		public void Write(int Numerator, int Denominator)
 		{
-			float ProgressValue = Denominator > 0 ? ((float)Numerator / (float)Denominator) : 1.0f;
-			string ProgressString = String.Format("{0}%", Math.Round(ProgressValue * 100.0f));
-			if (ProgressString != CurrentProgressString)
+			float ProgressValueFloat = Denominator > 0 ? ((float)Numerator / (float)Denominator) : 1.0f;
+			int ProgressValue = (int)Math.Round(ProgressValueFloat * 100.0f);
+
+			if (ProgressValue != CurrentProgressValue)
 			{
-				CurrentProgressString = ProgressString;
-				if (bWriteMarkup)
+				lock (this)
 				{
-					Logger.LogInformation("@progress '{Message}' {ProgressString}", Message, ProgressString);
+					if (ProgressValue <= CurrentProgressValue)
+					{
+						return;
+					}
+					CurrentProgressValue = ProgressValue;
+
+					string ProgressString = $"{ProgressValue}%";
+
+					if (bWriteMarkup)
+					{
+						Logger.LogInformation("@progress '{Message}' {ProgressString}", Message, ProgressString);
+					}
+					Status?.SetProgress(ProgressString);
 				}
-				Status?.SetProgress(ProgressString);
 			}
 		}
 	}
