[UBT]
* Optimizations to project generation. Cut down 10 seconds when intellisense is enabled
* Fixed so project generation with intellisense is deterministic
* Added HashSetPool and use it for temporary hashsets across UBT to reduce gc pressure
* Added so project generation is using file hash cache.
* Made more things run in parallel
* Various minor optimizations based on profiling

diff --git a/Engine/Source/Programs/UnrealBuildTool/Configuration/ModuleRules.cs b/Engine/Source/Programs/UnrealBuildTool/Configuration/ModuleRules.cs
--- a/Engine/Source/Programs/UnrealBuildTool/Configuration/ModuleRules.cs
+++ b/Engine/Source/Programs/UnrealBuildTool/Configuration/ModuleRules.cs
@@ -1012,7 +1012,9 @@ namespace UnrealBuildTool
 		///    'clang -Xclang -analyzer-checker-help-alpha' 
 		/// for the list of experimental checkers.
 		/// </summary>
-		public HashSet<string> StaticAnalyzerDisabledCheckers { get; } = [
+		public HashSet<string> StaticAnalyzerDisabledCheckers { get; } = []; // This is set to DefaultStaticAnalyzerDisabledCheckers if static analyzer is set on target
+
+		private static HashSet<string> DefaultStaticAnalyzerDisabledCheckers = [
 			// Work in progress, remove disabled checkers as all warnings are fixed
 			"core.BitwiseShift",
 			"core.CallAndMessage",
@@ -1505,6 +1507,11 @@ namespace UnrealBuildTool
 		/// <param name="target">Rules for building this target</param>
 		public ModuleRules(ReadOnlyTargetRules target)
 		{
+			if (target.StaticAnalyzer != StaticAnalyzer.None)
+			{
+				StaticAnalyzerDisabledCheckers = new(DefaultStaticAnalyzerDisabledCheckers);
+			}
+
 			Target = target;
 			CppCompileWarningSettings = new CppCompileWarnings(this, Logger);
 		}
diff --git a/Engine/Source/Programs/UnrealBuildTool/Configuration/UEBuildBinary.cs b/Engine/Source/Programs/UnrealBuildTool/Configuration/UEBuildBinary.cs
--- a/Engine/Source/Programs/UnrealBuildTool/Configuration/UEBuildBinary.cs
+++ b/Engine/Source/Programs/UnrealBuildTool/Configuration/UEBuildBinary.cs
@@ -2,7 +2,6 @@
 
 using System;
 using System.Collections.Generic;
-using System.IO;
 using System.Linq;
 using EpicGames.Core;
 using Microsoft.Extensions.Logging;
@@ -195,17 +194,6 @@ namespace UnrealBuildTool
 			});
 		}
 
-		/// <summary>
-		/// Creates all the modules referenced by this target.
-		/// </summary>
-		public void CreateAllDependentModules(UEBuildModule.CreateModuleDelegate CreateModule, ILogger Logger)
-		{
-			foreach (UEBuildModule Module in Modules)
-			{
-				Module.RecursivelyCreateModules(CreateModule, "Target", Logger);
-			}
-		}
-
 		/// <summary>
 		/// Expands all of the runtime dependencies.
 		/// </summary>
@@ -426,7 +414,9 @@ namespace UnrealBuildTool
 			}
 
 			List<UEBuildModule> ReferencedModules = new List<UEBuildModule>();
-			HashSet<UEBuildModule> IgnoreReferencedModules = new HashSet<UEBuildModule>();
+
+			using var rented = HashSetPool<UEBuildModule>.Rent();
+			HashSet<UEBuildModule> IgnoreReferencedModules = rented.Value;
 
 			foreach (UEBuildModule Module in Modules)
 			{
@@ -456,7 +446,10 @@ namespace UnrealBuildTool
 			}
 
 			List<UEBuildModule> DirectlyReferencedModules = new List<UEBuildModule>();
-			HashSet<UEBuildModule> VisitedModules = new HashSet<UEBuildModule>();
+
+			using var rented = HashSetPool<UEBuildModule>.Rent();
+			HashSet<UEBuildModule> VisitedModules = rented.Value;
+
 			for (int Idx = 0; Idx < ReferencedModules.Count; Idx++)
 			{
 				UEBuildModule SourceModule = ReferencedModules[Idx];
@@ -733,9 +726,11 @@ namespace UnrealBuildTool
 		private LinkEnvironment SetupBinaryLinkEnvironment(ReadOnlyTargetRules Target, UEToolChain ToolChain, LinkEnvironment LinkEnvironment, CppCompileEnvironment CompileEnvironment, ISourceFileWorkingSet WorkingSet, DirectoryReference ExeDir, IActionGraphBuilder Graph, ILogger Logger)
 		{
 			LinkEnvironment BinaryLinkEnvironment = new LinkEnvironment(LinkEnvironment);
-			HashSet<UEBuildModule> LinkEnvironmentVisitedModules = new HashSet<UEBuildModule>();
 			List<UEBuildBinary> BinaryDependencies = new List<UEBuildBinary>();
 
+			using var rented = HashSetPool<UEBuildModule>.Rent();
+			HashSet<UEBuildModule> LinkEnvironmentVisitedModules = rented.Value;
+
 			CppCompileEnvironment BinaryCompileEnvironment = CreateBinaryCompileEnvironment(CompileEnvironment);
 
 			if (BinaryCompileEnvironment.bUseSharedBuildEnvironment && Target.ProjectFile != null && IntermediateDirectory.IsUnderDirectory(Target.ProjectFile.Directory))
@@ -743,7 +738,8 @@ namespace UnrealBuildTool
 				BinaryCompileEnvironment.bUseSharedBuildEnvironment = false;
 			}
 
-			HashSet<FileItem> InputFilesLookup = new();
+			using var rented2 = HashSetPool<FileItem>.Rent();
+			HashSet<FileItem> InputFilesLookup = rented2.Value;
 
 			foreach (UEBuildModule Module in UEBuildModule.StableTopologicalSort(Modules))
 			{
diff --git a/Engine/Source/Programs/UnrealBuildTool/Configuration/UEBuildModule.cs b/Engine/Source/Programs/UnrealBuildTool/Configuration/UEBuildModule.cs
--- a/Engine/Source/Programs/UnrealBuildTool/Configuration/UEBuildModule.cs
+++ b/Engine/Source/Programs/UnrealBuildTool/Configuration/UEBuildModule.cs
@@ -2,12 +2,10 @@
 
 using System;
 using System.Collections.Generic;
-using System.Configuration;
 using System.IO;
 using System.Linq;
 using System.Reflection;
 using System.Text;
-using System.Threading.Tasks;
 using EpicGames.Core;
 using Microsoft.CodeAnalysis;
 using Microsoft.Extensions.Logging;
@@ -218,6 +216,11 @@ namespace UnrealBuildTool
 
 		public HashSet<UEBuildModule>? DependenciesForStableTopologicalSort { get; private set; }
 
+		static readonly string[] ApiDefinitions = ["", "DLLEXPORT", "DLLIMPORT"];
+
+		private string[] ModuleApiDefinitionStrings;
+		private string[] ModuleNonAttributedApiDefinitionStrings;
+
 		/// <summary>
 		/// Constructor
 		/// </summary>
@@ -343,6 +346,30 @@ namespace UnrealBuildTool
 			{
 				NatvisFiles.UnionWith(directory.EnumerateFiles().Where(x => x.HasExtension(".natvis") || x.HasExtension(".natstepfilter") || x.HasExtension(".natjmc")));
 			}
+
+			ModuleApiDefinitionStrings = new string[3];
+			ModuleNonAttributedApiDefinitionStrings = new string[3];
+
+
+			// We use a C-style attribute here because of where these macros are defined in UE. C++ attributes must come at the beginning of the function
+			// definition, like `[[my_attr]] void Foo();`, but these macros are used like `void FOO_API Foo();`, and if we used the C++ attribute name to
+			// disable AutoRTFM on this function the compiler would spit nails. But it'll happily consume the C-style attribute at that location and do
+			// the right thing, so we use that instead.
+			// Caution: Hard coded constant '1' in this string must be kept in sync with Engine\Source\Runtime\AutoRTFM\Public\AutoRTFMConstants.h.
+			String AutoRTFMDisable = (Rules.Target.bUseAutoRTFMCompiler && Rules.bDisableAutoRTFMInstrumentation && Rules.bInjectAutoRTFMAttributeDisables) ?
+				"__attribute__((autortfm(/*autortfm_mode_disable*/ 1)))" : String.Empty;
+
+			// We precalculate these because they are reused over and over again and causing tons of allocations
+			for (int i=0; i!=3; ++i)
+			{
+				string ApiDefinition = ApiDefinitions[i]; 
+
+				// Add a spacer if and only if we've actually used import/export above.
+				String Spacer = ApiDefinition.Length == 0 ? String.Empty : " ";
+				// Add the import or export declaration for the module
+				ModuleApiDefinitionStrings[i] = $"{ModuleApiDefine}={ApiDefinition}{Spacer}{AutoRTFMDisable}";
+				ModuleNonAttributedApiDefinitionStrings[i] = $"{ModuleNonAttributedDefine}={ApiDefinition}";
+			}
 		}
 
 		/// <summary>
@@ -439,7 +466,13 @@ namespace UnrealBuildTool
 		/// </summary>
 		protected HashSet<string> HashSetFromOptionalEnumerableStringParameter(IEnumerable<string> InEnumerableStrings)
 		{
-			return InEnumerableStrings == null ? [] : [.. InEnumerableStrings.Select(x => ExpandPathVariables(x, null, null))];
+			InEnumerableStrings.TryGetNonEnumeratedCount(out int estimatedCount);
+			HashSet<string> hashSet = new(estimatedCount);
+			foreach (string s in InEnumerableStrings)
+			{
+				hashSet.Add(ExpandPathVariables(s, null, null));
+			}
+			return hashSet;
 		}
 
 		/// <summary>
@@ -610,11 +643,9 @@ namespace UnrealBuildTool
 		/// <param name="ModuleToIncludePathsOnlyFlag"></param>
 		protected void FindIncludePathModulesInPublicCompileEnvironment(Dictionary<UEBuildModule, bool> ModuleToIncludePathsOnlyFlag)
 		{
-			if (!ModuleToIncludePathsOnlyFlag.ContainsKey(this))
+			// Add this module to the list
+			if (ModuleToIncludePathsOnlyFlag.TryAdd(this, true))
 			{
-				// Add this module to the list
-				ModuleToIncludePathsOnlyFlag.Add(this, true);
-
 				// Include any of its public include path modules in the compile environment too
 				foreach (UEBuildModule IncludePathModule in PublicIncludePathModules!)
 				{
@@ -706,41 +737,27 @@ namespace UnrealBuildTool
 
 			if (Rules.Type == ModuleRules.ModuleType.CPlusPlus)
 			{
-				String ApiDefinition = String.Empty;
+				int ApiDefinitionIndex = 0;
 
 				if (Rules.Target.LinkType == TargetLinkType.Monolithic)
 				{
 					if (Rules.Target.bShouldCompileAsDLL && (Rules.Target.bHasExports || Rules.ModuleSymbolVisibility == ModuleRules.SymbolVisibility.VisibileForDll))
 					{
-						ApiDefinition = "DLLEXPORT";
+						ApiDefinitionIndex = 1;
 					}
 				}
 				else if (Binary == null || (Rules.Target.bMergeModules ? (SourceBinary == null || SourceModule != this) : SourceBinary != Binary))
 				{
-					ApiDefinition = "DLLIMPORT";
+					ApiDefinitionIndex = 2;
 				}
 				else if (!(!Binary.bAllowExports || Rules.Target.bMergeModules))
 				{
-					ApiDefinition = "DLLEXPORT";
+					ApiDefinitionIndex = 1;
 				}
 
-				// We use a C-style attribute here because of where these macros are defined in UE. C++ attributes must come at the beginning of the function
-				// definition, like `[[my_attr]] void Foo();`, but these macros are used like `void FOO_API Foo();`, and if we used the C++ attribute name to
-				// disable AutoRTFM on this function the compiler would spit nails. But it'll happily consume the C-style attribute at that location and do
-				// the right thing, so we use that instead.
-				// Caution: Hard coded constant '1' in this string must be kept in sync with Engine\Source\Runtime\AutoRTFM\Public\AutoRTFMConstants.h.
-				String AutoRTFMDisable = (Rules.Target.bUseAutoRTFMCompiler && Rules.bDisableAutoRTFMInstrumentation && Rules.bInjectAutoRTFMAttributeDisables) ?
-					"__attribute__((autortfm(/*autortfm_mode_disable*/ 1)))" : String.Empty;
-
-				// Add a spacer if and only if we've actually used import/export above.
-				String Spacer = ApiDefinition.Length == 0 ? String.Empty : " ";
-
 				// Add the import or export declaration for the module
-				String ModuleApiDefinitionString = $"{ModuleApiDefine}={ApiDefinition}{Spacer}{AutoRTFMDisable}";
-				Definitions.Add(ModuleApiDefinitionString);
-
-				String ModuleNonAttributedApiDefinitionString = $"{ModuleNonAttributedDefine}={ApiDefinition}";
-				Definitions.Add(ModuleNonAttributedApiDefinitionString);
+				Definitions.Add(ModuleApiDefinitionStrings[ApiDefinitionIndex]);
+				Definitions.Add(ModuleNonAttributedApiDefinitionStrings[ApiDefinitionIndex]);
 			}
 
 			// Add the additional frameworks so that the compiler can know about their #include paths
@@ -1122,13 +1139,14 @@ namespace UnrealBuildTool
 		/// <summary>
 		/// Gathers all the module dependencies a PCH would have
 		/// </summary>
+		/// <param name="outReferenced">The collected modules</param>
 		/// <param name="bIncludePrivateModules">Whether to include private modules.</param>
-		public HashSet<UEBuildModule> GetAllDependencyModulesForPCH(bool bIncludePrivateModules)
+		public void GetAllDependencyModulesForPCH(HashSet<UEBuildModule> outReferenced, bool bIncludePrivateModules)
 		{
-			HashSet<UEBuildModule> ReferencedModules = new HashSet<UEBuildModule>();
-			HashSet<UEBuildModule> IgnoreReferencedModules = new HashSet<UEBuildModule> { this };
-			InternalGetAllDependencyModulesForPCH(ReferencedModules, IgnoreReferencedModules, bIncludePrivateModules);
-			return ReferencedModules;
+			using var rented = HashSetPool<UEBuildModule>.Rent();
+			HashSet<UEBuildModule> IgnoreReferencedModules = rented.Value;
+			IgnoreReferencedModules.Add(this);
+			InternalGetAllDependencyModulesForPCH(outReferenced, IgnoreReferencedModules, bIncludePrivateModules);
 		}
 
 		/// <summary>
@@ -1139,11 +1157,17 @@ namespace UnrealBuildTool
 		/// <param name="bIncludePrivateModules">Whether to include private modules.</param>
 		private void InternalGetAllDependencyModulesForPCH(HashSet<UEBuildModule> ReferencedModules, HashSet<UEBuildModule> IgnoreReferencedModules, bool bIncludePrivateModules)
 		{
-			HashSet<UEBuildModule> AllDependencyModules = new(
-				((bIncludePrivateModules && PrivateDependencyModules != null) ? PrivateDependencyModules.Count : 0) +
+			//HashSet<UEBuildModule> AllDependencyModules = new(
+			//	((bIncludePrivateModules && PrivateDependencyModules != null) ? PrivateDependencyModules.Count : 0) +
+			//	(PublicDependencyModules != null ? PublicDependencyModules.Count : 0) +
+			//	(PublicIncludePathModules != null ? PublicIncludePathModules!.Count : 0)
+			//	);
+
+			using var rented = HashSetPool<UEBuildModule>.Rent(((bIncludePrivateModules && PrivateDependencyModules != null) ? PrivateDependencyModules.Count : 0) +
 				(PublicDependencyModules != null ? PublicDependencyModules.Count : 0) +
-				(PublicIncludePathModules != null ? PublicIncludePathModules!.Count : 0)
-				);
+				(PublicIncludePathModules != null ? PublicIncludePathModules!.Count : 0));
+			HashSet<UEBuildModule> AllDependencyModules = rented.Value;
+
 			if (bIncludePrivateModules && PrivateDependencyModules != null)
 			{
 				AllDependencyModules.UnionWith(PrivateDependencyModules);
@@ -1181,13 +1205,15 @@ namespace UnrealBuildTool
 		/// <param name="CreateModule">Delegate to create a module with a given name</param>
 		/// <param name="ReferenceChain">Chain of references before reaching this module</param>
 		/// <param name="Logger">Logger for output</param>
-		public void RecursivelyCreateModules(CreateModuleDelegate CreateModule, string ReferenceChain, ILogger Logger)
+		public void RecursivelyCreateModules(CreateModuleDelegate CreateModule, StringBuilder ReferenceChain, ILogger Logger)
 		{
 			// Name of this reference
 			string ThisRefName = (RulesFile == null) ? Name : RulesFile.GetFileName();
 
 			// Set the reference chain for anything referenced by this module
-			string NextReferenceChain = String.Format("{0} -> {1}", ReferenceChain, ThisRefName);
+			//string NextReferenceChain = String.Format("{0} -> {1}", ReferenceChain, ThisRefName);
+			int OriginalReferenceChainLength = ReferenceChain.Length;
+			ReferenceChain.Append(" -> ").Append(ThisRefName);
 
 			// Check if this module is invalid for the current target, based on Module attributes
 			if (Rules.CppCompileWarningSettings.ModuleUnsupportedWarningLevel != WarningLevel.Off)
@@ -1196,11 +1222,11 @@ namespace UnrealBuildTool
 				{
 					if (Rules.CppCompileWarningSettings.ModuleUnsupportedWarningLevel == WarningLevel.Warning)
 					{
-						Logger.LogWarning("Warning: Referenced module '{Module}' not supported '{InvalidReason}' via {ReferenceChain}", ThisRefName, InvalidReason, NextReferenceChain);
+						Logger.LogWarning("Warning: Referenced module '{Module}' not supported '{InvalidReason}' via {ReferenceChain}", ThisRefName, InvalidReason, ReferenceChain);
 					}
 					else if (Rules.CppCompileWarningSettings.ModuleUnsupportedWarningLevel == WarningLevel.Error)
 					{
-						Logger.LogError("Error: Referenced module '{Module}' not supported '{InvalidReason}' via {ReferenceChain}", ThisRefName, InvalidReason, NextReferenceChain);
+						Logger.LogError("Error: Referenced module '{Module}' not supported '{InvalidReason}' via {ReferenceChain}", ThisRefName, InvalidReason, ReferenceChain);
 					}
 				}
 			}
@@ -1212,11 +1238,11 @@ namespace UnrealBuildTool
 				{
 					if (Rules.CppCompileWarningSettings.PluginModuleUnsupportedWarningLevel == WarningLevel.Warning)
 					{
-						Logger.LogWarning("Warning: Referenced plugin '{Plugin}' module '{Module}' not supported '{InvalidReason}' via {ReferenceChain}", Rules.Plugin?.File.GetFileName(), ThisRefName, InvalidReason, NextReferenceChain);
+						Logger.LogWarning("Warning: Referenced plugin '{Plugin}' module '{Module}' not supported '{InvalidReason}' via {ReferenceChain}", Rules.Plugin?.File.GetFileName(), ThisRefName, InvalidReason, ReferenceChain);
 					}
 					else if (Rules.CppCompileWarningSettings.PluginModuleUnsupportedWarningLevel == WarningLevel.Error)
 					{
-						Logger.LogError("Error: Referenced plugin '{Plugin}' module '{Module}' not supported '{InvalidReason}' via {ReferenceChain}", Rules.Plugin?.File.GetFileName(), ThisRefName, InvalidReason, NextReferenceChain);
+						Logger.LogError("Error: Referenced plugin '{Plugin}' module '{Module}' not supported '{InvalidReason}' via {ReferenceChain}", Rules.Plugin?.File.GetFileName(), ThisRefName, InvalidReason, ReferenceChain);
 					}
 				}
 			}
@@ -1239,7 +1265,7 @@ namespace UnrealBuildTool
 				ReferenceStackParentModules ??= new(Comparer<UEBuildModule>.Create(static (a, b) => a.Name.CompareTo(b.Name)));
 
 				// Create all the dependency modules - pass through the reference stack so we can check for cycles
-				RecursivelyCreateModulesByName(Rules.PublicDependencyModuleNames, ref PublicDependencyModules, ref bDependsOnVerse, CreateModule, NextReferenceChain, Logger);
+				RecursivelyCreateModulesByName(Rules.PublicDependencyModuleNames, ref PublicDependencyModules, ref bDependsOnVerse, CreateModule, ReferenceChain, Logger);
 				if (Rules.Target.IsTestTarget)
 				{
 					// Move the test runner dependency to last position to give it the opportunity to build its special dependencies such as CoreUObject, ApplicationCore etc.
@@ -1249,23 +1275,26 @@ namespace UnrealBuildTool
 						Rules.PrivateDependencyModuleNames.Add("LowLevelTestsRunner");
 					}
 				}
-				RecursivelyCreateModulesByName(Rules.PrivateDependencyModuleNames, ref PrivateDependencyModules, ref bDependsOnVerse, CreateModule, NextReferenceChain, Logger);
+				RecursivelyCreateModulesByName(Rules.PrivateDependencyModuleNames, ref PrivateDependencyModules, ref bDependsOnVerse, CreateModule, ReferenceChain, Logger);
 
 				// Recursively create all the public include path modules
 				HashSet<string> PublicIncludePathModuleNames = new();
 				PublicIncludePathModuleNames.UnionWith(Rules.PublicIncludePathModuleNames!);
 				PublicIncludePathModuleNames.UnionWith(Rules.PublicDependencyModuleNames!);
-				RecursivelyCreateIncludePathModulesByName(PublicIncludePathModuleNames, ref PublicIncludePathModules, ref bDependsOnVerse, CreateModule, NextReferenceChain);
+				RecursivelyCreateIncludePathModulesByName(PublicIncludePathModuleNames, ref PublicIncludePathModules, ref bDependsOnVerse, CreateModule, ReferenceChain);
 
 				// Create the private include path modules
-				RecursivelyCreateIncludePathModulesByName(Rules.PrivateIncludePathModuleNames, ref PrivateIncludePathModules, ref bDependsOnVerse, CreateModule, NextReferenceChain);
+				RecursivelyCreateIncludePathModulesByName(Rules.PrivateIncludePathModuleNames, ref PrivateIncludePathModules, ref bDependsOnVerse, CreateModule, ReferenceChain);
 
 				// Dynamic loads aren't considered a reference chain so start with an empty stack
-				RecursivelyCreateModulesByName(Rules.DynamicallyLoadedModuleNames, ref DynamicallyLoadedModules, ref bDependsOnVerse, CreateModule, NextReferenceChain, Logger);
+				RecursivelyCreateModulesByName(Rules.DynamicallyLoadedModuleNames, ref DynamicallyLoadedModules, ref bDependsOnVerse, CreateModule, ReferenceChain, Logger);
 
 				PublicDependencyModules?.ForEach(x => x.ReferenceStackParentModules?.Add(this));
 				PrivateDependencyModules?.ForEach(x => x.ReferenceStackParentModules?.Add(this));
 			}
+
+			// Reset reference chain
+			ReferenceChain.Length = OriginalReferenceChainLength;
 		}
 
 		record CircularReference
@@ -1320,7 +1349,10 @@ namespace UnrealBuildTool
 		bool ValidateCircularReferences(ILogger logger)
 		{
 			List<CircularReference> circularReferences = [];
-			CheckForCircularReferences(this, [], [], circularReferences);
+
+			using var rented = HashSetPool<UEBuildModule>.Rent();
+			HashSet<UEBuildModule> processed = rented.Value;
+			CheckForCircularReferences(this, processed, [], circularReferences);
 
 			IEnumerable<CircularReference> allowed = circularReferences.Where(x => !x.IsError).OrderBy(x => x.ReferenceStack.Count);
 			IEnumerable<CircularReference> errors = circularReferences.Where(x => x.IsError).OrderBy(x => x.ReferenceStack.Count);
@@ -1416,7 +1448,7 @@ namespace UnrealBuildTool
 			}
 		}
 
-		private static void RecursivelyCreateModulesByName(List<string> ModuleNames, ref List<UEBuildModule>? Modules, ref bool bDependsOnVerse, CreateModuleDelegate CreateModule, string ReferenceChain, ILogger Logger)
+		private static void RecursivelyCreateModulesByName(List<string> ModuleNames, ref List<UEBuildModule>? Modules, ref bool bDependsOnVerse, CreateModuleDelegate CreateModule, StringBuilder ReferenceChain, ILogger Logger)
 		{
 			// Check whether the module list is already set. We set this immediately (via the ref) to avoid infinite recursion.
 			if (Modules == null)
@@ -1425,7 +1457,7 @@ namespace UnrealBuildTool
 				for (int i = 0; i < ModuleNames.Count; i++)
 				{
 					string ModuleName = ModuleNames[i];
-					UEBuildModule Module = CreateModule(ModuleName, ReferenceChain);
+					UEBuildModule Module = CreateModule(ModuleName, ReferenceChain.ToString());
 					if (!Modules.Contains(Module))
 					{
 						Module.RecursivelyCreateModules(CreateModule, ReferenceChain, Logger);
@@ -1436,28 +1468,32 @@ namespace UnrealBuildTool
 			}
 		}
 
-		private static void RecursivelyCreateIncludePathModulesByName(IEnumerable<string> ModuleNames, ref List<UEBuildModule>? Modules, ref bool bDependsOnVerse, CreateModuleDelegate CreateModule, string ReferenceChain)
+		private static void RecursivelyCreateIncludePathModulesByName(IEnumerable<string> ModuleNames, ref List<UEBuildModule>? Modules, ref bool bDependsOnVerse, CreateModuleDelegate CreateModule, StringBuilder ReferenceChain)
 		{
 			// Check whether the module list is already set. We set this immediately (via the ref) to avoid infinite recursion.
 			if (Modules == null)
 			{
+				string ReferenceChainString = ReferenceChain.ToString();
 				Modules = new List<UEBuildModule>();
 				foreach (string ModuleName in ModuleNames)
 				{
-					UEBuildModule Module = CreateModule(ModuleName, ReferenceChain);
+					UEBuildModule Module = CreateModule(ModuleName, ReferenceChainString);
 					if (!Modules.Contains(Module))
 					{
 						// Name of this reference
 						string ModuleRefName = (Module.RulesFile == null) ? Module.Name : Module.RulesFile.GetFileName();
 						// Set the reference chain for anything referenced by this module
-						string NextReferenceChain = String.Format("{0} -> {1} (public include)", ReferenceChain, ModuleRefName);
+						int OriginalReferenceChainLength = ReferenceChain.Length;
+						ReferenceChain.Append(" -> ").Append(ModuleRefName).Append(" (public include)");
 
 						HashSet<string> PublicIncludePathModuleNames = new();
 						PublicIncludePathModuleNames.UnionWith(Module.Rules.PublicIncludePathModuleNames);
 						PublicIncludePathModuleNames.UnionWith(Module.Rules.PublicDependencyModuleNames);
-						RecursivelyCreateIncludePathModulesByName(PublicIncludePathModuleNames, ref Module.PublicIncludePathModules, ref Module.bDependsOnVerse, CreateModule, NextReferenceChain);
+						RecursivelyCreateIncludePathModulesByName(PublicIncludePathModuleNames, ref Module.PublicIncludePathModules, ref Module.bDependsOnVerse, CreateModule, ReferenceChain);
 						Modules.Add(Module);
 						bDependsOnVerse |= Module.bDependsOnVerse;
+
+						ReferenceChain.Length = OriginalReferenceChainLength;
 					}
 				}
 			}
diff --git a/Engine/Source/Programs/UnrealBuildTool/Configuration/UEBuildModuleCPP.cs b/Engine/Source/Programs/UnrealBuildTool/Configuration/UEBuildModuleCPP.cs
--- a/Engine/Source/Programs/UnrealBuildTool/Configuration/UEBuildModuleCPP.cs
+++ b/Engine/Source/Programs/UnrealBuildTool/Configuration/UEBuildModuleCPP.cs
@@ -9,7 +9,6 @@ using System.Collections.Generic;
 using System.IO;
 using System.Linq;
 using System.Text;
-using System.Text.Json;
 using UnrealBuildBase;
 
 namespace UnrealBuildTool
@@ -516,7 +515,7 @@ namespace UnrealBuildTool
 					IncludePaths.Add(GeneratedCodeDirectoryUHT!);
 				}
 
-				if (bHasVerse || bForceAddIncludePath)
+				if (bHasVerse || bForceAddGeneratedCodeIncludePath) // bHasVerse should be up-to-date already from ctor.. so does not need to always be added when generating project files
 				{
 					IncludePaths.Add(GeneratedCodeDirectoryVNI!);
 				}
@@ -1326,6 +1325,8 @@ namespace UnrealBuildTool
 			return HeaderFileItems;
 		}
 
+		static ConcurrentDictionary<string, Lazy<HashSet<string>>> s_sharedPchModuleDependenciesCache = [];
+
 		/// <summary>
 		/// Gets the list of module dependencies that this module declares and then tries to optimize the list to what is actually used by the header.
 		/// </summary>
@@ -1335,15 +1336,48 @@ namespace UnrealBuildTool
 		/// <returns>Returns an optimized list of modules this module's shared PCH uses</returns>
 		private HashSet<UEBuildModule> GetSharedPCHModuleDependencies(CppCompileEnvironment CompileEnvironment, FileItem PCHHeaderFile, ILogger Logger)
 		{
-			Dictionary<UEBuildModule, bool> ModuleToIncludePathsOnlyFlag = new Dictionary<UEBuildModule, bool>();
+			Dictionary<UEBuildModule, bool> ModuleToIncludePathsOnlyFlag = new();
 			FindModulesInPublicCompileEnvironment(ModuleToIncludePathsOnlyFlag);
 
+			string key = PCHHeaderFile.FullName + CompileEnvironment.Platform;
+			HashSet<string> moduleNames = s_sharedPchModuleDependenciesCache.GetOrAdd(key, (k) => new Lazy<HashSet<string>>(() =>
+			{
+				HashSet<UEBuildModule> modules = GetSharedPCHModuleDependenciesUncached(CompileEnvironment, PCHHeaderFile, Logger, ModuleToIncludePathsOnlyFlag);
+				return modules.Select(m => m.Name).ToHashSet();
+			})).Value;
+
+			HashSet<UEBuildModule> modules = [];
+			foreach (UEBuildModule module in ModuleToIncludePathsOnlyFlag.Keys)
+			{
+				if (moduleNames.Contains(module.Name))
+				{
+					modules.Add(module);
+				}
+			}
+
+			// This can actually diff because Engine Shared Pch in editor builds can reach UnrealEd includes which is bad
+			if (modules.Count != moduleNames.Count)
+			{
+				throw new Exception($"This should never happen. {PCHHeaderFile} is found with different number of module dependencies");
+			}
+
+			return modules;
+
+			//return GetSharedPCHModuleDependenciesUncached(CompileEnvironment, PCHHeaderFile, Logger, ModuleToIncludePathsOnlyFlag);
+		}
+
+		private HashSet<UEBuildModule> GetSharedPCHModuleDependenciesUncached(CppCompileEnvironment CompileEnvironment, FileItem PCHHeaderFile, ILogger Logger, Dictionary<UEBuildModule, bool> ModuleToIncludePathsOnlyFlag)
+		{
 			HashSet<UEBuildModule> AllModuleDeps = ModuleToIncludePathsOnlyFlag.Keys.ToHashSet();
 			HashSet<UEBuildModule> ModuleDeps = new HashSet<UEBuildModule>(AllModuleDeps) { this };
 
-			HashSet<string> VisitedIncludes = new();
+			using var rented = HashSetPool<string>.Rent();
+			HashSet<string> VisitedIncludes = rented.Value;
+
 			HashSet<FileItem> FoundIncludeFileItems = new();
-			HashSet<string> IncludesNotFound = new();
+
+			using var rented3 = HashSetPool<string>.Rent();
+			HashSet<string> IncludesNotFound = rented3.Value;
 
 
 			// Go through user directories and create bloom filter for files and dirs so we can early out quickly
@@ -1360,7 +1394,11 @@ namespace UnrealBuildTool
 				UserIncludePaths.Add((dirItem, dirFilter, fileFilter));
 			}
 
-			Stack<(FileItem, List<string>)> stack = new ();
+			// First, traverse _all_ includes we can reach from the pch header file and store them in FoundIncludeFileItems
+			// This can be cached per-platform since we ignore things like ifdef DEBUG, SHIPPING etc.
+			// Reason we want to cache this is because
+
+			Stack<(FileItem, List<string>)> stack = new();
 
 			stack.Push((PCHHeaderFile, CompileEnvironment.MetadataCache.GetHeaderIncludes(PCHHeaderFile)));
 
@@ -1437,19 +1475,19 @@ namespace UnrealBuildTool
 							if (!TransformedHeaderInclude.Contains('.'))
 							{
 #if DEBUG
-							Logger.LogDebug("{0} SharedPCH - Skipping '{1}' found in '{2}' because it doesn't appear to be a module header.", Name, TransformedHeaderInclude, FileToSearch.Location);
+								Logger.LogDebug("{0} SharedPCH - Skipping '{1}' found in '{2}' because it doesn't appear to be a module header.", Name, TransformedHeaderInclude, FileToSearch.Location);
 #endif
 							}
 							else if (TransformedHeaderInclude.EndsWith(".generated.h"))
 							{
 #if DEBUG
-							Logger.LogDebug("{0} SharedPCH - Skipping '{1}' found in '{2}' because it appears to be a generated UHT header.", Name, TransformedHeaderInclude, FileToSearch.Location);
+								Logger.LogDebug("{0} SharedPCH - Skipping '{1}' found in '{2}' because it appears to be a generated UHT header.", Name, TransformedHeaderInclude, FileToSearch.Location);
 #endif
 							}
 							else
 							{
 #if DEBUG
-							Logger.LogDebug("{0} SharedPCH - Could not find include directory for '{1}' found in '{2}'.", Name, TransformedHeaderInclude, FileToSearch.Location);
+								Logger.LogDebug("{0} SharedPCH - Could not find include directory for '{1}' found in '{2}'.", Name, TransformedHeaderInclude, FileToSearch.Location);
 #endif
 								IncludesNotFound.Add(TransformedHeaderInclude);
 							}
@@ -1679,7 +1717,8 @@ namespace UnrealBuildTool
 					if (Rules.Target.bChainPCHs)
 					{
 						// Find all the dependencies of this module
-						HashSet<UEBuildModule> ReferencedModules = new HashSet<UEBuildModule>();
+						using var rented = HashSetPool<UEBuildModule>.Rent();
+						HashSet<UEBuildModule> ReferencedModules = rented.Value;
 						Template.Module.GetAllDependencyModules(new List<UEBuildModule>(), ReferencedModules, bIncludeDynamicallyLoaded: false, bForceCircular: false, bOnlyDirectDependencies: false);
 
 						PrecompiledHeaderTemplate? ParentTemplate = CompileEnvironment.SharedPCHs
@@ -2125,7 +2164,10 @@ namespace UnrealBuildTool
 				if (CompileEnvironment.bHasPrecompiledHeader == false && CompileEnvironment.SharedPCHs.Count > 0 && !CompileEnvironment.bIsBuildingLibrary && Rules.PCHUsage != ModuleRules.PCHUsageMode.NoSharedPCHs)
 				{
 					// Find the first shared PCH module we can use that doesn't reference this module
-					HashSet<UEBuildModule> AllDependencies = GetAllDependencyModulesForPCH(true);
+					using var rented = HashSetPool<UEBuildModule>.Rent();
+					HashSet<UEBuildModule> AllDependencies = rented.Value;
+					GetAllDependencyModulesForPCH(AllDependencies, true);
+
 					PrecompiledHeaderTemplate? Template = CompileEnvironment.SharedPCHs.FirstOrDefault(x => !x.ModuleDependencies.Contains(this) && AllDependencies.Contains(x.Module) && x.IsValidFor(CompileEnvironment));
 
 					if (Template != null)
@@ -2275,7 +2317,8 @@ namespace UnrealBuildTool
 				Writer.Append("#undef ").Append(Undefinition).AppendLine();
 			}
 
-			HashSet<string> Processed = new();
+			using var rented = HashSetPool<string>.Rent();
+			HashSet<string> Processed = rented.Value;
 			foreach (string Definition in Definitions.Where(x => Processed.Add(x)))
 			{
 				int EqualsIdx = Definition.IndexOf('=', StringComparison.Ordinal);
@@ -2516,11 +2559,12 @@ namespace UnrealBuildTool
 			}
 		}
 
-		public CppCompileEnvironment CreateCompileEnvironmentForIntellisense(ReadOnlyTargetRules Target, CppCompileEnvironment BaseCompileEnvironment, ILogger Logger)
+		public CppCompileEnvironment CreateCompileEnvironmentForIntellisense(ReadOnlyTargetRules Target, CppCompileEnvironment BaseCompileEnvironment, ILogger Logger, NullActionGraphBuilder? builder = null)
 		{
+			builder ??= new NullActionGraphBuilder(Logger);
 			CppCompileEnvironment CompileEnvironment = CreateModuleCompileEnvironment(Target, BaseCompileEnvironment, Logger);
-			CompileEnvironment = SetupPrecompiledHeaders(Target, null, CompileEnvironment, new List<FileItem>(), new NullActionGraphBuilder(Logger), Logger);
-			CreateHeaderForDefinitions(CompileEnvironment, IntermediateDirectory, null, new NullActionGraphBuilder(Logger));
+			CompileEnvironment = SetupPrecompiledHeaders(Target, null, CompileEnvironment, new List<FileItem>(), builder, Logger);
+			CreateHeaderForDefinitions(CompileEnvironment, IntermediateDirectory, null, builder);
 			return CompileEnvironment;
 		}
 
@@ -2680,11 +2724,13 @@ namespace UnrealBuildTool
 
 			Result.Definitions.AddRange(EngineIncludeOrderHelper.GetDeprecationDefines(Rules.IncludeOrderVersion));
 
+			DirectoryReference? ProjectDirectory = Target.ProjectFile?.Directory;
+
 			// If module intermediate files are under the project folder we can set UE_PROJECT_NAME and friends.
 			// This makes it possible to create modular builds where you don't have to provide the project name on cmd line
-			if (Target.ProjectFile != null && IntermediateDirectory.IsUnderDirectory(Target.ProjectFile.Directory))
+			if (ProjectDirectory != null && IntermediateDirectory.IsUnderDirectory(ProjectDirectory))
 			{
-				string ProjectName = Target.ProjectFile.GetFileNameWithoutExtension();
+				string ProjectName = Target.ProjectFile!.GetFileNameWithoutExtension();
 				Result.Definitions.Add(String.Format("UE_PROJECT_NAME={0}", ProjectName));
 				Result.Definitions.Add(String.Format("UE_TARGET_NAME={0}", Target.OriginalName));
 			}
@@ -2762,9 +2808,9 @@ namespace UnrealBuildTool
 			}
 
 			// Project root path should only be added if the module needs it, to prevent violating shared environment actions
-			if (Target.ProjectFile != null && (!Result.bUseSharedBuildEnvironment || Result.AllIncludePath.Any(x => x.IsUnderDirectory(Target.ProjectFile.Directory))))
+			if (ProjectDirectory != null && (!Result.bUseSharedBuildEnvironment || Result.AllIncludePath.Any(x => x.IsUnderDirectory(ProjectDirectory))))
 			{
-				Result.RootPaths[CppRootPathFolder.Project] = Target.ProjectFile.Directory;
+				Result.RootPaths[CppRootPathFolder.Project] = ProjectDirectory;
 			}
 
 			return Result;
diff --git a/Engine/Source/Programs/UnrealBuildTool/Configuration/UEBuildPlatform.cs b/Engine/Source/Programs/UnrealBuildTool/Configuration/UEBuildPlatform.cs
--- a/Engine/Source/Programs/UnrealBuildTool/Configuration/UEBuildPlatform.cs
+++ b/Engine/Source/Programs/UnrealBuildTool/Configuration/UEBuildPlatform.cs
@@ -693,7 +693,10 @@ namespace UnrealBuildTool
 		/// <returns>List of platform groups that this platform is a member of</returns>
 		public static IEnumerable<UnrealPlatformGroup> GetPlatformGroups(UnrealTargetPlatform Platform)
 		{
-			return PlatformGroupDictionary.Where(x => x.Value.Contains(Platform)).Select(x => x.Key);
+			// Avoid allocations by using yield (one allocation for enumerable still created)
+			foreach (var kvp in PlatformGroupDictionary)
+				if (kvp.Value.Contains(Platform))
+					yield return kvp.Key;
 		}
 
 		/// <summary>
diff --git a/Engine/Source/Programs/UnrealBuildTool/Configuration/UEBuildTarget.cs b/Engine/Source/Programs/UnrealBuildTool/Configuration/UEBuildTarget.cs
--- a/Engine/Source/Programs/UnrealBuildTool/Configuration/UEBuildTarget.cs
+++ b/Engine/Source/Programs/UnrealBuildTool/Configuration/UEBuildTarget.cs
@@ -2551,7 +2551,7 @@ namespace UnrealBuildTool
 		/// Creates a global compile environment suitable for generating project files.
 		/// </summary>
 		/// <returns>New compile environment</returns>
-		public CppCompileEnvironment CreateCompileEnvironmentForProjectFiles(ILogger Logger)
+		public CppCompileEnvironment CreateCompileEnvironmentForProjectFiles(ILogger Logger, NullActionGraphBuilder? builder = null)
 		{
 			CppConfiguration CppConfiguration = GetCppConfiguration(Configuration);
 
@@ -2565,15 +2565,8 @@ namespace UnrealBuildTool
 
 			FindSharedPCHs(Binaries, GlobalCompileEnvironment, Logger);
 
-			GlobalSingleInstanceMutex? mutex = null;
-			if (Rules.BuildEnvironment == TargetBuildEnvironment.Shared)
-			{
-				string MutexName = GlobalSingleInstanceMutex.GetUniqueMutexForPath("UnrealBuildTool_CreateCompileEnvironmentForProjectFiles", Unreal.EngineDirectory);
-				mutex = new GlobalSingleInstanceMutex(MutexName, true);
-			}
-
-			CreateSharedPCHInstances(Rules, TargetToolChain, Binaries, GlobalCompileEnvironment, new NullActionGraphBuilder(Logger), Logger);
-			mutex?.Dispose();
+			builder ??= new NullActionGraphBuilder(Logger);
+			CreateSharedPCHInstances(Rules, TargetToolChain, Binaries, GlobalCompileEnvironment, builder, Logger);
 
 			return GlobalCompileEnvironment;
 		}
@@ -4350,8 +4343,6 @@ namespace UnrealBuildTool
 		/// </summary>
 		public void PreBuildSetup(ILogger Logger)
 		{
-			using var _ = Timeline.ScopeEvent("PreBuildSetup");
-
 			// Describe what's being built.
 			Logger.LogDebug("Building {AppName} - {TargetName} - {Platform} - {Configuration}", AppName, TargetName, Platform, Configuration);
 
@@ -4380,11 +4371,12 @@ namespace UnrealBuildTool
 
 			// Create all the modules referenced by the existing binaries
 			UEBuildModule CreateModuleDelegate(string Name, string RefChain) => FindOrCreateModuleByName(Name, RefChain, Logger);
+			StringBuilder ReferenceChain = new("Target");
 			foreach (UEBuildBinary Binary in Binaries)
 			{
 				foreach (UEBuildModule Module in Binary.Modules)
 				{
-					Module.RecursivelyCreateModules(CreateModuleDelegate, "Target", Logger);
+					Module.RecursivelyCreateModules(CreateModuleDelegate, ReferenceChain, Logger);
 				}
 			}
 
@@ -4679,7 +4671,9 @@ namespace UnrealBuildTool
 			SortedDictionary<string, (int, PrecompiledHeaderTemplate)> SharedPCHModuleToPriority = new ();
 			Parallel2.ForEach(SharedPCHModules, SharedPCHModule =>
 			{
-				HashSet<UEBuildModule> Dependencies = SharedPCHModule.GetAllDependencyModulesForPCH(false);
+				using var rented = HashSetPool<UEBuildModule>.Rent();
+				HashSet<UEBuildModule> Dependencies = rented.Value;
+				SharedPCHModule.GetAllDependencyModulesForPCH(Dependencies, false);
 				Dependencies.IntersectWith(SharedPCHModules);
 
 				// TODO This is bad.. unfortunately does circular dependency logic not working when getting dependency modules
@@ -4740,38 +4734,42 @@ namespace UnrealBuildTool
 				ModuleInfos[SharedPCHIndex] = new ConcurrentBag<Tuple<UEBuildModuleCPP, CppCompileEnvironment>>();
 			}
 
+			// Sometimes it is one binary with lots of modules, sometimes it is lots of binaries with one module each
 			Parallel2.ForEach(OriginalBinaries, Binary =>
 			{
 				CppCompileEnvironment BinaryCompileEnvironment = Binary.CreateBinaryCompileEnvironment(GlobalCompileEnvironment);
 
-				foreach (UEBuildModuleCPP Module in Binary.Modules.OfType<UEBuildModuleCPP>())
+				Parallel2.ForEach(Binary.Modules.OfType<UEBuildModuleCPP>(), Module =>
 				{
 					CppCompileEnvironment ModuleCompileEnvironment = Module.CreateModuleCompileEnvironment(Target, BinaryCompileEnvironment, Logger);
 
 					// Is it using a private PCH?
 					if (Module.Rules.PrivatePCHHeaderFile != null && (Module.Rules.PCHUsage == ModuleRules.PCHUsageMode.NoSharedPCHs || Module.Rules.PCHUsage == ModuleRules.PCHUsageMode.UseExplicitOrSharedPCHs))
 					{
-						continue;
+						return;
 					}
 
-					// Is it using a shared PCH?
-					if (Module.Rules.PCHUsage != ModuleRules.PCHUsageMode.NoPCHs)
+					// Is it using no PCH?
+					if (Module.Rules.PCHUsage == ModuleRules.PCHUsageMode.NoPCHs)
 					{
-						if (!ModuleCompileEnvironment.bIsBuildingLibrary && Module.Rules.PCHUsage != ModuleRules.PCHUsageMode.NoSharedPCHs)
-						{
-							// Find all the dependencies of this module
-							HashSet<UEBuildModule> ReferencedModules = new HashSet<UEBuildModule>();
-							Module.GetAllDependencyModules(new List<UEBuildModule>(), ReferencedModules, bIncludeDynamicallyLoaded: false, bForceCircular: false, bOnlyDirectDependencies: true);
+						return;
+					}
 
-							// Find the first shared PCH module we can use
-							int SharedPCHIndex = ModuleCompileEnvironment.SharedPCHs.FindIndex(x => ReferencedModules.Contains(x.Module));
-							if (SharedPCHIndex != -1 && GlobalCompileEnvironment.SharedPCHs[SharedPCHIndex].IsValidFor(ModuleCompileEnvironment))
-							{
-								ModuleInfos[SharedPCHIndex].Add(new Tuple<UEBuildModuleCPP, CppCompileEnvironment>(Module, ModuleCompileEnvironment));
-							}
+					if (!ModuleCompileEnvironment.bIsBuildingLibrary && Module.Rules.PCHUsage != ModuleRules.PCHUsageMode.NoSharedPCHs)
+					{
+						// Find all the dependencies of this module
+						using var rented = HashSetPool<UEBuildModule>.Rent();
+						HashSet<UEBuildModule> ReferencedModules = rented.Value;
+						Module.GetAllDependencyModules(new List<UEBuildModule>(), ReferencedModules, bIncludeDynamicallyLoaded: false, bForceCircular: false, bOnlyDirectDependencies: true);
+
+						// Find the first shared PCH module we can use
+						int SharedPCHIndex = ModuleCompileEnvironment.SharedPCHs.FindIndex(x => ReferencedModules.Contains(x.Module));
+						if (SharedPCHIndex != -1 && GlobalCompileEnvironment.SharedPCHs[SharedPCHIndex].IsValidFor(ModuleCompileEnvironment))
+						{
+							ModuleInfos[SharedPCHIndex].Add(new Tuple<UEBuildModuleCPP, CppCompileEnvironment>(Module, ModuleCompileEnvironment));
 						}
 					}
-				}
+				});
 			});
 
 			// Create the PCH instances
@@ -4866,7 +4864,9 @@ namespace UnrealBuildTool
 		protected void AddAllValidModulesToTarget(ILogger Logger)
 		{
 			// Find all the modules that are part of the target
-			HashSet<string> ValidModuleNames = new HashSet<string>();
+			using var rented = HashSetPool<string>.Rent();
+			HashSet<string> ValidModuleNames = rented.Value;
+
 			foreach (UEBuildModuleCPP Module in Modules.Values.OfType<UEBuildModuleCPP>())
 			{
 				if (Module.Binary != null)
@@ -4879,7 +4879,8 @@ namespace UnrealBuildTool
 			IReadOnlySet<string> ExcludeFolders = UEBuildPlatform.GetBuildPlatform(Platform).GetExcludedFolderNames();
 
 			// Set of module names to build
-			HashSet<string> FilteredModuleNames = new HashSet<string>();
+			using var rented2 = HashSetPool<string>.Rent();
+			HashSet<string> FilteredModuleNames = rented2.Value;
 
 			// Only add engine modules for non-program targets. Programs only compile allowed modules through plugins.
 			if (TargetType != TargetType.Program)
@@ -5045,16 +5046,17 @@ namespace UnrealBuildTool
 			// Now create all the precompiled modules, making sure they don't reference anything that's not in the precompiled set
 
 			// Gather the set of all modules traversed while processing ValidModuleNames, to ensure we have any extra modules found during RecursivelyCreateModules()
-			HashSet<UEBuildModule> AllModules = new HashSet<UEBuildModule>();
+			using var rented3 = HashSetPool<UEBuildModule>.Rent();
+			HashSet<UEBuildModule> AllModules = rented3.Value;
 
 			foreach (string ModuleName in ValidModuleNames)
 			{
-				string PrecompileReferenceChain = "allmodules option";
+				StringBuilder PrecompileReferenceChain = new("allmodules option");
 				if (ModulePluginSource.TryGetValue(ModuleName, out PluginInfo? value))
 				{
-					PrecompileReferenceChain = $"{PrecompileReferenceChain} -> {value.File.GetFileName()}";
+					PrecompileReferenceChain.Append(" -> ").Append(value.File.GetFileName());
 				}
-				UEBuildModule Module = FindOrCreateModuleByName(ModuleName, PrecompileReferenceChain, Logger);
+				UEBuildModule Module = FindOrCreateModuleByName(ModuleName, PrecompileReferenceChain.ToString(), Logger);
 				AllModules.Add(Module);
 				Module.RecursivelyCreateModules(
 					(string ModuleName, string ReferenceChain) =>
@@ -5766,9 +5768,10 @@ namespace UnrealBuildTool
 				}
 
 				// Create the dependencies set
-				HashSet<UEBuildPlugin> Dependencies = new HashSet<UEBuildPlugin>();
+				HashSet<UEBuildPlugin> Dependencies;
 				if (Info.Descriptor.Plugins != null)
 				{
+					Dependencies = new(Info.Descriptor.Plugins.Count + 1);
 					foreach (PluginReferenceDescriptor NextReference in Info.Descriptor.Plugins)
 					{
 						// Ignore any plugin dependency which is programmatically filtered out by the target
@@ -5792,6 +5795,11 @@ namespace UnrealBuildTool
 						}
 					}
 				}
+				else
+				{
+					Dependencies = [];
+				}
+					
 				Instance.Dependencies = Dependencies;
 
 				// Stage the descriptor if the plugin contains content or Verse code
diff --git a/Engine/Source/Programs/UnrealBuildTool/ProjectFiles/AliasedFile.cs b/Engine/Source/Programs/UnrealBuildTool/ProjectFiles/AliasedFile.cs
--- a/Engine/Source/Programs/UnrealBuildTool/ProjectFiles/AliasedFile.cs
+++ b/Engine/Source/Programs/UnrealBuildTool/ProjectFiles/AliasedFile.cs
@@ -1,13 +1,14 @@
 ﻿// Copyright Epic Games, Inc. All Rights Reserved.
 
 using EpicGames.Core;
+using System;
 
 namespace UnrealBuildTool
 {
 	/// <summary>
 	/// Class that stores info about aliased file.
 	/// </summary>
-	struct AliasedFile
+	struct AliasedFile : IComparable<AliasedFile>
 	{
 		public AliasedFile(FileReference Location, string FileSystemPath, string ProjectPath)
 		{
@@ -16,6 +17,9 @@ namespace UnrealBuildTool
 			this.ProjectPath = ProjectPath;
 		}
 
+		public int CompareTo(AliasedFile other) => string.CompareOrdinal(Location.FullName, other.Location.FullName);
+
+
 		// Full location on disk.
 		public readonly FileReference Location;
 
diff --git a/Engine/Source/Programs/UnrealBuildTool/ProjectFiles/ProjectFile.cs b/Engine/Source/Programs/UnrealBuildTool/ProjectFiles/ProjectFile.cs
--- a/Engine/Source/Programs/UnrealBuildTool/ProjectFiles/ProjectFile.cs
+++ b/Engine/Source/Programs/UnrealBuildTool/ProjectFiles/ProjectFile.cs
@@ -466,6 +466,22 @@ namespace UnrealBuildTool
 			}
 		}
 
+		/// <summary>
+		/// Takes the given path and tries to rebase it relative to the project.
+		/// </summary>
+		public static void NormalizeProjectPath(ref StackStringBuilder outBuilder, FileSystemReference InputPath)
+		{
+			// Try to make it relative to the solution directory.
+			if (InputPath.IsUnderDirectory(ProjectFileGenerator.PrimaryProjectPath))
+			{
+				InputPath.MakeRelativeTo(ref outBuilder, ProjectFileGenerator.IntermediateProjectFilesPath);
+			}
+			else
+			{
+				outBuilder.Append(InputPath.FullName);
+			}
+		}
+
 		/// <summary>
 		/// Takes the given path, normalizes it, and quotes it if necessary.
 		/// </summary>
diff --git a/Engine/Source/Programs/UnrealBuildTool/ProjectFiles/ProjectFileGenerator.cs b/Engine/Source/Programs/UnrealBuildTool/ProjectFiles/ProjectFileGenerator.cs
--- a/Engine/Source/Programs/UnrealBuildTool/ProjectFiles/ProjectFileGenerator.cs
+++ b/Engine/Source/Programs/UnrealBuildTool/ProjectFiles/ProjectFileGenerator.cs
@@ -460,6 +460,8 @@ namespace UnrealBuildTool
 
 		readonly ProjectDescriptor? OnlyGameProjectDescriptor;
 
+		private FileHashCache? _fileHashCache;
+
 		/// <summary>
 		/// Cache for GlobalCompileEnvironment, to reduce GPF time, currently only useful on Mac
 		/// </summary>
@@ -1068,11 +1070,14 @@ namespace UnrealBuildTool
 				}
 			}
 
+			_fileHashCache = new(warnOnNoPartition:false, backupModifiedFiles:false);
+			_fileHashCache.MountAll(IntermediateProjectFilesPath, Logger);
+
 			// Write out the name of the primary project file and it's path, so the runtime knows to use it
 			FileReference PrimaryProjectNameLocation = FileReference.Combine(Unreal.EngineDirectory, "Intermediate", "ProjectFiles", "PrimaryProjectName.txt");
-			Utils.WriteFileIfChanged(PrimaryProjectNameLocation, PrimaryProjectName, Logger);
+			_fileHashCache.WriteFileIfChanged(FileItem.GetItemByFileReference(PrimaryProjectNameLocation), PrimaryProjectName, Logger, false);
 			FileReference PrimaryProjectPathLocation = FileReference.Combine(Unreal.EngineDirectory, "Intermediate", "ProjectFiles", "PrimaryProjectPath.txt");
-			Utils.WriteFileIfChanged(PrimaryProjectPathLocation, DirectoryReference.Combine(PrimaryProjectPath, PrimaryProjectName).FullName.Replace("\\", "/", StringComparison.InvariantCultureIgnoreCase), Logger);
+			_fileHashCache.WriteFileIfChanged(FileItem.GetItemByFileReference(PrimaryProjectPathLocation), DirectoryReference.Combine(PrimaryProjectPath, PrimaryProjectName).FullName.Replace("\\", "/", StringComparison.InvariantCultureIgnoreCase), Logger, false);
 
 			bool bCleanProjectFiles = Arguments.Any(x => x.Equals("-CleanProjects", StringComparison.InvariantCultureIgnoreCase));
 			if (bCleanProjectFiles)
@@ -1427,8 +1432,8 @@ namespace UnrealBuildTool
 
 				// Figure out which targets we need about IntelliSense for.  We only need to worry about targets for projects
 				// that we're actually generating in this session.
-				List<Tuple<ProjectFile, ProjectTarget>> IntelliSenseTargetFiles = new List<Tuple<ProjectFile, ProjectTarget>>();
-				List<Tuple<ProjectFile, ProjectTarget>> AllProjectTargetFiles = new List<Tuple<ProjectFile, ProjectTarget>>();
+				List<(ProjectFile, ProjectTarget)> IntelliSenseTargetFiles = [];
+				List<(ProjectFile, ProjectTarget)> AllProjectTargetFiles = [];
 				{
 					// Engine targets
 					foreach (ProjectFile EngineProject in EngineProjects)
@@ -1440,9 +1445,9 @@ namespace UnrealBuildTool
 							// for good quality IntelliSense.  For example, we want WITH_EDITORONLY_DATA=1, so using the editor targets works well.
 							if (bMakeProjectPerTarget || EngineTarget.TargetRules!.Type == TargetType.Editor) // @todo <-- try removing first term
 							{
-								IntelliSenseTargetFiles.Add(Tuple.Create(EngineProject, EngineTarget));
+								IntelliSenseTargetFiles.Add((EngineProject, EngineTarget));
 							}
-							AllProjectTargetFiles.Add(Tuple.Create(EngineProject, EngineTarget));
+							AllProjectTargetFiles.Add((EngineProject, EngineTarget));
 						}
 					}
 
@@ -1453,9 +1458,9 @@ namespace UnrealBuildTool
 						{
 							if (ProjectTarget.TargetFilePath != null)
 							{
-								IntelliSenseTargetFiles.Add(Tuple.Create(ProgramProject, ProjectTarget));
+								IntelliSenseTargetFiles.Add((ProgramProject, ProjectTarget));
 							}
-							AllProjectTargetFiles.Add(Tuple.Create(ProgramProject, ProjectTarget));
+							AllProjectTargetFiles.Add((ProgramProject, ProjectTarget));
 						}
 					}
 
@@ -1468,9 +1473,9 @@ namespace UnrealBuildTool
 							// for good quality IntelliSense.  For example, we want WITH_EDITORONLY_DATA=1, so using the editor targets works well.
 							if (bMakeProjectPerTarget || ProjectTarget.TargetRules!.Type == TargetType.Editor)
 							{
-								IntelliSenseTargetFiles.Add(Tuple.Create(GameProject, ProjectTarget));
+								IntelliSenseTargetFiles.Add((GameProject, ProjectTarget));
 							}
-							AllProjectTargetFiles.Add(Tuple.Create(GameProject, ProjectTarget));
+							AllProjectTargetFiles.Add((GameProject, ProjectTarget));
 						}
 					}
 				}
@@ -1490,6 +1495,10 @@ namespace UnrealBuildTool
 				{
 					WriteProjectFiles(PlatformProjectGenerators, Logger);
 				}
+
+				_fileHashCache.SaveAll();
+				_fileHashCache = null;
+
 				Logger.LogDebug("Project generation complete ({GeneratedProjectFilesCount} generated, {OtherProjectFilesCount} imported)", GeneratedProjectFiles.Count, OtherProjectFiles.Count);
 			}
 			return bSuccess;
@@ -1533,12 +1542,7 @@ namespace UnrealBuildTool
 				{
 					using var t = Timeline.ScopeEvent(Item.Key?.GetFileName()??"Unknown");
 					QueryTargetsMode.WriteTargetInfo(Item.Key, Item.Value, QueryTargetsMode.GetDefaultOutputFile(Item.Key), new CommandLineArguments(GetTargetArguments(Arguments)), Logger);
-
-					lock (Progress)
-					{
-						Interlocked.Increment(ref TasksFinished);
-						Progress.Write(TasksFinished, ProjectFileToRules.Count);
-					}
+					Progress.Write(Interlocked.Increment(ref TasksFinished), ProjectFileToRules.Count);
 				});
 			}
 		}
@@ -2329,18 +2333,18 @@ namespace UnrealBuildTool
 		/// <param name="Targets">Targets to build for intellisense</param>
 		/// <param name="Logger">Logger for output</param>
 		/// <return>Whether the process was successful or not</return>
-		private void GenerateIntelliSenseData(string[] Arguments, List<Tuple<ProjectFile, ProjectTarget>> Targets, ILogger Logger)
+		private void GenerateIntelliSenseData(string[] Arguments, List<(ProjectFile, ProjectTarget)> Targets, ILogger Logger)
 		{
 			if (ShouldGenerateIntelliSenseData() && Targets.Any())
 			{
 				// Ignore projects for platforms we can't build for
-				List<Tuple<ProjectFile, ProjectTarget>> SupportedTargets = Targets.Where(x => x.Item2.SupportedPlatforms.Any(x => GetIntelliSensePlatforms().Contains(x))).ToList();
+				List<(ProjectFile, ProjectTarget)> SupportedTargets = Targets.Where(x => x.Item2.SupportedPlatforms.Any(x => GetIntelliSensePlatforms().Contains(x))).ToList();
 
 				using ITimelineEvent CompileTimer = Timeline.ScopeEvent("Compiling Rules Assemblies");
 				using (ProgressWriter Progress = new ProgressWriter("Compiling Rules Assemblies", true, Logger))
 				{
 					int TasksFinished = 0;
-					foreach (Tuple<ProjectFile, ProjectTarget> Item in SupportedTargets)
+					foreach ((ProjectFile, ProjectTarget) Item in SupportedTargets)
 					{
 						ProjectTarget CurTarget = Item.Item2;
 
@@ -2363,11 +2367,7 @@ namespace UnrealBuildTool
 						// Create the target rules assembly 
 						RulesCompiler.CreateTargetRulesAssembly(CurTarget.UnrealProjectFilePath, CurTarget.Name, false, false, bUsePrecompiled, TargetDesc.ForeignPlugin, TargetDesc.bBuildPluginAsLocal, Logger);
 
-						lock (Progress)
-						{
-							Interlocked.Increment(ref TasksFinished);
-							Progress.Write(TasksFinished, SupportedTargets.Count);
-						}
+						Progress.Write(Interlocked.Increment(ref TasksFinished), SupportedTargets.Count);
 					}
 				}
 				CompileTimer.Finish();
@@ -2382,10 +2382,9 @@ namespace UnrealBuildTool
 					int NumTasks = NumTargets;
 					int TasksFinished = 0;
 
-					Parallel.ForEach(SupportedTargets, Item =>
+					void CreateBuildTarget(ProjectFile TargetProjectFile, ProjectTarget CurTarget)
 					{
-						ProjectFile TargetProjectFile = Item.Item1;
-						ProjectTarget CurTarget = Item.Item2;
+						using var _ = Timeline.ScopeEvent(CurTarget.Name);
 
 						List<string> NewArguments = new List<string>(Arguments.Length + 4);
 						if (CurTarget.TargetRules!.Type != TargetType.Program)
@@ -2409,7 +2408,8 @@ namespace UnrealBuildTool
 								UEBuildTarget Target = UEBuildTarget.Create(TargetDesc, true, false, bUsePrecompiled, UnrealIntermediateEnvironment.GenerateProjectFiles, Logger);
 								BuildTargets.Add(Target);
 								BuildTargetToProjectFile.TryAdd(Target, TargetProjectFile);
-								BuildTargetToCompileEnvironmentCacheKey.TryAdd(Target, MakeGlobalCompileEnvironmentCacheKey(CurTarget.UnrealProjectFilePath, CurTarget.Name, IntellisensePlatform, UnrealTargetConfiguration.Development, DefaultArchitecture));
+								string cacheKey = MakeGlobalCompileEnvironmentCacheKey(CurTarget.UnrealProjectFilePath, CurTarget.Name, IntellisensePlatform, UnrealTargetConfiguration.Development, DefaultArchitecture);
+								BuildTargetToCompileEnvironmentCacheKey.TryAdd(Target, cacheKey);
 							}
 						}
 						catch (Exception Ex)
@@ -2417,12 +2417,10 @@ namespace UnrealBuildTool
 							Logger.LogWarning("Exception while creating build target for {Target}: {Ex}", CurTarget.Name, Ex.ToString());
 						}
 						
-						lock (Progress)
-						{
-							Interlocked.Increment(ref TasksFinished);
-							Progress.Write(TasksFinished, SupportedTargets.Count);
-						}
-					});
+						Progress.Write(Interlocked.Increment(ref TasksFinished), SupportedTargets.Count);
+					}
+
+					Parallel2.ForEach(SupportedTargets, Item => CreateBuildTarget(Item.Item1, Item.Item2));
 				}
 				BuildTargetsTimer.Finish();
 
@@ -2437,8 +2435,13 @@ namespace UnrealBuildTool
 						TotalTasks += BuildTargets.SelectMany(x => x.Binaries).SelectMany(x => x.Modules.OfType<UEBuildModuleCPP>()).Count();
 					}
 
+					NullActionGraphBuilder builder = new(Logger, _fileHashCache);
+
+					ConcurrentDictionary<ProjectFile, ConcurrentBag<(UEBuildModuleCPP, CppCompileEnvironment)>> toAdd = [];
+
 					void Gen(UEBuildTarget Target)
 					{
+						using var _ = Timeline.ScopeEvent(Target.TargetName);
 						AddTargetForIntellisense(Target, Logger);
 
 						// If the project generator just cares about the result of UEBuildTarget.Create, skip generating the compile environments.
@@ -2452,65 +2455,49 @@ namespace UnrealBuildTool
 							}
 							else
 							{
-								GlobalCompileEnvironment = Target.CreateCompileEnvironmentForProjectFiles(Logger);
+								GlobalCompileEnvironment = Target.CreateCompileEnvironmentForProjectFiles(Logger, builder);
 							}
-							
-							foreach (UEBuildBinary Binary in Target.Binaries)
+
+							void ProcessBinary(UEBuildBinary Binary)
 							{
-								CppCompileEnvironment BinaryCompileEnvironment = Binary.CreateBinaryCompileEnvironment(GlobalCompileEnvironment);
+								CppCompileEnvironment? BinaryCompileEnvironment = null;
 								foreach (UEBuildModuleCPP Module in Binary.Modules.OfType<UEBuildModuleCPP>())
 								{
 									if (ModuleToEditorProjectFileMap.TryGetValue(Module.RulesFile, out ProjectFile? ProjectFileForIDE) && ProjectFileForIDE == TargetProjectFile)
 									{
-										CppCompileEnvironment ModuleCompileEnvironment = Module.CreateCompileEnvironmentForIntellisense(Target.Rules, BinaryCompileEnvironment, Logger);
-										lock (ProjectFileForIDE)
+										if (BinaryCompileEnvironment == null)
 										{
-											ProjectFileForIDE.AddModuleForIntelliSense(Module, ModuleCompileEnvironment);
+											BinaryCompileEnvironment = Binary.CreateBinaryCompileEnvironment(GlobalCompileEnvironment);
 										}
-									}
+										CppCompileEnvironment ModuleCompileEnvironment = Module.CreateCompileEnvironmentForIntellisense(Target.Rules, BinaryCompileEnvironment, Logger, builder);
 
-									lock (Progress)
-									{
-										Interlocked.Increment(ref TasksFinished);
-										Progress.Write(TasksFinished, TotalTasks);
+										var list = toAdd.GetOrAdd(ProjectFileForIDE, _ => new());
+										list.Add((Module, ModuleCompileEnvironment));
 									}
+									Progress.Write(Interlocked.Increment(ref TasksFinished), TotalTasks);
 								}
 							}
+							Parallel2.ForEach(Target.Binaries, Binary => ProcessBinary(Binary));
 						}
 
-						lock (Progress)
-						{
-							Interlocked.Increment(ref TasksFinished);
-							Progress.Write(TasksFinished, TotalTasks);
-						}
+						Progress.Write(Interlocked.Increment(ref TasksFinished), TotalTasks);
 					}
 
-					IEnumerable<UEBuildTarget> Shared = BuildTargets.Where(x => x.bUseSharedBuildEnvironment);
-					IEnumerable<UEBuildTarget> Unique = BuildTargets.Where(x => !x.bUseSharedBuildEnvironment);
-
-					IEnumerable<IGrouping<(UnrealTargetPlatform Platform, UnrealTargetConfiguration Configuration, TargetType TargetType), UEBuildTarget>> ParallelGroups = Shared.GroupBy(x => (x.Platform, x.Configuration, x.TargetType));
-
-					// Unique environments can be processed in parallel
-					// Shared environments can be processed in parallel, but each unique combination of Platform\Configuration\TargetType must be processed serially
-					// DebugGame can not be run in parallel with Development
-					// This could be made more async but it's complicated
-
-					Parallel.ForEach(Unique.Where(x => x.Configuration != UnrealTargetConfiguration.DebugGame), Target => Gen(Target));
-					Parallel.ForEach(Unique.Where(x => x.Configuration == UnrealTargetConfiguration.DebugGame), Target => Gen(Target));
-					
-					Parallel.ForEach(ParallelGroups.Where(x => x.Key.Configuration != UnrealTargetConfiguration.DebugGame), Group =>
+					Parallel2.ForEach(BuildTargets, Target =>
 					{
-						foreach (UEBuildTarget Target in Group)
-						{
-							Gen(Target);
-						}
+						Gen(Target);
 					});
 
-					Parallel.ForEach(ParallelGroups.Where(x => x.Key.Configuration == UnrealTargetConfiguration.DebugGame), Group =>
+					using var _ = Timeline.ScopeEvent("AddModules");
+					Parallel2.ForEach(toAdd, kv =>
 					{
-						foreach (UEBuildTarget Target in Group)
+						// We have to sort this array because there are multiple .build.cs that shares same Source folder (WireInterface among others)
+						// And the source files end up under different module if this list is not sorted.
+						var moduleAndEnv = kv.Value.ToArray();
+						Array.Sort(moduleAndEnv, static (a, b) => a.Item1.IntermediateDirectory.CompareTo(b.Item1.IntermediateDirectory));
+						foreach (var (module, compileEnv) in moduleAndEnv)
 						{
-							Gen(Target);
+							kv.Key.AddModuleForIntelliSense(module, compileEnv);
 						}
 					});
 				}
@@ -3683,10 +3670,7 @@ namespace UnrealBuildTool
 						return;
 					}
 
-					lock (Progress)
-					{
-						Progress.Write(++ProjectsFinished, TotalProjectFileCount);
-					}
+					Progress.Write(Interlocked.Increment(ref ProjectsFinished), TotalProjectFileCount);
 				});
 
 				if (ProjectCreationFailed)
@@ -3726,7 +3710,7 @@ namespace UnrealBuildTool
 		/// <param name="PlatformProjectGenerators"></param>
 		/// <param name="Targets"></param>
 		/// <param name="Logger"></param>
-		protected virtual void AddAdditionalNativeTargetInformation(PlatformProjectGeneratorCollection PlatformProjectGenerators, List<Tuple<ProjectFile, ProjectTarget>> Targets, ILogger Logger)
+		protected virtual void AddAdditionalNativeTargetInformation(PlatformProjectGeneratorCollection PlatformProjectGenerators, List<(ProjectFile, ProjectTarget)> Targets, ILogger Logger)
 		{
 		}
 
@@ -3740,6 +3724,12 @@ namespace UnrealBuildTool
 		/// <returns>True if the file was saved, or if it didn't need to be overwritten because the content was unchanged</returns>
 		public static bool WriteFileIfChanged(string FileName, string NewFileContents, ILogger Logger, Encoding? InEncoding = null)
 		{
+			if (Current != null && Current._fileHashCache != null)
+			{
+				Current._fileHashCache.WriteFileIfChanged(FileItem.GetItemByPath(FileName), NewFileContents, Logger, false);
+				return true;
+			}
+
 			// Check to see if the file already exists, and if so, load it up
 			string? LoadedFileContent = null;
 			bool FileAlreadyExists = File.Exists(FileName);
diff --git a/Engine/Source/Programs/UnrealBuildTool/ProjectFiles/VisualStudio/VCProject.cs b/Engine/Source/Programs/UnrealBuildTool/ProjectFiles/VisualStudio/VCProject.cs
--- a/Engine/Source/Programs/UnrealBuildTool/ProjectFiles/VisualStudio/VCProject.cs
+++ b/Engine/Source/Programs/UnrealBuildTool/ProjectFiles/VisualStudio/VCProject.cs
@@ -858,15 +858,19 @@ namespace UnrealBuildTool
 				FilteredList = CleanPaths.ToArray();
 			}
 
-			// The user might have specified Foo, \Foo, Foo\, or \Foo\ as excludes.
-			// Directory paths don't contain a trailing slash so add one to the string
-			// we'll compare to our list so we can Contains and EndsWith to catch:
-			// <Path>\Foo\Dir
-			// <Path>\Foo
-			string PathWithSeparator = InPath.FullName + Path.DirectorySeparatorChar;
-
 			if (FilteredList.Length > 0)
 			{
+				// The user might have specified Foo, \Foo, Foo\, or \Foo\ as excludes.
+				// Directory paths don't contain a trailing slash so add one to the string
+				// we'll compare to our list so we can Contains and EndsWith to catch:
+				// <Path>\Foo\Dir
+				// <Path>\Foo
+				Span<char> initial = stackalloc char[512];
+				using StackStringBuilder sb = new(initial);
+				sb.Append(InPath.FullName);
+				sb.Append(Path.DirectorySeparatorChar);
+				ReadOnlySpan<char> PathWithSeparator = sb.AsSpan();
+
 				foreach (string Entry in FilteredList)
 				{
 					if (PathWithSeparator.Contains(Entry, StringComparison.OrdinalIgnoreCase)
@@ -912,13 +916,13 @@ namespace UnrealBuildTool
 		/// <param name="Builder">String builder for the property value</param>
 		/// <param name="Collection">Collection of include paths</param>
 		/// <param name="IgnorePaths">Set of paths to ignore</param>
-		void AppendIncludePaths(StringBuilder Builder, IncludePathsCollection Collection, HashSet<DirectoryReference> IgnorePaths)
+		void AppendIncludePaths(ref StackStringBuilder Builder, IncludePathsCollection Collection, HashSet<DirectoryReference> IgnorePaths)
 		{
 			foreach (DirectoryReference IncludePath in Collection.AbsolutePaths)
 			{
 				if (!IgnorePaths.Contains(IncludePath) && !IncludePathIsFilteredOut(IncludePath))
 				{
-					Builder.Append(NormalizeProjectPath(IncludePath));
+					NormalizeProjectPath(ref Builder, IncludePath);
 					Builder.Append(';');
 				}
 			}
@@ -1233,13 +1237,31 @@ namespace UnrealBuildTool
 			}
 
 			// Gather source folder and file info
-			List<AliasedFile> LocalAliasedFiles = new(AliasedFiles);
+			ConcurrentBag<AliasedFile> LocalAliasedFiles = new(AliasedFiles);
 			ConcurrentDictionary<DirectoryReference, string?> DirectoryToPchFile = new();
 			ConcurrentDictionary<DirectoryReference, string> DirectoryToForceIncludePaths = new();
 			ConcurrentDictionary<DirectoryReference, string> DirectoryToIncludeSearchPaths = new();
 			{
-				foreach (SourceFile CurFile in SourceFiles)
+				Parallel2.ForEach(SourceFiles, CurFile =>
 				{
+					// get the filetype as represented to Visual Studio
+					VCFileType FileType = GetVCFileType(CurFile.Reference.FullName);
+
+					DirectoryReference FileSystemPathDir = CurFile.Reference.Directory;
+
+					// if the filetype is an include and its path is filtered out, skip it entirely (should we do this for any type of
+					// file? Possibly, but not today due to potential fallout)
+					if (CanVCFileTypeBeIncluded(FileType) && IncludePathIsFilteredOut(FileSystemPathDir))
+					{
+						return;
+					}
+
+					// Allow filtering of any type of file
+					if (FilePathIsFilteredOut(FileSystemPathDir))
+					{
+						return;
+					}
+
 					// We want all source file and directory paths in the project files to be relative to the project file's
 					// location on the disk.  Convert the path to be relative to the project file directory
 					string ProjectRelativeSourceFile = CurFile.Reference.MakeRelativeTo(ProjectFilePath.Directory);
@@ -1268,32 +1290,8 @@ namespace UnrealBuildTool
 						FilterRelativeSourceDirectory = FilterRelativeSourceDirectory.Substring(0, LastSeparatorIdx);
 					}
 
-					LocalAliasedFiles.Add(new AliasedFile(CurFile.Reference, ProjectRelativeSourceFile, FilterRelativeSourceDirectory));
-				}
-
-				LocalAliasedFiles.Sort((a, b) =>
-				{
-					return a.Location.CompareTo(b.Location);
-				});
-
-				Parallel2.ForEach(LocalAliasedFiles, LocalAliasedFile =>
-				{
-					// get the filetype as represented to Visual Studio
-					VCFileType FileType = GetVCFileType(LocalAliasedFile.FileSystemPath);
-					DirectoryReference FileSystemPathDir = new DirectoryReference(LocalAliasedFile.FileSystemPath);
-
-					// if the filetype is an include and its path is filtered out, skip it entirely (should we do this for any type of
-					// file? Possibly, but not today due to potential fallout)
-					if (CanVCFileTypeBeIncluded(FileType) && IncludePathIsFilteredOut(FileSystemPathDir))
-					{
-						return;
-					}
-
-					// Allow filtering of any type of file
-					if (FilePathIsFilteredOut(FileSystemPathDir))
-					{
-						return;
-					}
+					AliasedFile LocalAliasedFile = new(CurFile.Reference, ProjectRelativeSourceFile, FilterRelativeSourceDirectory);
+					LocalAliasedFiles.Add(LocalAliasedFile);
 
 					if (!CanVCFileTypeBeCompiled(FileType))
 					{
@@ -1345,10 +1343,13 @@ namespace UnrealBuildTool
 					{
 						DirectoryToIncludeSearchPaths.GetOrAdd(Directory, _ =>
 						{
-							StringBuilder Builder = new StringBuilder();
-							AppendIncludePaths(Builder, BuildEnvironment.UserIncludePaths, SharedIncludeSearchPathsSet);
-							AppendIncludePaths(Builder, BuildEnvironment.SystemIncludePaths, SharedIncludeSearchPathsSet);
-							return Builder.ToString();
+							Span<char> initial = stackalloc char[32*1024];
+							StackStringBuilder Builder = new(initial);
+							AppendIncludePaths(ref Builder, BuildEnvironment.UserIncludePaths, SharedIncludeSearchPathsSet);
+							AppendIncludePaths(ref Builder, BuildEnvironment.SystemIncludePaths, SharedIncludeSearchPathsSet);
+							string res = Builder.ToString();
+							Builder.Dispose();
+							return res;
 						});
 					}
 				});
@@ -1475,35 +1476,43 @@ namespace UnrealBuildTool
 
 					void AddProperties(IDictionary<DirectoryReference, string> DirectoryToStringDict, string CommonPropertyPrefix, string PropertyNamePrefix, string PropertyValuePrefix)
 					{
-						Dictionary<string, string> UpdatedValues = new();
+						Dictionary<string, string> UpdatedValues = new(StringComparer.Ordinal);
 						List<KeyValuePair<DirectoryReference, string>> KVPList = DirectoryToStringDict.ToList();
 						KVPList.SortBy(kvp => kvp.Key);
 
 						// Map each directory to a property
 						{
 							int PropertyIndex = 0;
+
+							StringBuilder PropertyNameBuilder = new(PropertyNamePrefix);
+							int PropertyNameBuilderLength = PropertyNameBuilder.Length;
+							string NextPropertyName = PropertyNameBuilder.ToString();
+
 							foreach (KeyValuePair<DirectoryReference, string> DirectoryKVP in KVPList)
 							{
 								string? Value = DirectoryKVP.Value;
-								if (!String.IsNullOrEmpty(Value))
+								if (String.IsNullOrEmpty(Value))
+								{
+									continue;
+								}
+
+								if (UpdatedValues.TryAdd(Value, NextPropertyName))
+								{
+									string PropertyName = NextPropertyName;
+									PropertyNameBuilder.Append("_").Append(++PropertyIndex);
+									NextPropertyName = PropertyNameBuilder.ToString();
+									PropertyNameBuilder.Length = PropertyNameBuilderLength;
+									DirectoryToStringDict[DirectoryKVP.Key] = PropertyName;
+								}
+								else
 								{
-									if (!UpdatedValues.ContainsKey(Value))
-									{
-										string PropertyName = PropertyNamePrefix;
-										if (PropertyIndex > 0)
-										{
-											PropertyName += "_" + PropertyIndex;
-										}
-										PropertyIndex++;
-										UpdatedValues.Add(Value, PropertyName);
-									}
 									DirectoryToStringDict[DirectoryKVP.Key] = UpdatedValues[Value];
 								}
 							}
 						}
 
 						// Find the common property values
-						Dictionary<string, string> ValueToCommonPropertyDict = new();
+						Dictionary<string, string> ValueToCommonPropertyDict = new(StringComparer.Ordinal);
 						{
 							Dictionary<string, int> ValueAndCount = new();
 							foreach (string PropertyValue in UpdatedValues.Keys)
@@ -1526,7 +1535,7 @@ namespace UnrealBuildTool
 							}
 
 							List<string> CommonProperties = ValueAndCount.Where(kvp => kvp.Value > 1).Select(kvp => kvp.Key).ToList();
-							CommonProperties.Sort();
+							CommonProperties.Sort(StringComparer.Ordinal);
 
 							// Write out the common property values
 							int CommonPropertyValueIndex = 0;
@@ -1554,16 +1563,19 @@ namespace UnrealBuildTool
 									PathArray[Index] = PropertyValue;
 								}
 							}
+							Array.Sort(PathArray, String.CompareOrdinal);
 							string NewValue = String.Join(";", PathArray);
 							CommonProjectFileContent.AppendLine($"    <{kvp.Value}>{PropertyValuePrefix}{NewValue}</{kvp.Value}>");
 						}
 					};
 
+					// Use short names on these because there are _a ton_ of them in the big projects. the vcxproj of our biggest project went from 70mb to 31mb
+
 					//AdditionalIncludeDirectories
-					AddProperties(DirectoryToIncludeSearchPaths, "ProjectAdditionalIncludeDirectories", "ClCompile_AdditionalIncludeDirectories", "$(NMakeIncludeSearchPath);");
+					AddProperties(DirectoryToIncludeSearchPaths, "AID", "ClCompile_AdditionalIncludeDirectories", "$(NMakeIncludeSearchPath);");
 
 					//ForcedIncludeFiles
-					AddProperties(DirectoryToForceIncludePaths, "ProjectForcedIncludeFiles", "ClCompile_ForcedIncludeFiles", String.Empty);
+					AddProperties(DirectoryToForceIncludePaths, "FIF", "ClCompile_ForcedIncludeFiles", String.Empty);
 
 					// AdditionalOptions
 					{
@@ -1571,22 +1583,30 @@ namespace UnrealBuildTool
 						Dictionary<string, string> UpdatedValues = new();
 						List<KeyValuePair<DirectoryReference, string?>> KVPList = DirectoryToPchFile.ToList();
 						KVPList.SortBy(kvp => kvp.Key);
+
+						StringBuilder PropertyNameBuilder = new("ClCompile_AdditionalOptions");
+						int PropertyNameBuilderLength = PropertyNameBuilder.Length;
+						string NextPropertyName = PropertyNameBuilder.ToString();
+
 						foreach (KeyValuePair<DirectoryReference, string?> DirectoryKVP in KVPList)
 						{
 							string? Value = DirectoryKVP.Value;
-							if (!String.IsNullOrEmpty(Value))
+							if (String.IsNullOrEmpty(Value))
+							{
+								continue;
+							}
+							if (UpdatedValues.TryAdd(Value, NextPropertyName))
+							{
+								string PropertyName = NextPropertyName;
+								PropertyNameBuilder.Append("_").Append(++ValueIndex);
+								NextPropertyName = PropertyNameBuilder.ToString();
+								PropertyNameBuilder.Length = PropertyNameBuilderLength;
+
+								CommonProjectFileContent.AppendLine($"    <{PropertyName}>$(AdditionalOptions) /Yu\"{Value}\"</{PropertyName}>");
+								DirectoryToPchFile[DirectoryKVP.Key] = PropertyName;
+							}
+							else
 							{
-								if (!UpdatedValues.ContainsKey(Value))
-								{
-									string PropertyName = "ClCompile_AdditionalOptions";
-									if (ValueIndex > 0)
-									{
-										PropertyName += "_" + ValueIndex;
-									}
-									ValueIndex++;
-									UpdatedValues.Add(Value, PropertyName);
-									CommonProjectFileContent.AppendLine($"    <{PropertyName}>$(AdditionalOptions) /Yu\"{Value}\"</{PropertyName}>");
-								}
 								DirectoryToPchFile[DirectoryKVP.Key] = UpdatedValues[Value];
 							}
 						}
@@ -1607,7 +1627,10 @@ namespace UnrealBuildTool
 				// Add all file directories to the filters file as solution filters
 				HashSet<string> FilterDirectories = new HashSet<string>();
 
-				foreach (AliasedFile AliasedFile in LocalAliasedFiles)
+				AliasedFile[] LocalAliasedFilesArray = LocalAliasedFiles.ToArray();
+				Array.Sort(LocalAliasedFilesArray);
+
+				foreach (AliasedFile AliasedFile in LocalAliasedFilesArray)
 				{
 					// No need to add the root directory relative to the project (it would just be an empty string!)
 					if (!String.IsNullOrWhiteSpace(AliasedFile.ProjectPath))
@@ -1615,23 +1638,8 @@ namespace UnrealBuildTool
 						FiltersFileIsNeeded = EnsureFilterPathExists(AliasedFile.ProjectPath, VCFiltersFileContent, FilterDirectories);
 					}
 
-					// get the filetype as represented to Visual Studio
+					// Get the filetype as represented to Visual Studio
 					VCFileType FileType = GetVCFileType(AliasedFile.FileSystemPath);
-					DirectoryReference FileSystemPathDir = new DirectoryReference(AliasedFile.FileSystemPath);
-
-					// if the filetype is an include and its path is filtered out, skip it entirely (should we do this for any type of
-					// file? Possibly, but not today due to potential fallout)
-					if (CanVCFileTypeBeIncluded(FileType) && IncludePathIsFilteredOut(FileSystemPathDir))
-					{
-						continue;
-					}
-
-					// Allow filtering of any type of file
-					if (FilePathIsFilteredOut(FileSystemPathDir))
-					{
-						continue;
-					}
-
 					string VCFileType = GetVCFileTypeString(FileType);
 
 					if (!CanVCFileTypeBeCompiled(FileType))
diff --git a/Engine/Source/Programs/UnrealBuildTool/ProjectFiles/Xcode/XcodeProjectFileGenerator.cs b/Engine/Source/Programs/UnrealBuildTool/ProjectFiles/Xcode/XcodeProjectFileGenerator.cs
--- a/Engine/Source/Programs/UnrealBuildTool/ProjectFiles/Xcode/XcodeProjectFileGenerator.cs
+++ b/Engine/Source/Programs/UnrealBuildTool/ProjectFiles/Xcode/XcodeProjectFileGenerator.cs
@@ -524,7 +524,7 @@ namespace UnrealBuildTool
 			}
 		}
 
-		protected override void AddAdditionalNativeTargetInformation(PlatformProjectGeneratorCollection PlatformProjectGenerators, List<Tuple<ProjectFile, ProjectTarget>> Targets, ILogger Logger)
+		protected override void AddAdditionalNativeTargetInformation(PlatformProjectGeneratorCollection PlatformProjectGenerators, List<(ProjectFile, ProjectTarget)> Targets, ILogger Logger)
 		{
 			using ITimelineEvent Timer = Timeline.ScopeEvent("Adding additional Xcode project info", Logger);
 			DateTime MainStart = DateTime.UtcNow;
@@ -577,7 +577,8 @@ namespace UnrealBuildTool
 						List<ModuleRules.RuntimeDependency> Dylibs = new();
 						// Generate a compile environment for each module in the binary, and cache the one per Target, since that is the slowest one
 						CppCompileEnvironment GlobalCompileEnvironment = Target.CreateCompileEnvironmentForProjectFiles(Logger);
-						GlobalCompileEnvironmentCache.TryAdd(MakeGlobalCompileEnvironmentCacheKey(CurTarget.UnrealProjectFilePath, CurTarget.Name, Platform, UnrealTargetConfiguration.Development, new UnrealArchitectures(Arch)), GlobalCompileEnvironment);
+						string cacheKey = MakeGlobalCompileEnvironmentCacheKey(CurTarget.UnrealProjectFilePath, CurTarget.Name, Platform, UnrealTargetConfiguration.Development, new UnrealArchitectures(Arch));
+						GlobalCompileEnvironmentCache.TryAdd(cacheKey, GlobalCompileEnvironment);
 
 						foreach (UEBuildBinary Binary in Target.Binaries)
 						{
diff --git a/Engine/Source/Programs/UnrealBuildTool/System/ActionGraphBuilder.cs b/Engine/Source/Programs/UnrealBuildTool/System/ActionGraphBuilder.cs
--- a/Engine/Source/Programs/UnrealBuildTool/System/ActionGraphBuilder.cs
+++ b/Engine/Source/Programs/UnrealBuildTool/System/ActionGraphBuilder.cs
@@ -89,14 +89,15 @@ namespace UnrealBuildTool
 	sealed class NullActionGraphBuilder : IActionGraphBuilder
 	{
 		private readonly ILogger Logger;
+		private readonly FileHashCache? _fileHashCache;
 
 		/// <summary>
 		/// Constructor
 		/// </summary>
-		/// <param name="InLogger"></param>
-		public NullActionGraphBuilder(ILogger InLogger)
+		public NullActionGraphBuilder(ILogger InLogger, FileHashCache? fileHashCache = null)
 		{
 			Logger = InLogger;
+			_fileHashCache = fileHashCache;
 		}
 
 		/// <inheritdoc/>
@@ -107,13 +108,27 @@ namespace UnrealBuildTool
 		/// <inheritdoc/>
 		public void CreateIntermediateTextFile(FileItem FileItem, string Contents, bool AllowAsync = true)
 		{
-			Utils.WriteFileIfChanged(FileItem, Contents, Logger);
+			if (_fileHashCache != null)
+			{
+				_fileHashCache.WriteFileIfChanged(FileItem, Contents, Logger, false);
+			}
+			else
+			{
+				Utils.WriteFileIfChanged(FileItem, Contents, Logger);
+			}
 		}
 
 		/// <inheritdoc/>
 		public void CreateIntermediateTextFile(FileItem FileItem, IEnumerable<string> ContentLines, bool AllowAsync = true)
 		{
-			Utils.WriteFileIfChanged(FileItem, ContentLines, Logger);
+			if (_fileHashCache != null)
+			{
+				_fileHashCache.WriteFileIfChanged(FileItem, ContentLines, Logger, false);
+			}
+			else
+			{
+				Utils.WriteFileIfChanged(FileItem, ContentLines, Logger);
+			}
 		}
 
 		/// <inheritdoc/>
diff --git a/Engine/Source/Programs/UnrealBuildTool/System/HashSetPool.cs b/Engine/Source/Programs/UnrealBuildTool/System/HashSetPool.cs
new file mode 100644
--- /dev/null
+++ b/Engine/Source/Programs/UnrealBuildTool/System/HashSetPool.cs
@@ -0,0 +1,71 @@
+// Copyright Epic Games, Inc. All Rights Reserved.
+
+using Microsoft.Extensions.ObjectPool;
+using System;
+using System.Collections.Generic;
+
+// Usage:
+//
+// using var rented = HashSetPool<string>.Rent(estimate);
+// var hashSet = rented.Value;
+//
+// use hashSet...
+
+static class HashSetPool<T>
+{
+	public static readonly ObjectPool<HashSet<T>> Shared = new DefaultObjectPool<HashSet<T>>(new Policy(), maximumRetained: 256);
+
+	sealed class Policy : PooledObjectPolicy<HashSet<T>>
+	{
+		public override HashSet<T> Create() => new(capacity: 256); // <- set a good default
+
+		public override bool Return(HashSet<T> obj)
+		{
+			obj.Clear();
+			return true;
+		}
+	}
+
+	public static Pooled<HashSet<T>> Rent(int ensureCapacity = 0)
+	{
+		var set = Shared.Get();
+		if (ensureCapacity > 0)
+		{
+			set.EnsureCapacity(ensureCapacity);
+		}
+		return new Pooled<HashSet<T>>(Shared, set);
+	}
+}
+
+/// <summary>
+/// Tiny disposable wrapper so you can `using` it
+/// </summary>
+public sealed class Pooled<T> : IDisposable where T : class
+{
+	readonly ObjectPool<T> _pool;
+
+	/// <summary>
+	/// </summary>
+	public T Value { get; }
+
+	/// <summary>
+	/// </summary>
+	public Pooled(ObjectPool<T> pool, T value)
+	{
+		_pool = pool;
+		Value = value;
+	}
+
+	/// <summary>
+	/// </summary>
+	public void Dispose() => _pool.Return(Value);
+}
+
+/// <summary>
+/// For ordinal compare
+/// </summary>
+sealed class StringSetPolicy : PooledObjectPolicy<HashSet<string>>
+{
+	public override HashSet<string> Create() => new(StringComparer.Ordinal);
+	public override bool Return(HashSet<string> obj) { obj.Clear(); return true; }
+}
\ No newline at end of file
diff --git a/Engine/Source/Programs/UnrealBuildTool/System/PrecompiledHeaderTemplate.cs b/Engine/Source/Programs/UnrealBuildTool/System/PrecompiledHeaderTemplate.cs
--- a/Engine/Source/Programs/UnrealBuildTool/System/PrecompiledHeaderTemplate.cs
+++ b/Engine/Source/Programs/UnrealBuildTool/System/PrecompiledHeaderTemplate.cs
@@ -34,12 +34,12 @@ namespace UnrealBuildTool
 		/// <summary>
 		/// Instances of this PCH
 		/// </summary>
-		public List<PrecompiledHeaderInstance> Instances = new List<PrecompiledHeaderInstance>();
+		public List<PrecompiledHeaderInstance> Instances = [];
 
 		/// <summary>
 		/// All the module dependencies this template has
 		/// </summary>
-		public HashSet<UEBuildModule> ModuleDependencies = new HashSet<UEBuildModule>();
+		public HashSet<UEBuildModule> ModuleDependencies;
 
 		/// <summary>
 		/// Constructor
diff --git a/Engine/Source/Programs/UnrealBuildTool/System/RulesAssembly.cs b/Engine/Source/Programs/UnrealBuildTool/System/RulesAssembly.cs
--- a/Engine/Source/Programs/UnrealBuildTool/System/RulesAssembly.cs
+++ b/Engine/Source/Programs/UnrealBuildTool/System/RulesAssembly.cs
@@ -1,6 +1,7 @@
 // Copyright Epic Games, Inc. All Rights Reserved.
 
 using System;
+using System.Collections.Concurrent;
 using System.Collections.Generic;
 using System.Diagnostics.CodeAnalysis;
 using System.IO;
@@ -428,6 +429,8 @@ namespace UnrealBuildTool
 			}
 		}
 
+		private ConcurrentDictionary<string, Type?> _moduleNameToType = [];
+
 		/// <summary>
 		/// Gets the type defining rules for the given module within this assembly
 		/// </summary>
@@ -435,15 +438,18 @@ namespace UnrealBuildTool
 		/// <returns>The rules type for this module, or null if not found</returns>
 		private Type? GetModuleRulesTypeInternal(string ModuleName)
 		{
-			// The build module must define a type named 'Rules' that derives from our 'ModuleRules' type.  
-			Type? RulesObjectType = CompiledAssembly?.GetType(ModuleName);
-			if (RulesObjectType == null)
+			return _moduleNameToType.GetOrAdd(ModuleName, ModuleName =>
 			{
-				// Temporary hack to avoid System namespace collisions
-				// @todo projectfiles: Make rules assemblies require namespaces.
-				RulesObjectType = CompiledAssembly?.GetType("UnrealBuildTool.Rules." + ModuleName);
-			}
-			return RulesObjectType;
+				// The build module must define a type named 'Rules' that derives from our 'ModuleRules' type.  
+				Type? RulesObjectType = CompiledAssembly?.GetType(ModuleName);
+				if (RulesObjectType == null)
+				{
+					// Temporary hack to avoid System namespace collisions
+					// @todo projectfiles: Make rules assemblies require namespaces.
+					RulesObjectType = CompiledAssembly?.GetType("UnrealBuildTool.Rules." + ModuleName);
+				}
+				return RulesObjectType;
+			});
 		}
 
 		/// <summary>
diff --git a/Engine/Source/Programs/UnrealBuildTool/System/SourceFileSearch.cs b/Engine/Source/Programs/UnrealBuildTool/System/SourceFileSearch.cs
--- a/Engine/Source/Programs/UnrealBuildTool/System/SourceFileSearch.cs
+++ b/Engine/Source/Programs/UnrealBuildTool/System/SourceFileSearch.cs
@@ -56,6 +56,40 @@ namespace UnrealBuildTool
 		{
 			return FindFiles(DirectoryItem.GetItemByDirectoryReference(DirectoryToSearch), SubdirectoryNamesToExclude, SearchSubdirectories, SearchedDirectories);
 		}
+		public static List<FileReference> FindFiles2(IEnumerable<DirectoryReference> DirectoriesToSearch, List<string>? SubdirectoryNamesToExclude = null, bool SearchSubdirectories = true, HashSet<DirectoryReference>? SearchedDirectories = null)
+		{
+			// Build a list of directory names that we don't want to search under. We always ignore intermediate directories.
+			string[] ExcludedDirectorySuffixes;
+			if (SubdirectoryNamesToExclude == null)
+			{
+				ExcludedDirectorySuffixes = DefaultExcludedDirectorySuffixes;
+			}
+			else
+			{
+				ExcludedDirectorySuffixes = SubdirectoryNamesToExclude.Select(x => Path.DirectorySeparatorChar + x.ToLowerInvariant()).Union(DefaultExcludedDirectorySuffixes).ToArray();
+			}
+
+			// Find all the files
+			List<FileReference> FoundFiles = new List<FileReference>();
+
+			ParallelQueue queue = new();
+			foreach (DirectoryReference dir in DirectoriesToSearch)
+			{
+				queue.Enqueue(() =>
+				{
+					DirectoryItem dirItem = DirectoryItem.GetItemByDirectoryReference(dir);
+					if (SearchSubdirectories)
+					{
+						FindFilesInternalRecursive(dirItem, ExcludedDirectorySuffixes, FoundFiles, SearchedDirectories, queue);
+					}
+					else
+					{
+						FindFilesInternal(dirItem, FoundFiles, SearchedDirectories, queue);
+					}
+				});
+			}
+			return FoundFiles;
+		}
 
 		public static List<FileReference> FindFiles(DirectoryItem DirectoryToSearch, List<string>? SubdirectoryNamesToExclude = null, bool SearchSubdirectories = true, HashSet<DirectoryReference>? SearchedDirectories = null)
 		{
@@ -74,57 +108,107 @@ namespace UnrealBuildTool
 			List<FileReference> FoundFiles = new List<FileReference>();
 			if (SearchSubdirectories)
 			{
-				FindFilesInternalRecursive(DirectoryToSearch, ExcludedDirectorySuffixes, FoundFiles, SearchedDirectories);
+				FindFilesInternalRecursive(DirectoryToSearch, ExcludedDirectorySuffixes, FoundFiles, SearchedDirectories, null);
 			}
 			else
 			{
-				FindFilesInternal(DirectoryToSearch, ExcludedDirectorySuffixes, FoundFiles, SearchedDirectories);
+				FindFilesInternal(DirectoryToSearch, FoundFiles, SearchedDirectories, null);
 			}
 			return FoundFiles;
 		}
 
-		static void FindFilesInternal(DirectoryItem Directory, string[] ExcludedDirectorySuffixes, List<FileReference> FoundFiles, HashSet<DirectoryReference>? SearchedDirectories)
+		static bool FindFilesInternal(DirectoryItem Directory, List<FileReference> FoundFiles, HashSet<DirectoryReference>? SearchedDirectories, ParallelQueue? queue)
 		{
-			SearchedDirectories?.Add(Directory.Location);
+			if (queue != null && SearchedDirectories != null)
+			{
+				lock (SearchedDirectories)
+				{
+					SearchedDirectories?.Add(Directory.Location);
+				}
+			}
+			else
+			{
+				SearchedDirectories?.Add(Directory.Location);
+			}
+
+			List<FileReference> originalFoundFiles = FoundFiles;
+			if (queue != null)
+			{
+				FoundFiles = [];
+			}
 
 			foreach (FileItem File in Directory.EnumerateFiles())
 			{
-				if (ShouldInclude(File.Location, DefaultExcludedFileSuffixes))
+				// Ignore Mac resource fork files on non-HFS partitions
+				if (File.Name.StartsWith("._"))
 				{
-					FoundFiles.Add(File.Location);
+					continue;
 				}
+
+				if (IsExcluded(DefaultExcludedFileSuffixes, File.Name))
+				{
+					continue;
+				}
+
+				if (File.Name == ".ubtignore")
+				{
+					return false;
+				}
+
+				FoundFiles.Add(File.Location);
 			}
+
+			if (queue != null)
+			{
+				lock (originalFoundFiles)
+				{
+					originalFoundFiles.AddRange(FoundFiles);
+				}
+			}
+
+			return true;
 		}
 
-		static void FindFilesInternalRecursive(DirectoryItem Directory, string[] ExcludedDirectorySuffixes, List<FileReference> FoundFiles, HashSet<DirectoryReference>? SearchedDirectories)
+		static void FindFilesInternalRecursive(DirectoryItem Directory, string[] ExcludedDirectorySuffixes, List<FileReference> FoundFiles, HashSet<DirectoryReference>? SearchedDirectories, ParallelQueue? queue)
 		{
-			FindFilesInternal(Directory, ExcludedDirectorySuffixes, FoundFiles, SearchedDirectories);
+			if (!FindFilesInternal(Directory, FoundFiles, SearchedDirectories, queue))
+			{
+				return;
+			}
 
 			foreach (DirectoryItem SubDirectory in Directory.EnumerateDirectories())
 			{
-				if (ShouldInclude(SubDirectory.Location, ExcludedDirectorySuffixes))
+				if (IsExcluded(ExcludedDirectorySuffixes, SubDirectory.Location.FullName))
+				{
+					continue;
+				}
+
+				if (SubDirectory.Name.StartsWith("__External"))
+				{
+					continue;
+				}
+
+				if (queue != null)
 				{
-					FindFilesInternalRecursive(SubDirectory, ExcludedDirectorySuffixes, FoundFiles, SearchedDirectories);
+					queue.Enqueue(() => FindFilesInternalRecursive(SubDirectory, ExcludedDirectorySuffixes, FoundFiles, SearchedDirectories, queue));
+				}
+				else
+				{
+					FindFilesInternalRecursive(SubDirectory, ExcludedDirectorySuffixes, FoundFiles, SearchedDirectories, queue);
 				}
 			}
 		}
 
-		static bool ShouldInclude(FileSystemReference Reference, string[] ExcludedSuffixes)
+		static bool IsExcluded(string[] excludeSuffixes, string testString)
 		{
-			// Ignore Mac resource fork files on non-HFS partitions
-			if (Path.GetFileName(Reference.FullName).StartsWith("._"))
+			foreach (string excludeSuffix in excludeSuffixes)
 			{
-				return false;
-			}
-
-			foreach (string ExcludedSuffix in ExcludedSuffixes)
-			{
-				if (Reference.FullName.EndsWith(ExcludedSuffix, FileSystemReference.Comparison))
+				if (testString.EndsWith(excludeSuffix, FileSystemReference.Comparison))
 				{
-					return false;
+					return true;
 				}
 			}
-			return true;
+			return false;
 		}
 	}
 }
