[UBT]
* Fixed so not finding FileHashCache.bin is not a warning

diff --git a/Engine/Source/Programs/UnrealBuildTool/System/FileHashCache.cs b/Engine/Source/Programs/UnrealBuildTool/System/FileHashCache.cs
--- a/Engine/Source/Programs/UnrealBuildTool/System/FileHashCache.cs
+++ b/Engine/Source/Programs/UnrealBuildTool/System/FileHashCache.cs
@@ -612,6 +612,9 @@ namespace UnrealBuildTool
 					partition.Files.TryAdd(FileItem.GetItemByPath(path), new() { hash = hash, lastWriteTimeUtc = lastWriteTimeUtc });
 				}
 			}
+			catch (FileNotFoundException)
+			{
+			}
 			catch (Exception Ex)
 			{
 				logger.LogWarning("Unable to read {Location}. See log for additional information.", partition.Location);
