Fix a few typos

diff --git a/Engine/Source/Programs/Shared/EpicGames.Build/UEBuildPlatformSDK.cs b/Engine/Source/Programs/Shared/EpicGames.Build/UEBuildPlatformSDK.cs
--- a/Engine/Source/Programs/Shared/EpicGames.Build/UEBuildPlatformSDK.cs
+++ b/Engine/Source/Programs/Shared/EpicGames.Build/UEBuildPlatformSDK.cs
@@ -2233,7 +2233,7 @@ namespace EpicGames.Core
 
 					if (!RunAutoSDKHooks(AutoSDKRoot, GetAutoSDKDirectoryForMainVersion(), DesiredSDKLevel, SDKHookType.Install, false))
 					{
-						Logger.LogDebug("Failed to install required SDK {SdkVersion}.  Attemping to uninstall", GetAutoSDKDirectoryForMainVersion());
+						Logger.LogDebug("Failed to install required SDK {SdkVersion}.  Attempting to uninstall", GetAutoSDKDirectoryForMainVersion());
 						RunAutoSDKHooks(AutoSDKRoot, GetAutoSDKDirectoryForMainVersion(), DesiredSDKLevel, SDKHookType.Uninstall, false);
 						return;
 					}
@@ -2346,4 +2346,4 @@ namespace EpicGames.Core
 		#endregion
 
 	}
-}
\ No newline at end of file
+}
diff --git a/Engine/Source/Programs/UnrealBuildTool/Platform/Android/UEDeployAndroid.cs b/Engine/Source/Programs/UnrealBuildTool/Platform/Android/UEDeployAndroid.cs
--- a/Engine/Source/Programs/UnrealBuildTool/Platform/Android/UEDeployAndroid.cs
+++ b/Engine/Source/Programs/UnrealBuildTool/Platform/Android/UEDeployAndroid.cs
@@ -4950,7 +4950,7 @@ dependencies.implementation('androidx.tracing:tracing:1.2.0')
 
 			if (BypassGradlePackaging && !bEnableBundle && !bSkipGradleBuild && !bOnlyGradleGeneration && !bCreateFromScratch)
 			{
-				Logger.LogInformation("Attemping BypassGradlePackaging");
+				Logger.LogInformation("Attempting BypassGradlePackaging");
 				int BuildListComboRemaining = BuildListComboTotal;
 
 				foreach (Tuple<UnrealArch, string> build in BuildList)
@@ -4969,7 +4969,7 @@ dependencies.implementation('androidx.tracing:tracing:1.2.0')
 					string ApkFilename = Path.GetFileNameWithoutExtension(OutputPath).Replace("UnrealGame", ProjectName);
 					string DestApkName = Path.Combine(DestApkDirectory, ApkFilename + ".apk");
 
-					// As we are always making seperate APKs we need to put the architecture into the name
+					// As we are always making separate APKs we need to put the architecture into the name
 					DestApkName = AndroidToolChain.InlineArchName(DestApkName, Arch);
 
 					if (!File.Exists(SourceSOName))
@@ -5343,7 +5343,7 @@ dependencies.implementation('androidx.tracing:tracing:1.2.0')
 					string ApkFilename = Path.GetFileNameWithoutExtension(OutputPath).Replace("UnrealGame", ProjectName);
 					DestApkName = Path.Combine(DestApkDirectory, ApkFilename + ".apk");
 
-					// As we are always making seperate APKs we need to put the architecture into the name
+					// As we are always making separate APKs we need to put the architecture into the name
 					DestApkName = AndroidToolChain.InlineArchName(DestApkName, Arch);
 
 					bool bUseAFS = true;
@@ -5415,7 +5415,7 @@ popd
 					string ApkFilename = Path.GetFileNameWithoutExtension(OutputPath).Replace("UnrealGame", ProjectName);
 					DestApkName = Path.Combine(DestApkDirectory, ApkFilename + ".apk");
 
-					// As we are always making seperate APKs we need to put the architecture into the name
+					// As we are always making separate APKs we need to put the architecture into the name
 					DestApkName = AndroidToolChain.InlineArchName(DestApkName, Arch);
 
 					if (!bOnlyGradleGeneration && !File.Exists(SourceSOName))
diff --git a/Engine/Source/Programs/UnrealBuildTool/System/Utils.cs b/Engine/Source/Programs/UnrealBuildTool/System/Utils.cs
--- a/Engine/Source/Programs/UnrealBuildTool/System/Utils.cs
+++ b/Engine/Source/Programs/UnrealBuildTool/System/Utils.cs
@@ -1537,7 +1537,7 @@ namespace UnrealBuildTool
 			// look for -project=<path>, if it does not exist check arguments for anything that has .uproject in it
 			if (!Arguments.TryGetValue("-Project=", out CandidateProjectPath) )
 			{
-				Logger.LogDebug("No project arugment provided while trying to parse Project File. Attempting fallback parsing to attempt to find a Project File.");
+				Logger.LogDebug("No project argument provided while trying to parse Project File. Attempting fallback parsing to attempt to find a Project File.");
 
 				// Go through the argument list and try to match poorly (or well..) formed arguments like 
 				// EngineTest, EngineTest.uproject
