[UBT]
* Removed replaced usage of var with explicit types

diff --git a/Engine/Source/Programs/UnrealBuildTool/Configuration/UEBuildBinary.cs b/Engine/Source/Programs/UnrealBuildTool/Configuration/UEBuildBinary.cs
--- a/Engine/Source/Programs/UnrealBuildTool/Configuration/UEBuildBinary.cs
+++ b/Engine/Source/Programs/UnrealBuildTool/Configuration/UEBuildBinary.cs
@@ -415,7 +415,7 @@ namespace UnrealBuildTool
 
 			List<UEBuildModule> ReferencedModules = new List<UEBuildModule>();
 
-			using var rented = HashSetPool<UEBuildModule>.Rent();
+			using Pooled<HashSet<UEBuildModule>> rented = HashSetPool<UEBuildModule>.Rent();
 			HashSet<UEBuildModule> IgnoreReferencedModules = rented.Value;
 
 			foreach (UEBuildModule Module in Modules)
@@ -447,7 +447,7 @@ namespace UnrealBuildTool
 
 			List<UEBuildModule> DirectlyReferencedModules = new List<UEBuildModule>();
 
-			using var rented = HashSetPool<UEBuildModule>.Rent();
+			using Pooled<HashSet<UEBuildModule>> rented = HashSetPool<UEBuildModule>.Rent();
 			HashSet<UEBuildModule> VisitedModules = rented.Value;
 
 			for (int Idx = 0; Idx < ReferencedModules.Count; Idx++)
@@ -728,7 +728,7 @@ namespace UnrealBuildTool
 			LinkEnvironment BinaryLinkEnvironment = new LinkEnvironment(LinkEnvironment);
 			List<UEBuildBinary> BinaryDependencies = new List<UEBuildBinary>();
 
-			using var rented = HashSetPool<UEBuildModule>.Rent();
+			using Pooled<HashSet<UEBuildModule>> rented = HashSetPool<UEBuildModule>.Rent();
 			HashSet<UEBuildModule> LinkEnvironmentVisitedModules = rented.Value;
 
 			CppCompileEnvironment BinaryCompileEnvironment = CreateBinaryCompileEnvironment(CompileEnvironment);
@@ -738,7 +738,7 @@ namespace UnrealBuildTool
 				BinaryCompileEnvironment.bUseSharedBuildEnvironment = false;
 			}
 
-			using var rented2 = HashSetPool<FileItem>.Rent();
+			using Pooled<HashSet<FileItem>> rented2 = HashSetPool<FileItem>.Rent();
 			HashSet<FileItem> InputFilesLookup = rented2.Value;
 
 			foreach (UEBuildModule Module in UEBuildModule.StableTopologicalSort(Modules))
diff --git a/Engine/Source/Programs/UnrealBuildTool/Configuration/UEBuildModule.cs b/Engine/Source/Programs/UnrealBuildTool/Configuration/UEBuildModule.cs
--- a/Engine/Source/Programs/UnrealBuildTool/Configuration/UEBuildModule.cs
+++ b/Engine/Source/Programs/UnrealBuildTool/Configuration/UEBuildModule.cs
@@ -1100,7 +1100,7 @@ namespace UnrealBuildTool
 		/// <param name="bOnlyDirectDependencies">True to return only this module's direct dependencies</param>
 		public void GetAllDependencyModules(List<UEBuildModule> ReferencedModules, HashSet<UEBuildModule> IgnoreReferencedModules, bool bIncludeDynamicallyLoaded, bool bForceCircular, bool bOnlyDirectDependencies)
 		{
-			var stack = new Stack<UEBuildModule>();
+			Stack<UEBuildModule> stack = new();
 			stack.Push(this);
 
 			void QueueReferencedModules(List<UEBuildModule> Modules)
@@ -1143,7 +1143,7 @@ namespace UnrealBuildTool
 		/// <param name="bIncludePrivateModules">Whether to include private modules.</param>
 		public void GetAllDependencyModulesForPCH(HashSet<UEBuildModule> outReferenced, bool bIncludePrivateModules)
 		{
-			using var rented = HashSetPool<UEBuildModule>.Rent();
+			using Pooled<HashSet<UEBuildModule>> rented = HashSetPool<UEBuildModule>.Rent();
 			HashSet<UEBuildModule> IgnoreReferencedModules = rented.Value;
 			IgnoreReferencedModules.Add(this);
 			InternalGetAllDependencyModulesForPCH(outReferenced, IgnoreReferencedModules, bIncludePrivateModules);
@@ -1163,7 +1163,7 @@ namespace UnrealBuildTool
 			//	(PublicIncludePathModules != null ? PublicIncludePathModules!.Count : 0)
 			//	);
 
-			using var rented = HashSetPool<UEBuildModule>.Rent(((bIncludePrivateModules && PrivateDependencyModules != null) ? PrivateDependencyModules.Count : 0) +
+			using Pooled<HashSet<UEBuildModule>> rented = HashSetPool<UEBuildModule>.Rent(((bIncludePrivateModules && PrivateDependencyModules != null) ? PrivateDependencyModules.Count : 0) +
 				(PublicDependencyModules != null ? PublicDependencyModules.Count : 0) +
 				(PublicIncludePathModules != null ? PublicIncludePathModules!.Count : 0));
 			HashSet<UEBuildModule> AllDependencyModules = rented.Value;
@@ -1350,7 +1350,7 @@ namespace UnrealBuildTool
 		{
 			List<CircularReference> circularReferences = [];
 
-			using var rented = HashSetPool<UEBuildModule>.Rent();
+			using Pooled<HashSet<UEBuildModule>> rented = HashSetPool<UEBuildModule>.Rent();
 			HashSet<UEBuildModule> processed = rented.Value;
 			CheckForCircularReferences(this, processed, [], circularReferences);
 
@@ -1658,12 +1658,12 @@ namespace UnrealBuildTool
 
 			int words = (nodeCount + 63) >> 6;
 			long bytesNeeded = (long)nodeCount * words * sizeof(ulong);
-			var depBits = new ulong[(long)nodeCount * words];
+			ulong[] depBits = new ulong[(long)nodeCount * words];
 
 			for (int i = 0; i < nodeCount; i++)
 			{
 				int rowBase = i * words;
-				foreach (var d in NodeList[i].DependenciesForStableTopologicalSort!)
+				foreach (UEBuildModule d in NodeList[i].DependenciesForStableTopologicalSort!)
 				{
 					if (indexOf.TryGetValue(d, out int j) && j != i)
 					{
diff --git a/Engine/Source/Programs/UnrealBuildTool/Configuration/UEBuildModuleCPP.cs b/Engine/Source/Programs/UnrealBuildTool/Configuration/UEBuildModuleCPP.cs
--- a/Engine/Source/Programs/UnrealBuildTool/Configuration/UEBuildModuleCPP.cs
+++ b/Engine/Source/Programs/UnrealBuildTool/Configuration/UEBuildModuleCPP.cs
@@ -1371,12 +1371,12 @@ namespace UnrealBuildTool
 			HashSet<UEBuildModule> AllModuleDeps = ModuleToIncludePathsOnlyFlag.Keys.ToHashSet();
 			HashSet<UEBuildModule> ModuleDeps = new HashSet<UEBuildModule>(AllModuleDeps) { this };
 
-			using var rented = HashSetPool<string>.Rent();
+			using Pooled<HashSet<string>> rented = HashSetPool<string>.Rent();
 			HashSet<string> VisitedIncludes = rented.Value;
 
 			HashSet<FileItem> FoundIncludeFileItems = new();
 
-			using var rented3 = HashSetPool<string>.Rent();
+			using Pooled<HashSet<string>> rented3 = HashSetPool<string>.Rent();
 			HashSet<string> IncludesNotFound = rented3.Value;
 
 
@@ -1717,7 +1717,7 @@ namespace UnrealBuildTool
 					if (Rules.Target.bChainPCHs)
 					{
 						// Find all the dependencies of this module
-						using var rented = HashSetPool<UEBuildModule>.Rent();
+						using Pooled<HashSet<UEBuildModule>> rented = HashSetPool<UEBuildModule>.Rent();
 						HashSet<UEBuildModule> ReferencedModules = rented.Value;
 						Template.Module.GetAllDependencyModules(new List<UEBuildModule>(), ReferencedModules, bIncludeDynamicallyLoaded: false, bForceCircular: false, bOnlyDirectDependencies: false);
 
@@ -2164,7 +2164,7 @@ namespace UnrealBuildTool
 				if (CompileEnvironment.bHasPrecompiledHeader == false && CompileEnvironment.SharedPCHs.Count > 0 && !CompileEnvironment.bIsBuildingLibrary && Rules.PCHUsage != ModuleRules.PCHUsageMode.NoSharedPCHs)
 				{
 					// Find the first shared PCH module we can use that doesn't reference this module
-					using var rented = HashSetPool<UEBuildModule>.Rent();
+					using Pooled<HashSet<UEBuildModule>> rented = HashSetPool<UEBuildModule>.Rent();
 					HashSet<UEBuildModule> AllDependencies = rented.Value;
 					GetAllDependencyModulesForPCH(AllDependencies, true);
 
@@ -2317,7 +2317,7 @@ namespace UnrealBuildTool
 				Writer.Append("#undef ").Append(Undefinition).AppendLine();
 			}
 
-			using var rented = HashSetPool<string>.Rent();
+			using Pooled<HashSet<string>> rented = HashSetPool<string>.Rent();
 			HashSet<string> Processed = rented.Value;
 			foreach (string Definition in Definitions.Where(x => Processed.Add(x)))
 			{
@@ -2745,7 +2745,7 @@ namespace UnrealBuildTool
 				Result.Definitions.Add(String.Format("UE_WITH_DEBUG_INFO={0}", BaseCompileEnvironment.bCreateDebugInfo ? "1" : "0"));
 
 				StringBuilder sb = new();
-				foreach (var (id, vfs, local) in BaseCompileEnvironment.RootPaths)
+				foreach ((uint id, DirectoryReference vfs, DirectoryReference local) in BaseCompileEnvironment.RootPaths)
 				{
 					sb.Append(vfs.FullName.Replace("\\", "\\\\")).Append(';');
 					sb.Append(local.FullName.Replace("\\", "\\\\")).Append(';');
diff --git a/Engine/Source/Programs/UnrealBuildTool/Configuration/UEBuildTarget.cs b/Engine/Source/Programs/UnrealBuildTool/Configuration/UEBuildTarget.cs
--- a/Engine/Source/Programs/UnrealBuildTool/Configuration/UEBuildTarget.cs
+++ b/Engine/Source/Programs/UnrealBuildTool/Configuration/UEBuildTarget.cs
@@ -1256,7 +1256,7 @@ namespace UnrealBuildTool
 		/// <returns>The build target object for the specified build rules source file</returns>
 		public static UEBuildTarget Create(TargetDescriptor Descriptor, bool bSkipRulesCompile, bool bForceRulesCompile, bool bUsePrecompiled, UnrealIntermediateEnvironment IntermediateEnvironment, ILogger Logger)
 		{
-			var (RulesObject, RulesAssembly) = CreateRulesObject(Descriptor, bSkipRulesCompile, bForceRulesCompile, bUsePrecompiled, IntermediateEnvironment, Logger);
+			(TargetRules RulesObject, RulesAssembly RulesAssembly) = CreateRulesObject(Descriptor, bSkipRulesCompile, bForceRulesCompile, bUsePrecompiled, IntermediateEnvironment, Logger);
 
 			// Generate a build target from this rules module
 			UEBuildTarget Target;
@@ -1289,7 +1289,7 @@ namespace UnrealBuildTool
 			TargetRules RulesObject;
 			using (GlobalTracer.Instance.BuildSpan("RulesAssembly.CreateTargetRules()").StartActive())
 			{
-				using var _ = Timeline.ScopeEvent("CreateTargetRules");
+				using ITimelineEvent _ = Timeline.ScopeEvent("CreateTargetRules");
 				RulesObject = RulesAssembly.CreateTargetRules(Descriptor.Name, Descriptor.Platform, Descriptor.Configuration, Descriptor.Architectures, Descriptor.ProjectFile, Descriptor.AdditionalArguments, Logger, Descriptor.IsTestsTarget, IntermediateEnvironment: IntermediateEnvironment);
 			}
 			if ((ProjectFileGenerator.bGenerateProjectFiles == false) && !RulesObject.GetSupportedPlatforms().Contains(Descriptor.Platform))
@@ -2713,7 +2713,7 @@ namespace UnrealBuildTool
 			// Prepare cached data for UHT header generation
 			using (GlobalTracer.Instance.BuildSpan("UHTExecution.SetupUObjectModules()").StartActive())
 			{
-				using var _ = Timeline.ScopeEvent("UHTModules");
+				using ITimelineEvent _ = Timeline.ScopeEvent("UHTModules");
 
 				UHTExecution.SetupUObjectModules(
 					ModulesToGenerateHeadersFor,
@@ -2745,7 +2745,7 @@ namespace UnrealBuildTool
 
 #if __VPROJECT_AVAILABLE__
 			{
-				using var _ = Timeline.ScopeEvent("VNIModules");
+				using ITimelineEvent _ = Timeline.ScopeEvent("VNIModules");
 				VNIExecution.SetupVNIModules(ModulesToGenerateHeadersFor, out Makefile.VNIModules);
 
 				// Copy Verse BPVM usage flag
@@ -2806,7 +2806,7 @@ namespace UnrealBuildTool
 			{
 				TargetMakefile makefile = context.Makefile;
 
-				foreach (var entry in context.GenerateHeaderFuncs)
+				foreach (GenerateHeaderFunc entry in context.GenerateHeaderFuncs)
 				{
 					DirectoryReference.CreateDirectory(entry.Dir);
 					entry.Func(Logger, entry.Dir);
@@ -2824,7 +2824,7 @@ namespace UnrealBuildTool
 						lastRunUObjectModules = lastRunContext.Makefile.UObjectModules.ToDictionary(k => k.ModuleName, k => k);
 					}
 
-					foreach (var module in makefile.UObjectModules)
+					foreach (UHTModuleInfo module in makefile.UObjectModules)
 					{
 						if (lastRunUObjectModules.TryGetValue(module.ModuleName, out UHTModuleInfo? lastModule))
 						{
@@ -2848,7 +2848,7 @@ namespace UnrealBuildTool
 						lastRunVNIModules = new(lastRunContext.Makefile.VNIModules.Select(m => m.ModuleName));
 					}
 
-					foreach (var module in makefile.VNIModules)
+					foreach (VNIModuleInfo module in makefile.VNIModules)
 					{
 						if (!lastRunVNIModules.Contains(module.ModuleName))
 						{
@@ -2892,7 +2892,7 @@ namespace UnrealBuildTool
 
 		internal static void WaitCodeGen(List<CodeGenContext> contexts, bool prefetchAfterReset)
 		{
-			foreach (var context in contexts)
+			foreach (CodeGenContext context in contexts)
 			{
 				if (context.uhtRun != null)
 				{
@@ -2915,7 +2915,7 @@ namespace UnrealBuildTool
 		internal static List<Task> GetTimestampTasks(List<CodeGenContext> contexts)
 		{
 			List<Task> timestampTasks = [];
-			foreach (var codeGenContext in contexts)
+			foreach (CodeGenContext codeGenContext in contexts)
 			{
 				if (codeGenContext.uhtTimestamp != null)
 				{
@@ -3352,7 +3352,7 @@ namespace UnrealBuildTool
 
 			// Wait for makefile tasks to finish (tasks writing files to disk etc)
 			{
-				using var _ = Timeline.ScopeEvent("WaitOnWrites");
+				using ITimelineEvent _ = Timeline.ScopeEvent("WaitOnWrites");
 				MakefileBuilder.WaitOnWriteTasks();
 			}
 
@@ -3841,12 +3841,12 @@ namespace UnrealBuildTool
 			{
 				binary.PrimaryModule.bCreatePerModuleFile = false;
 
-				var queue = new Queue<(UEBuildModule Module, UEBuildModule Source)>();
-				var visited = new HashSet<UEBuildModule>();
+				Queue<(UEBuildModule Module, UEBuildModule Source)> queue = new();
+				HashSet<UEBuildModule> visited = new();
 
 				// Queue up primary module and all direct dependencies
 				queue.Enqueue((binary.PrimaryModule, binary.PrimaryModule));
-				foreach (var m in binary.Modules)
+				foreach (UEBuildModule m in binary.Modules)
 				{
 					if (visited.Add(m))
 					{
@@ -3857,7 +3857,7 @@ namespace UnrealBuildTool
 				// Search for non-engine PerModuleInlineFile.. 
 				while (queue.Count > 0)
 				{
-					var current = queue.Dequeue();
+					(UEBuildModule Module, UEBuildModule Source) current = queue.Dequeue();
 					if (current.Module.RulesFile.IsUnderDirectory(Unreal.EngineDirectory))
 					{
 						continue;
@@ -3868,7 +3868,7 @@ namespace UnrealBuildTool
 						return;
 					}
 
-					foreach (var dep in current.Module.PublicDependencyModules!.Union(current.Module.PrivateDependencyModules!))
+					foreach (UEBuildModule? dep in current.Module.PublicDependencyModules!.Union(current.Module.PrivateDependencyModules!))
 					{
 						if (visited.Add(dep))
 						{
@@ -4671,7 +4671,7 @@ namespace UnrealBuildTool
 			SortedDictionary<string, (int, PrecompiledHeaderTemplate)> SharedPCHModuleToPriority = new ();
 			Parallel2.ForEach(SharedPCHModules, SharedPCHModule =>
 			{
-				using var rented = HashSetPool<UEBuildModule>.Rent();
+				using Pooled<HashSet<UEBuildModule>> rented = HashSetPool<UEBuildModule>.Rent();
 				HashSet<UEBuildModule> Dependencies = rented.Value;
 				SharedPCHModule.GetAllDependencyModulesForPCH(Dependencies, false);
 				Dependencies.IntersectWith(SharedPCHModules);
@@ -4758,7 +4758,7 @@ namespace UnrealBuildTool
 					if (!ModuleCompileEnvironment.bIsBuildingLibrary && Module.Rules.PCHUsage != ModuleRules.PCHUsageMode.NoSharedPCHs)
 					{
 						// Find all the dependencies of this module
-						using var rented = HashSetPool<UEBuildModule>.Rent();
+						using Pooled<HashSet<UEBuildModule>> rented = HashSetPool<UEBuildModule>.Rent();
 						HashSet<UEBuildModule> ReferencedModules = rented.Value;
 						Module.GetAllDependencyModules(new List<UEBuildModule>(), ReferencedModules, bIncludeDynamicallyLoaded: false, bForceCircular: false, bOnlyDirectDependencies: true);
 
@@ -4864,7 +4864,7 @@ namespace UnrealBuildTool
 		protected void AddAllValidModulesToTarget(ILogger Logger)
 		{
 			// Find all the modules that are part of the target
-			using var rented = HashSetPool<string>.Rent();
+			using Pooled<HashSet<string>> rented = HashSetPool<string>.Rent();
 			HashSet<string> ValidModuleNames = rented.Value;
 
 			foreach (UEBuildModuleCPP Module in Modules.Values.OfType<UEBuildModuleCPP>())
@@ -4879,7 +4879,7 @@ namespace UnrealBuildTool
 			IReadOnlySet<string> ExcludeFolders = UEBuildPlatform.GetBuildPlatform(Platform).GetExcludedFolderNames();
 
 			// Set of module names to build
-			using var rented2 = HashSetPool<string>.Rent();
+			using Pooled<HashSet<string>> rented2 = HashSetPool<string>.Rent();
 			HashSet<string> FilteredModuleNames = rented2.Value;
 
 			// Only add engine modules for non-program targets. Programs only compile allowed modules through plugins.
@@ -5046,7 +5046,7 @@ namespace UnrealBuildTool
 			// Now create all the precompiled modules, making sure they don't reference anything that's not in the precompiled set
 
 			// Gather the set of all modules traversed while processing ValidModuleNames, to ensure we have any extra modules found during RecursivelyCreateModules()
-			using var rented3 = HashSetPool<UEBuildModule>.Rent();
+			using Pooled<HashSet<UEBuildModule>> rented3 = HashSetPool<UEBuildModule>.Rent();
 			HashSet<UEBuildModule> AllModules = rented3.Value;
 
 			foreach (string ModuleName in ValidModuleNames)
diff --git a/Engine/Source/Programs/UnrealBuildTool/Modes/BuildMode.cs b/Engine/Source/Programs/UnrealBuildTool/Modes/BuildMode.cs
--- a/Engine/Source/Programs/UnrealBuildTool/Modes/BuildMode.cs
+++ b/Engine/Source/Programs/UnrealBuildTool/Modes/BuildMode.cs
@@ -181,7 +181,7 @@ namespace UnrealBuildTool
 				// Parse all the target descriptors
 				using (GlobalTracer.Instance.BuildSpan("TargetDescriptor.ParseCommandLine()").StartActive())
 				{
-					using var _ = Timeline.ScopeEvent("ParseCommandLine");
+					using ITimelineEvent _ = Timeline.ScopeEvent("ParseCommandLine");
 
 					TargetDescriptors = TargetDescriptor.ParseCommandLine(Arguments, BuildConfiguration, Logger);
 
@@ -448,38 +448,6 @@ namespace UnrealBuildTool
 					cppDependencyCache.Mount(Targets[TargetIdx].Descriptor, Makefiles[TargetIdx].TargetType, Logger);
 				}
 
-#if false // Doesn't seem to make things faster. Revisit since it feels like this should be faster
-				// Prepopulate directories
-				for (int Idx = 0; Idx < Makefiles.Length; Idx++)
-				{
-					TargetDescriptor TargetDescriptor = TargetDescriptors[Idx];
-					TargetMakefile Makefile = Makefiles[Idx];
-					//HashSet<DirectoryItem> handledDir = new();
-					//ParallelQueue queue = new();
-
-					Task throwaway = FileMetadataPrefetch.QueueDirectoryTrees(Makefile.UObjectModules.Select(m => m.GeneratedCodeDirectory));
-					throwaway.Wait();
-					/*
-					foreach (var module in Makefile.UObjectModules)
-					{
-						module.GeneratedCodeDirectory
-					foreach (IExternalAction action in Makefile.Actions)
-					{
-						foreach (FileItem prereq in action.PrerequisiteItems)
-						{
-							DirectoryItem dir = prereq.Directory;
-							if (handledDir.Add(dir))
-								queue.Enqueue(() =>
-								{
-									dir.CacheDirectories();
-								});
-						}
-					}
-					queue.Drain();
-					*/
-				}
-#endif
-
 				using ITimelineEvent ActionGraphTimer = Timeline.ScopeEvent($"Preparing ActionGraph");
 				// Make sure that none of the actions conflict with any other (producing output files differently, etc...)
 				using (GlobalTracer.Instance.BuildSpan("ActionGraph.CheckForConflicts").StartActive())
