[UBT]
* Removed usage of var in favor of explicit type in more files

#FailedValidations Preflight

diff --git a/Engine/Source/Programs/UnrealBuildTool/Configuration/CppCompileWarnings.cs b/Engine/Source/Programs/UnrealBuildTool/Configuration/CppCompileWarnings.cs
--- a/Engine/Source/Programs/UnrealBuildTool/Configuration/CppCompileWarnings.cs
+++ b/Engine/Source/Programs/UnrealBuildTool/Configuration/CppCompileWarnings.cs
@@ -1512,7 +1512,7 @@ namespace UnrealBuildTool
 			List<string> warningArguments = [];
 			List<ValueTuple<ApplyWarningsAttribute, WarningLevel>> warningAttributesBuffer = [];
 
-			foreach (var (mi, miAttr) in s_membersWithApplyWarningsAttribute.Value)
+			foreach ((MemberInfo mi, List<ApplyWarningsAttribute> miAttr) in s_membersWithApplyWarningsAttribute.Value)
 			{
 				ApplyWarningsAttribute? applyWarningsAttribute = null;
 
@@ -1570,7 +1570,7 @@ namespace UnrealBuildTool
 		{
 			cppCompileWarnings._logger?.LogTrace("Applying default C++ Compiler warnings for the Target({TargetName}) context.", cppCompileWarnings.BuildContextProvider?._targetRulesPrivate?.Name);
 
-			foreach (var (mi, miAttr) in s_membersWithWarningLevelDefaultAttribute.Value)
+			foreach ((MemberInfo mi, List<WarningLevelDefaultAttribute> miAttr) in s_membersWithWarningLevelDefaultAttribute.Value)
 			{
 				WarningLevel defaultWarningLevel = WarningLevel.Default;
 				defaultWarningLevel = CoalesceDefaultWarningFromMemberInfo(mi, miAttr, InitializationContext.Target, defaultWarningLevel, cppCompileWarnings.BuildContextProvider, cppCompileWarnings._logger);
@@ -1663,7 +1663,7 @@ namespace UnrealBuildTool
 
 			BuildSettingsVersion buildSettingsVersion = BuildContextProvider != null ? BuildContextProvider._buildContext.GetBuildSettings() : BuildSettingsVersion.V1;
 
-			foreach (var (mi, miAttr) in s_membersWithWarningLevelDefaultAttribute.Value)
+			foreach ((MemberInfo mi, List<WarningLevelDefaultAttribute> miAttr) in s_membersWithWarningLevelDefaultAttribute.Value)
 			{
 				WarningLevel defaultWarningLevel = WarningLevel.Default;
 				defaultWarningLevel = CoalesceDefaultWarningFromMemberInfo(mi, miAttr, InitializationContext.Constructor, defaultWarningLevel, BuildContextProvider, _logger);
diff --git a/Engine/Source/Programs/UnrealBuildTool/Configuration/ModuleRules.cs b/Engine/Source/Programs/UnrealBuildTool/Configuration/ModuleRules.cs
--- a/Engine/Source/Programs/UnrealBuildTool/Configuration/ModuleRules.cs
+++ b/Engine/Source/Programs/UnrealBuildTool/Configuration/ModuleRules.cs
@@ -2125,7 +2125,7 @@ namespace UnrealBuildTool
 		/// <param name="items"></param>
 		public static void AddRange<T>(this HashSet<T> set, params T[] items)
 		{
-			foreach (var item in items) set.Add(item);
+			foreach (T? item in items) set.Add(item);
 		}
 	}
 }
\ No newline at end of file
diff --git a/Engine/Source/Programs/UnrealBuildTool/Configuration/TargetRules.cs b/Engine/Source/Programs/UnrealBuildTool/Configuration/TargetRules.cs
--- a/Engine/Source/Programs/UnrealBuildTool/Configuration/TargetRules.cs
+++ b/Engine/Source/Programs/UnrealBuildTool/Configuration/TargetRules.cs
@@ -3688,14 +3688,14 @@ namespace UnrealBuildTool
 		/// <returns></returns>
 		internal static List<PropertyInfo> GetSubObjectProperties(Type t) => s_subPropsCache.GetOrAdd(t, static type => new Lazy<List<PropertyInfo>>(() =>
 		{
-			var list = new List<PropertyInfo>();
+			List<PropertyInfo> list = new();
 
 			if (type.BaseType != null)
 			{
 				list.AddRange(GetSubObjectProperties(type.BaseType));
 			}
 
-			var props = type.GetProperties(BindingFlags.Instance | BindingFlags.Public | BindingFlags.DeclaredOnly);
+			PropertyInfo[] props = type.GetProperties(BindingFlags.Instance | BindingFlags.Public | BindingFlags.DeclaredOnly);
 			foreach (PropertyInfo pi in props)
 			{
 				if (pi.IsDefined(typeof(ConfigSubObjectAttribute), inherit: false))
diff --git a/Engine/Source/Programs/UnrealBuildTool/Configuration/UEBuildPlatform.cs b/Engine/Source/Programs/UnrealBuildTool/Configuration/UEBuildPlatform.cs
--- a/Engine/Source/Programs/UnrealBuildTool/Configuration/UEBuildPlatform.cs
+++ b/Engine/Source/Programs/UnrealBuildTool/Configuration/UEBuildPlatform.cs
@@ -694,9 +694,13 @@ namespace UnrealBuildTool
 		public static IEnumerable<UnrealPlatformGroup> GetPlatformGroups(UnrealTargetPlatform Platform)
 		{
 			// Avoid allocations by using yield (one allocation for enumerable still created)
-			foreach (var kvp in PlatformGroupDictionary)
+			foreach (KeyValuePair<UnrealPlatformGroup, List<UnrealTargetPlatform>> kvp in PlatformGroupDictionary)
+			{
 				if (kvp.Value.Contains(Platform))
+				{
 					yield return kvp.Key;
+				}
+			}
 		}
 
 		/// <summary>
diff --git a/Engine/Source/Programs/UnrealBuildTool/Modes/CleanMode.cs b/Engine/Source/Programs/UnrealBuildTool/Modes/CleanMode.cs
--- a/Engine/Source/Programs/UnrealBuildTool/Modes/CleanMode.cs
+++ b/Engine/Source/Programs/UnrealBuildTool/Modes/CleanMode.cs
@@ -60,7 +60,7 @@ namespace UnrealBuildTool
 			using ScopedTimer CleanTimer = new ScopedTimer("CleanMode.Clean()", Logger);
 			using IScope Scope = GlobalTracer.Instance.BuildSpan("CleanMode.Clean()").StartActive();
 
-			using var _ = Timeline.ScopeEvent("Cleaning");
+			using ITimelineEvent _ = Timeline.ScopeEvent("Cleaning");
 
 			if (TargetDescriptors.Count == 0)
 			{
diff --git a/Engine/Source/Programs/UnrealBuildTool/Modes/IWYUMode.cs b/Engine/Source/Programs/UnrealBuildTool/Modes/IWYUMode.cs
--- a/Engine/Source/Programs/UnrealBuildTool/Modes/IWYUMode.cs
+++ b/Engine/Source/Programs/UnrealBuildTool/Modes/IWYUMode.cs
@@ -979,7 +979,7 @@ namespace UnrealBuildTool
 							if (fixedPath.StartsWith(Dir.FullName))
 							{
 								string printable = value.File.Substring(Dir.FullName.Length + 1);
-								var entry = new IWYUIncludeEntry() { Full = MissingInclude.Full, Printable = printable, System = false, Resolved = value, };
+								IWYUIncludeEntry entry = new IWYUIncludeEntry() { Full = MissingInclude.Full, Printable = printable, System = false, Resolved = value, };
 								Info.Includes.Add(entry);
 								break;
 							}
diff --git a/Engine/Source/Programs/UnrealBuildTool/Platform/Android/UEDeployAndroid.cs b/Engine/Source/Programs/UnrealBuildTool/Platform/Android/UEDeployAndroid.cs
--- a/Engine/Source/Programs/UnrealBuildTool/Platform/Android/UEDeployAndroid.cs
+++ b/Engine/Source/Programs/UnrealBuildTool/Platform/Android/UEDeployAndroid.cs
@@ -1452,7 +1452,7 @@ namespace UnrealBuildTool
 			{
 				IEnumerable<string> SOFiles = Directory.EnumerateFiles(FolderPath, "*.so", SearchOption.AllDirectories);
 
-				foreach (var SOFile in SOFiles)
+				foreach (string SOFile in SOFiles)
 				{
 					if (!IsSO16KbCompatible(SOFile, UnrealArch))
 					{
diff --git a/Engine/Source/Programs/UnrealBuildTool/Platform/Apple/UEBuildApple.cs b/Engine/Source/Programs/UnrealBuildTool/Platform/Apple/UEBuildApple.cs
--- a/Engine/Source/Programs/UnrealBuildTool/Platform/Apple/UEBuildApple.cs
+++ b/Engine/Source/Programs/UnrealBuildTool/Platform/Apple/UEBuildApple.cs
@@ -227,7 +227,7 @@ namespace UnrealBuildTool
 			
 			CompileEnvironment.Definitions.Add("PLATFORM_APPLE=1");
 
-			AppleExports.GetSwiftIntegrationSettings(Target.ProjectFile, Target.Type, Platform, out var bUseSwiftUIMain, out _);
+			AppleExports.GetSwiftIntegrationSettings(Target.ProjectFile, Target.Type, Platform, out bool bUseSwiftUIMain, out _);
 			CompileEnvironment.Definitions.Add("UE_USE_SWIFT_UI_MAIN=" + (bUseSwiftUIMain ? "1" : "0"));
 
 			CompileEnvironment.Definitions.Add("WITH_TTS=0");
diff --git a/Engine/Source/Programs/UnrealBuildTool/Platform/Linux/LinuxToolChain.cs b/Engine/Source/Programs/UnrealBuildTool/Platform/Linux/LinuxToolChain.cs
--- a/Engine/Source/Programs/UnrealBuildTool/Platform/Linux/LinuxToolChain.cs
+++ b/Engine/Source/Programs/UnrealBuildTool/Platform/Linux/LinuxToolChain.cs
@@ -889,7 +889,7 @@ namespace UnrealBuildTool
 			List<string> GlobalArguments = [];
 			GetCompileArguments_Global(CompileEnvironment, GlobalArguments);
 
-			//var BuildPlatform = UEBuildPlatform.GetBuildPlatform(CompileEnvironment.Platform);
+			//UEBuildPlatform BuildPlatform = UEBuildPlatform.GetBuildPlatform(CompileEnvironment.Platform);
 
 			if (!bHasPrintedBuildDetails)
 			{
diff --git a/Engine/Source/Programs/UnrealBuildTool/ProjectFiles/ProjectFileGenerator.cs b/Engine/Source/Programs/UnrealBuildTool/ProjectFiles/ProjectFileGenerator.cs
--- a/Engine/Source/Programs/UnrealBuildTool/ProjectFiles/ProjectFileGenerator.cs
+++ b/Engine/Source/Programs/UnrealBuildTool/ProjectFiles/ProjectFileGenerator.cs
@@ -556,7 +556,7 @@ namespace UnrealBuildTool
 			{
 				if (FileReference.Exists(ProjectFile))
 				{
-					using var _ = Timeline.ScopeEvent(ProjectFile.GetFileName());
+					using ITimelineEvent _ = Timeline.ScopeEvent(ProjectFile.GetFileName());
 					VCSharpProjectFile Project = new VCSharpProjectFile(ProjectFile, Logger);
 					lock (projectFiles)
 					{
@@ -1540,7 +1540,7 @@ namespace UnrealBuildTool
 				int TasksFinished = 0;
 				Parallel2.ForEach(ProjectFileToRules, Item =>
 				{
-					using var t = Timeline.ScopeEvent(Item.Key?.GetFileName()??"Unknown");
+					using ITimelineEvent t = Timeline.ScopeEvent(Item.Key?.GetFileName() ?? "Unknown");
 					QueryTargetsMode.WriteTargetInfo(Item.Key, Item.Value, QueryTargetsMode.GetDefaultOutputFile(Item.Key), new CommandLineArguments(GetTargetArguments(Arguments)), Logger);
 					Progress.Write(Interlocked.Increment(ref TasksFinished), ProjectFileToRules.Count);
 				});
@@ -2384,7 +2384,7 @@ namespace UnrealBuildTool
 
 					void CreateBuildTarget(ProjectFile TargetProjectFile, ProjectTarget CurTarget)
 					{
-						using var _ = Timeline.ScopeEvent(CurTarget.Name);
+						using ITimelineEvent _ = Timeline.ScopeEvent(CurTarget.Name);
 
 						List<string> NewArguments = new List<string>(Arguments.Length + 4);
 						if (CurTarget.TargetRules!.Type != TargetType.Program)
@@ -2441,7 +2441,7 @@ namespace UnrealBuildTool
 
 					void Gen(UEBuildTarget Target)
 					{
-						using var _ = Timeline.ScopeEvent(Target.TargetName);
+						using ITimelineEvent _ = Timeline.ScopeEvent(Target.TargetName);
 						AddTargetForIntellisense(Target, Logger);
 
 						// If the project generator just cares about the result of UEBuildTarget.Create, skip generating the compile environments.
@@ -2471,7 +2471,7 @@ namespace UnrealBuildTool
 										}
 										CppCompileEnvironment ModuleCompileEnvironment = Module.CreateCompileEnvironmentForIntellisense(Target.Rules, BinaryCompileEnvironment, Logger, builder);
 
-										var list = toAdd.GetOrAdd(ProjectFileForIDE, _ => new());
+										ConcurrentBag<(UEBuildModuleCPP, CppCompileEnvironment)> list = toAdd.GetOrAdd(ProjectFileForIDE, _ => new());
 										list.Add((Module, ModuleCompileEnvironment));
 									}
 									Progress.Write(Interlocked.Increment(ref TasksFinished), TotalTasks);
@@ -2488,14 +2488,14 @@ namespace UnrealBuildTool
 						Gen(Target);
 					});
 
-					using var _ = Timeline.ScopeEvent("AddModules");
+					using ITimelineEvent _ = Timeline.ScopeEvent("AddModules");
 					Parallel2.ForEach(toAdd, kv =>
 					{
 						// We have to sort this array because there are multiple .build.cs that shares same Source folder (WireInterface among others)
 						// And the source files end up under different module if this list is not sorted.
-						var moduleAndEnv = kv.Value.ToArray();
+						(UEBuildModuleCPP, CppCompileEnvironment)[] moduleAndEnv = kv.Value.ToArray();
 						Array.Sort(moduleAndEnv, static (a, b) => a.Item1.IntermediateDirectory.CompareTo(b.Item1.IntermediateDirectory));
-						foreach (var (module, compileEnv) in moduleAndEnv)
+						foreach ((UEBuildModuleCPP module, CppCompileEnvironment compileEnv) in moduleAndEnv)
 						{
 							kv.Key.AddModuleForIntelliSense(module, compileEnv);
 						}
@@ -2691,7 +2691,7 @@ namespace UnrealBuildTool
 
 			Parallel2.ForEach(AllModuleFiles, (CurModuleFile) =>
 			{
-				using var _ = Timeline.ScopeEvent(CurModuleFile.GetFileName());
+				using ITimelineEvent _ = Timeline.ScopeEvent(CurModuleFile.GetFileName());
 
 				// The module's "base directory" is simply the directory where its xxx.Build.cs file is stored.  We'll always
 				// harvest source files for this module in this base directory directory and all of its sub-directories.
@@ -2808,7 +2808,7 @@ namespace UnrealBuildTool
 			});
 
 			{
-				using var _ = Timeline.ScopeEvent("AddPlugins");
+				using ITimelineEvent _ = Timeline.ScopeEvent("AddPlugins");
 				Parallel2.ForEach(PluginsToAdd, (Rec) =>
 				{
 					AddPluginFilesToProject(Rec.Item1, Rec.Item2, Rec.Item3);
@@ -3079,7 +3079,7 @@ namespace UnrealBuildTool
 			// First step
 			Parallel2.ForEach(AllTargetsExceptSubTargets, (TargetFilePath) =>
 			{
-				using var _ = Timeline.ScopeEvent(TargetFilePath.GetFileNameWithoutAnyExtensions());
+				using ITimelineEvent _ = Timeline.ScopeEvent(TargetFilePath.GetFileNameWithoutAnyExtensions());
 
 				TargetRules TargetRulesObject;
 				string ProjectFileNameBase;
@@ -3662,7 +3662,7 @@ namespace UnrealBuildTool
 						return;
 					}
 
-					using var _ = Timeline.ScopeEvent(CurProject.ProjectFilePath.GetFileName());
+					using ITimelineEvent _ = Timeline.ScopeEvent(CurProject.ProjectFilePath.GetFileName());
 
 					if (!CurProject.WriteProjectFile(SupportedPlatforms, SupportedConfigurations, PlatformProjectGenerators, Logger))
 					{
diff --git a/Engine/Source/Programs/UnrealBuildTool/ProjectFiles/VisualStudioWorkspace/VSWorkspaceProjectFile.cs b/Engine/Source/Programs/UnrealBuildTool/ProjectFiles/VisualStudioWorkspace/VSWorkspaceProjectFile.cs
--- a/Engine/Source/Programs/UnrealBuildTool/ProjectFiles/VisualStudioWorkspace/VSWorkspaceProjectFile.cs
+++ b/Engine/Source/Programs/UnrealBuildTool/ProjectFiles/VisualStudioWorkspace/VSWorkspaceProjectFile.cs
@@ -123,14 +123,14 @@ namespace UnrealBuildTool
 
 					try
 					{
-						foreach (var Architecture in TargetDesc.Architectures.Architectures)
+						foreach (UnrealArch Architecture in TargetDesc.Architectures.Architectures)
 						{
 							UEBuildTarget BuildTarget = UEBuildTarget.Create(TargetDesc, false, false, false, UnrealIntermediateEnvironment.GenerateProjectFiles, Logger);
 							BuildTarget.PreBuildSetup(Logger);
 
 							ExportedTargetInfo TargetInfo = ExportTarget(BuildTarget, Architecture, bBuildByDefault, PlatformProjectGenerators, Logger, VCFormat);
 
-							var sb = new StringBuilder();
+							StringBuilder sb = new();
 							sb.Append(ProjectTarget.TargetFilePath.GetFileNameWithoutAnyExtensions());
 							sb.Append("_");
 							sb.Append(TargetInfo.Configuration);
@@ -468,7 +468,7 @@ namespace UnrealBuildTool
 
 				if (Architecture != null)
 				{
-					var result = Architecture.ToString();
+					string? result = Architecture.ToString();
 					return result != null ? result : string.Empty;
 				}
 			}
diff --git a/Engine/Source/Programs/UnrealBuildTool/ProjectFiles/VisualStudioWorkspace/VSWorkspaceProjectFileGenerator.cs b/Engine/Source/Programs/UnrealBuildTool/ProjectFiles/VisualStudioWorkspace/VSWorkspaceProjectFileGenerator.cs
--- a/Engine/Source/Programs/UnrealBuildTool/ProjectFiles/VisualStudioWorkspace/VSWorkspaceProjectFileGenerator.cs
+++ b/Engine/Source/Programs/UnrealBuildTool/ProjectFiles/VisualStudioWorkspace/VSWorkspaceProjectFileGenerator.cs
@@ -105,7 +105,7 @@ namespace UnrealBuildTool
 
 		protected override bool WriteProjectFiles(PlatformProjectGeneratorCollection PlatformProjectGenerators, ILogger Logger)
 		{
-			var ProjectGenerationDescriptionData = ReadProjectGenerateDescriptionFilePath()?.Targets;
+			Dictionary<string, AvailableConfiguration>? ProjectGenerationDescriptionData = ReadProjectGenerateDescriptionFilePath()?.Targets;
 			using ProgressWriter Progress = new("Writing project files...", true, Logger);
 			List<ProjectFile> ProjectsToGenerate = new(GeneratedProjectFiles);
 			if (ProjectNames.Any())
@@ -136,7 +136,7 @@ namespace UnrealBuildTool
 
 			List<UnrealTargetConfiguration> Configurations = ConfigurationsToGenerate.ToList();
 			VCProjectFileFormat VCFormat = GetVCProjectFileFormat(Logger);
-			var wroteVCProjectFiles = false;
+			bool wroteVCProjectFiles = false;
 
 			for (int ProjectFileIndex = 0; ProjectFileIndex < ProjectsToGenerate.Count; ++ProjectFileIndex)
 			{
@@ -150,7 +150,7 @@ namespace UnrealBuildTool
 					return false;
 				}
 
-				foreach (var Target in CurrentProject.ProjectTargets.OfType<ProjectTarget>())
+				foreach (ProjectTarget Target in CurrentProject.ProjectTargets.OfType<ProjectTarget>())
 				{
 					if (!TargetTypes.Any() || (Target.TargetRules != null && TargetTypes.Contains(Target.TargetRules.Type)))
 					{
@@ -512,12 +512,12 @@ namespace UnrealBuildTool
 					FilteredPlatforms = new List<UnrealTargetPlatform>();
 					FilteredConfigurations = new List<UnrealTargetConfiguration>();
 
-					foreach (var Platform in Platforms)
+					foreach (UnrealTargetPlatform Platform in Platforms)
 					{
 						IncludePlatform(AvailableConfiguration, Platform, FilteredPlatforms);
 					}
 
-					foreach (var Configuration in Configurations)
+					foreach (UnrealTargetConfiguration Configuration in Configurations)
 					{
 						IncludeConfiguration(AvailableConfiguration, Configuration, FilteredConfigurations);
 					}
@@ -528,10 +528,10 @@ namespace UnrealBuildTool
 					return true;
 				}
 
-				var Settings = new VCProjectFileSettings();
+				VCProjectFileSettings Settings = new();
 				Settings.MaxSharedIncludePaths = 0;
 
-				var VCProject = new VCProjectFile(Output, RootPath, Format, false, true, null, Settings, true);
+				VCProjectFile VCProject = new(Output, RootPath, Format, false, true, null, Settings, true);
 				VCProject.ProjectTargets.Add(Target);
 				if (VCProject.WriteProjectFile(FilteredPlatforms, FilteredConfigurations, PlatformProjectGenerators, Logger))
 				{
@@ -546,7 +546,7 @@ namespace UnrealBuildTool
 					while (Nodes.Count != 0)
 					{
 						XmlNode Node = Nodes.Dequeue();
-						using (var NodeList = Node.ChildNodes)
+						using (XmlNodeList NodeList = Node.ChildNodes)
 						{
 							if (NodeList != null)
 							{
@@ -566,7 +566,7 @@ namespace UnrealBuildTool
 										}
 										else
 										{
-											foreach (var RegEx in BlockedNodesListRegEx)
+											foreach (Regex RegEx in BlockedNodesListRegEx)
 											{
 												if (RegEx.Match(ChildNode.Name).Success)
 												{
@@ -590,7 +590,7 @@ namespace UnrealBuildTool
 							}
 						}
 
-						foreach (var ChildNode in RemovedChildNodes)
+						foreach (XmlNode ChildNode in RemovedChildNodes)
 						{
 							Node.RemoveChild(ChildNode);
 						}
diff --git a/Engine/Source/Programs/UnrealBuildTool/ProjectFiles/Xcode/XcodeProject.cs b/Engine/Source/Programs/UnrealBuildTool/ProjectFiles/Xcode/XcodeProject.cs
--- a/Engine/Source/Programs/UnrealBuildTool/ProjectFiles/Xcode/XcodeProject.cs
+++ b/Engine/Source/Programs/UnrealBuildTool/ProjectFiles/Xcode/XcodeProject.cs
@@ -1413,7 +1413,7 @@ namespace UnrealBuildTool.XcodeProjectXcconfig
 							if (File.Exists(UExtPath))
 							{
 								string JsonText = File.ReadAllText(UExtPath);
-								var Options = new JsonSerializerOptions{ PropertyNameCaseInsensitive = true };
+								JsonSerializerOptions Options = new() { PropertyNameCaseInsensitive = true };
 								ExtConfig = JsonSerializer.Deserialize<ExtensionConfig>(JsonText, Options);
 								if (ExtConfig == null)
 								{
diff --git a/Engine/Source/Programs/UnrealBuildTool/System/AppDomainExtensions.cs b/Engine/Source/Programs/UnrealBuildTool/System/AppDomainExtensions.cs
--- a/Engine/Source/Programs/UnrealBuildTool/System/AppDomainExtensions.cs
+++ b/Engine/Source/Programs/UnrealBuildTool/System/AppDomainExtensions.cs
@@ -21,11 +21,11 @@ namespace UnrealBuildBase
 		/// <param name="baseClass">Base class to use</param>
 		public static IEnumerable<Type> EnumerateSubclasses(this AppDomain domain, Type baseClass)
 		{
-			var baseAsm = baseClass.Assembly;
-			var baseAsmName = baseAsm.GetName();
+			Assembly baseAsm = baseClass.Assembly;
+			AssemblyName baseAsmName = baseAsm.GetName();
 
 			// Only search in assemblies that are either baseClass's assembly or assemblies depending on it
-			foreach (var asm in domain.GetAssemblies())
+			foreach (Assembly asm in domain.GetAssemblies())
 			{
 				if (asm.IsDynamic)
 				{
@@ -62,7 +62,7 @@ namespace UnrealBuildBase
 
 				for (int i = 0; i < types.Length; i++)
 				{
-					var t = types[i];
+					Type t = types[i];
 					if (t != null && t.IsClass && !t.IsAbstract && t.IsSubclassOf(baseClass))
 					{
 						yield return t;
diff --git a/Engine/Source/Programs/UnrealBuildTool/System/BloomFilter.cs b/Engine/Source/Programs/UnrealBuildTool/System/BloomFilter.cs
--- a/Engine/Source/Programs/UnrealBuildTool/System/BloomFilter.cs
+++ b/Engine/Source/Programs/UnrealBuildTool/System/BloomFilter.cs
@@ -80,7 +80,7 @@ namespace UnrealBuildTool
 
 				if (separatorIndex != 0)
 				{
-					var dir = span.Slice(0, separatorIndex);
+					ReadOnlySpan<char> dir = span.Slice(0, separatorIndex);
 
 					if (MightContain(dir))
 					{
diff --git a/Engine/Source/Programs/UnrealBuildTool/System/CppDependencyCache.cs b/Engine/Source/Programs/UnrealBuildTool/System/CppDependencyCache.cs
--- a/Engine/Source/Programs/UnrealBuildTool/System/CppDependencyCache.cs
+++ b/Engine/Source/Programs/UnrealBuildTool/System/CppDependencyCache.cs
@@ -288,14 +288,14 @@ namespace UnrealBuildTool
 
 			CachePartition? partition = null;
 
-			foreach (var root in GlobalPartitions)
+			foreach ((DirectoryReference Root, List<(string Filter, CachePartition Partition)>) root in GlobalPartitions)
 			{
 				if (!InputFile.Location.IsUnderDirectory(root.Root))
 				{
 					continue;
 				}
 
-				foreach (var entry in root.Item2)
+				foreach ((string Filter, CachePartition Partition) entry in root.Item2)
 				{
 					if (filterPath.StartsWith(entry.Filter, FileReference.Comparison))
 					{
@@ -435,14 +435,14 @@ namespace UnrealBuildTool
 				CachePartition? partition = null;
 				lock (GlobalPartitions)
 				{
-					var root = GlobalPartitions.FirstOrDefault(e => e.Root == rootPath);
+					(DirectoryReference Root, List<(string Filter, CachePartition Partition)>) root = GlobalPartitions.FirstOrDefault(e => e.Root == rootPath);
 					if (root.Root == null)
 					{
 						root = new(rootPath, new());
 						GlobalPartitions.Add(root);
 					}
-					var filters = root.Item2;
-					var filter = filters.FirstOrDefault(f => f.Filter == filterPath);
+					List<(string Filter, CachePartition Partition)> filters = root.Item2;
+					(string Filter, CachePartition Partition) filter = filters.FirstOrDefault(f => f.Filter == filterPath);
 					partition = filter.Partition;
 					if (partition == null)
 					{
diff --git a/Engine/Source/Programs/UnrealBuildTool/System/DynamicCompilation.cs b/Engine/Source/Programs/UnrealBuildTool/System/DynamicCompilation.cs
--- a/Engine/Source/Programs/UnrealBuildTool/System/DynamicCompilation.cs
+++ b/Engine/Source/Programs/UnrealBuildTool/System/DynamicCompilation.cs
@@ -172,21 +172,21 @@ namespace UnrealBuildTool
 
 		private static SyntaxTree ParseSyntaxTree(FileReference sourceFileName, CSharpParseOptions parseOptions, CancellationToken cancellationToken)
 		{
-			using var fs = File.OpenRead(sourceFileName.FullName);
-			var text = SourceText.From(fs, encoding: Encoding.UTF8, canBeEmbedded: true);
+			using FileStream fs = File.OpenRead(sourceFileName.FullName);
+			SourceText text = SourceText.From(fs, encoding: Encoding.UTF8, canBeEmbedded: true);
 			return CSharpSyntaxTree.ParseText(text, parseOptions, path: sourceFileName.FullName, cancellationToken: cancellationToken);
 		}
 
 		private static async Task<IEnumerable<SyntaxTree>> ParseSyntaxTreesAsync(IEnumerable<FileReference> sourceFileNames, IEnumerable<string>? preprocessorDefines, CancellationToken cancellationToken)
 		{
-			var parseOptions = new CSharpParseOptions(
+			CSharpParseOptions parseOptions = new(
 				 languageVersion: LanguageVersion.CSharp12, // Prefer specific version, faster than using Latest due to version probing
 				 kind: SourceCodeKind.Regular,
 				 preprocessorSymbols: preprocessorDefines,
 				 documentationMode: DocumentationMode.None);
 
-			var trees = new ConcurrentBag<SyntaxTree>();
-			var dop = Math.Max(1, Environment.ProcessorCount - 1);
+			ConcurrentBag<SyntaxTree> trees = new();
+			int dop = Math.Max(1, Environment.ProcessorCount - 1);
 
 			await Parallel.ForEachAsync(sourceFileNames, new ParallelOptions { MaxDegreeOfParallelism = dop, CancellationToken = cancellationToken },
 				(f, token) =>
diff --git a/Engine/Source/Programs/UnrealBuildTool/System/FileHashCache.cs b/Engine/Source/Programs/UnrealBuildTool/System/FileHashCache.cs
--- a/Engine/Source/Programs/UnrealBuildTool/System/FileHashCache.cs
+++ b/Engine/Source/Programs/UnrealBuildTool/System/FileHashCache.cs
@@ -140,7 +140,7 @@ namespace UnrealBuildTool
 		/// <exception cref="Exception"></exception>
 		public static UInt128 GetFileHash(MemoryStream stream)
 		{
-			var xx = new XxHash128();
+			XxHash128 xx = new();
 
 			stream.Position = 0;
 			xx.Append(stream);
@@ -166,7 +166,7 @@ namespace UnrealBuildTool
 		/// </summary>
 		public Task? WriteFileIfChanged(FileItem fileItem, IEnumerable<string> content, ILogger logger, bool allowAsync)
 		{
-			var xx = new XxHash128();
+			XxHash128 xx = new();
 			foreach (string str in content)
 			{
 				xx.Append(MemoryMarshal.AsBytes(str.AsSpan()));
@@ -315,10 +315,10 @@ namespace UnrealBuildTool
 						{
 							const int BufferSize = 1024 * 1024; // 1 MB
 							byte[] buffer = new byte[BufferSize];
-							using var fs = new FileStream(fileItem.FullName, FileMode.Open, FileAccess.Read, FileShare.Read, BufferSize, FileOptions.SequentialScan);
+							using FileStream fs = new(fileItem.FullName, FileMode.Open, FileAccess.Read, FileShare.Read, BufferSize, FileOptions.SequentialScan);
 							int read;
 
-							var xx = new XxHash128();
+							XxHash128 xx = new();
 							while ((read = fs.Read(buffer, 0, buffer.Length)) > 0)
 								xx.Append(buffer.AsSpan(0, read));
 							isEqual = newHash == ToUInt128LittleEndian(xx.GetHashAndReset());
@@ -411,20 +411,20 @@ namespace UnrealBuildTool
 					if (data is MemoryStream stream)
 					{
 						stream.Position = 0;
-						using var file = new FileStream(fileItem.FullName, FileMode.Create, FileAccess.Write, FileShare.None);
+						using FileStream file = new(fileItem.FullName, FileMode.Create, FileAccess.Write, FileShare.None);
 						stream.CopyTo(file);
 					}
 					else if (data is string content)
 					{
-						using var fs = new FileStream(fileItem.FullName, FileMode.Create, FileAccess.Write, FileShare.Read, 1 << 20, FileOptions.SequentialScan);
-						using var writer = new StreamWriter(fs, new UTF8Encoding(false), 1 << 20, false);
+						using FileStream fs = new(fileItem.FullName, FileMode.Create, FileAccess.Write, FileShare.Read, 1 << 20, FileOptions.SequentialScan);
+						using StreamWriter writer = new(fs, new UTF8Encoding(false), 1 << 20, false);
 						writer.Write(content);
 					}
 					else if (data is IEnumerable<string> contentLines)
 					{
 						string contents = String.Join(Environment.NewLine, contentLines);
-						using var fs = new FileStream(fileItem.FullName, FileMode.Create, FileAccess.Write, FileShare.Read, 1 << 20, FileOptions.SequentialScan);
-						using var writer = new StreamWriter(fs, new UTF8Encoding(false), 1 << 20, false);
+						using FileStream fs = new(fileItem.FullName, FileMode.Create, FileAccess.Write, FileShare.Read, 1 << 20, FileOptions.SequentialScan);
+						using StreamWriter writer = new(fs, new UTF8Encoding(false), 1 << 20, false);
 						writer.Write(contents);
 					}
 					if (isLocked)
@@ -470,14 +470,14 @@ namespace UnrealBuildTool
 			Partition? partition = null;
 			lock (_roots)
 			{
-				var root = _roots.FirstOrDefault(e => e.root == rootPath);
+				(DirectoryReference root, List<(string filter, Partition partition)>) root = _roots.FirstOrDefault(e => e.root == rootPath);
 				if (root.root == null)
 				{
 					root = new(rootPath, new());
 					_roots.Add(root);
 				}
-				var filters = root.Item2;
-				var filter = filters.FirstOrDefault(f => f.filter == filterPath);
+				List<(string filter, Partition partition)> filters = root.Item2;
+				(string filter, Partition partition) filter = filters.FirstOrDefault(f => f.filter == filterPath);
 				partition = filter.partition;
 				if (partition == null)
 				{
@@ -545,14 +545,14 @@ namespace UnrealBuildTool
 
 			Partition? partition = null;
 
-			foreach (var root in _roots)
+			foreach ((DirectoryReference root, List<(string filter, Partition partition)>) root in _roots)
 			{
 				if (root.root.FullName.Length != 0 && !fileRef.IsUnderDirectory(root.root))
 				{
 					continue;
 				}
 
-				foreach (var entry in root.Item2)
+				foreach ((string filter, Partition partition) entry in root.Item2)
 				{
 					if (filterPath.StartsWith(entry.filter, FileReference.Comparison))
 					{
diff --git a/Engine/Source/Programs/UnrealBuildTool/System/FileMetadataPrefetch.cs b/Engine/Source/Programs/UnrealBuildTool/System/FileMetadataPrefetch.cs
--- a/Engine/Source/Programs/UnrealBuildTool/System/FileMetadataPrefetch.cs
+++ b/Engine/Source/Programs/UnrealBuildTool/System/FileMetadataPrefetch.cs
@@ -43,7 +43,7 @@ namespace UnrealBuildTool
 			{
 				return Task.Run(() =>
 				{
-					//using var _ = Timeline.ScopeEvent("PrefetchEngineDir");
+					//using ITimelineEvent _ = Timeline.ScopeEvent("PrefetchEngineDir");
 					ParallelQueue queue = new();
 					queue.Enqueue(() => ScanEngineDirectory(queue));
 					queue.Drain();
@@ -65,7 +65,7 @@ namespace UnrealBuildTool
 			{
 				return Task.Run(() =>
 				{
-					//using var _ = Timeline.ScopeEvent("PrefetchProjectDir");
+					//using ITimelineEvent _ = Timeline.ScopeEvent("PrefetchProjectDir");
 					ParallelQueue queue = new();
 					queue.Enqueue(() => ScanProjectDirectory(queue, DirectoryItem.GetItemByDirectoryReference(ProjectDirectory)));
 					queue.Drain();
@@ -85,7 +85,7 @@ namespace UnrealBuildTool
 #if UBT_PREFETCH_ENABLED
 			return Task.Run(() =>
 			{
-				//using var _ = Timeline.ScopeEvent("PrefetchDir");
+				//using ITimelineEvent _ = Timeline.ScopeEvent("PrefetchDir");
 				ParallelQueue queue = new();
 				foreach (DirectoryItem? dir in Directories)
 				{
@@ -105,7 +105,7 @@ namespace UnrealBuildTool
 		public static void Stop()
 		{
 #if UBT_PREFETCH_ENABLED
-			foreach (var  directory in QueuedDirectories.Values)
+			foreach (Lazy<Task> directory in QueuedDirectories.Values)
 			{
 				directory.Value.Wait();
 			}
diff --git a/Engine/Source/Programs/UnrealBuildTool/System/HashSetPool.cs b/Engine/Source/Programs/UnrealBuildTool/System/HashSetPool.cs
--- a/Engine/Source/Programs/UnrealBuildTool/System/HashSetPool.cs
+++ b/Engine/Source/Programs/UnrealBuildTool/System/HashSetPool.cs
@@ -28,7 +28,7 @@ static class HashSetPool<T>
 
 	public static Pooled<HashSet<T>> Rent(int ensureCapacity = 0)
 	{
-		var set = Shared.Get();
+		HashSet<T> set = Shared.Get();
 		if (ensureCapacity > 0)
 		{
 			set.EnsureCapacity(ensureCapacity);
diff --git a/Engine/Source/Programs/UnrealBuildTool/System/HotReload.cs b/Engine/Source/Programs/UnrealBuildTool/System/HotReload.cs
--- a/Engine/Source/Programs/UnrealBuildTool/System/HotReload.cs
+++ b/Engine/Source/Programs/UnrealBuildTool/System/HotReload.cs
@@ -1432,7 +1432,7 @@ namespace UnrealBuildTool
 
 							if (Action.RootPaths.bUseVfs)
 							{
-								foreach (var RootPath in Action.RootPaths)
+								foreach ((uint id, DirectoryReference vfs, DirectoryReference local) RootPath in Action.RootPaths)
 								{
 									if (vfsPathsAdded.Add(RootPath.vfs))
 									{
@@ -1446,7 +1446,7 @@ namespace UnrealBuildTool
 				Writer.WriteArrayEnd();
 
 				Writer.WriteArrayStart("Vfs");
-				foreach (var kv in vfsPaths)
+				foreach ((DirectoryReference virtualPath, DirectoryReference localPath) kv in vfsPaths)
 				{
 					Writer.WriteObjectStart();
 					Writer.WriteValue("Virtual", kv.virtualPath.FullName);
diff --git a/Engine/Source/Programs/UnrealBuildTool/System/NativeProjects.cs b/Engine/Source/Programs/UnrealBuildTool/System/NativeProjects.cs
--- a/Engine/Source/Programs/UnrealBuildTool/System/NativeProjects.cs
+++ b/Engine/Source/Programs/UnrealBuildTool/System/NativeProjects.cs
@@ -45,7 +45,7 @@ namespace UnrealBuildTool
 				{
 					if (CachedTargetNameToProjectFile == null)
 					{
-						using var _ = Timeline.ScopeEvent("Find target files");
+						using ITimelineEvent _ = Timeline.ScopeEvent("Find target files");
 						Dictionary<string, FileReference> TargetNameToProjectFile = new Dictionary<string, FileReference>();
 						foreach (FileReference ProjectFile in EnumerateProjectFiles(Logger))
 						{
diff --git a/Engine/Source/Programs/UnrealBuildTool/System/RulesAssembly.cs b/Engine/Source/Programs/UnrealBuildTool/System/RulesAssembly.cs
--- a/Engine/Source/Programs/UnrealBuildTool/System/RulesAssembly.cs
+++ b/Engine/Source/Programs/UnrealBuildTool/System/RulesAssembly.cs
@@ -225,7 +225,7 @@ namespace UnrealBuildTool
 			{
 				Task.Run(() =>
 				{
-					//using var timelineScope = Timeline.ScopeEvent("PrejitRules");
+					//using ITimelineEvent timelineScope = Timeline.ScopeEvent("PrejitRules");
 					ParallelQueue queue = new();
 					foreach (Type CompiledType in CompiledAssembly.GetTypes())
 					{
@@ -808,14 +808,14 @@ namespace UnrealBuildTool
 					bool mightContain = false;
 					if (info.File.IsUnderDirectory(Unreal.EngineDirectory))
 					{
-						var span = info.File.FullName.AsSpan(Unreal.EngineDirectory.FullName.Length + 1);
+						ReadOnlySpan<char> span = info.File.FullName.AsSpan(Unreal.EngineDirectory.FullName.Length + 1);
 						if (bloomFilter.MightContainSubString(span, Path.DirectorySeparatorChar))
 							mightContain = true;
 					}
 
 					if (!mightContain && ProjectDir != null && info.File.IsUnderDirectory(ProjectDir))
 					{
-						var span = info.File.FullName.AsSpan(ProjectDir.FullName.Length + 1);
+						ReadOnlySpan<char> span = info.File.FullName.AsSpan(ProjectDir.FullName.Length + 1);
 						if (bloomFilter.MightContainSubString(span, Path.DirectorySeparatorChar))
 							mightContain = true;
 					}
diff --git a/Engine/Source/Programs/UnrealBuildTool/System/RulesCompiler.cs b/Engine/Source/Programs/UnrealBuildTool/System/RulesCompiler.cs
--- a/Engine/Source/Programs/UnrealBuildTool/System/RulesCompiler.cs
+++ b/Engine/Source/Programs/UnrealBuildTool/System/RulesCompiler.cs
@@ -79,7 +79,7 @@ namespace UnrealBuildTool
 						List<PluginInfo> EnginePlugins = [];
 						List<PluginInfo> MarketplacePlugins = [];
 						{
-							using var ite = Timeline.ScopeEvent("EngineRulesPlugins");
+							using ITimelineEvent ite = Timeline.ScopeEvent("EngineRulesPlugins");
 							DirectoryReference MarketplaceDirectory = DirectoryReference.Combine(Unreal.EngineDirectory, "Plugins", "Marketplace");
 							foreach (PluginInfo PluginInfo in Plugins.ReadEnginePlugins(Unreal.EngineDirectory))
 							{
@@ -131,7 +131,7 @@ namespace UnrealBuildTool
 
 			// We want to prefetch the directories in one call to make it go as wide as possible against filesystem
 			{
-				using var fdt = Timeline.ScopeEvent("EngineRulesFetchDirs");
+				using ITimelineEvent fdt = Timeline.ScopeEvent("EngineRulesFetchDirs");
 
 				List<DirectoryReference> populateDirs = [];
 				foreach (PluginInfo info in Plugins)
@@ -193,13 +193,13 @@ namespace UnrealBuildTool
 			}
 
 			// Create the assembly
-			using var erat = Timeline.ScopeEvent("EngineRulesAssembly");
+			using ITimelineEvent erat = Timeline.ScopeEvent("EngineRulesAssembly");
 			DirectoryReference AssemblyDir = RootDirectories[0];
 			FileReference EngineAssemblyFileName = FileReference.Combine(AssemblyDir, "Intermediate", "Build", "BuildRules", AssemblyPrefix + "Rules" + FrameworkAssemblyExtension);
 			RulesAssembly EngineAssembly = new RulesAssembly(Scope, RootDirectories, Plugins, ModuleFileToContext, new List<FileReference>(), EngineAssemblyFileName, bContainsEngineModules: true, DefaultBuildSettings: DefaultEngineBuildSettingsVersion, bReadOnly: bReadOnly, bSkipCompile: bSkipCompile, bForceCompile: bForceCompile, Parent: Parent, Logger: Logger);
 			erat.Finish();
 
-			using var prat = Timeline.ScopeEvent("ProgramRulesAssembly");
+			using ITimelineEvent prat = Timeline.ScopeEvent("ProgramRulesAssembly");
 			Dictionary<FileReference, ModuleRulesContext> ProgramModuleFiles = new Dictionary<FileReference, ModuleRulesContext>();
 			foreach (DirectoryReference RootDirectory in RootDirectories)
 			{
@@ -286,7 +286,7 @@ namespace UnrealBuildTool
 				// Create the parent assembly
 				RulesAssembly Parent = CreateEngineRulesAssembly(bUsePrecompiled, bSkipCompile, bForceCompile, Logger);
 
-				using var ite = Timeline.ScopeEvent("ProjectRulesAssembly");
+				using ITimelineEvent ite = Timeline.ScopeEvent("ProjectRulesAssembly");
 
 				DirectoryReference MainProjectDirectory = ProjectFileName.Directory;
 				//DirectoryReference MainProjectSourceDirectory = DirectoryReference.Combine(MainProjectDirectory, "Source");
diff --git a/Engine/Source/Programs/UnrealBuildTool/System/TargetMakefile.cs b/Engine/Source/Programs/UnrealBuildTool/System/TargetMakefile.cs
--- a/Engine/Source/Programs/UnrealBuildTool/System/TargetMakefile.cs
+++ b/Engine/Source/Programs/UnrealBuildTool/System/TargetMakefile.cs
@@ -752,7 +752,7 @@ namespace UnrealBuildTool
 			TargetMakefile Makefile;
 			using (GlobalTracer.Instance.BuildSpan("Loading makefile").StartActive())
 			{
-				using var _ = Timeline.ScopeEvent("LoadMakefile");
+				using ITimelineEvent _ = Timeline.ScopeEvent("LoadMakefile");
 				try
 				{
 					using (BinaryArchiveReader Reader = new BinaryArchiveReader(MakefilePath))
@@ -833,7 +833,7 @@ namespace UnrealBuildTool
 
 				// Check to see if the number of plugins changed
 				{
-					using var _ = Timeline.ScopeEvent("FindPlugins");
+					using ITimelineEvent _ = Timeline.ScopeEvent("FindPlugins");
 
 					FileReference[] Plugins = EnumPluginsTask.Result;
 
@@ -882,7 +882,7 @@ namespace UnrealBuildTool
 		/// <returns>True if the makefile is valid, false otherwise</returns>
 		public static bool IsValidForSourceFiles(TargetMakefile Makefile, FileReference? ProjectFile, UnrealTargetPlatform Platform, ISourceFileWorkingSet WorkingSet, ILogger Logger, out string? ReasonNotLoaded)
 		{
-			using var _ = Timeline.ScopeEvent("MakefileValid");
+			using ITimelineEvent _ = Timeline.ScopeEvent("MakefileValid");
 
 			using (GlobalTracer.Instance.BuildSpan("TargetMakefile.IsValidForSourceFiles()").StartActive())
 			{
diff --git a/Engine/Source/Programs/UnrealBuildTool/System/UHTExecution.cs b/Engine/Source/Programs/UnrealBuildTool/System/UHTExecution.cs
--- a/Engine/Source/Programs/UnrealBuildTool/System/UHTExecution.cs
+++ b/Engine/Source/Programs/UnrealBuildTool/System/UHTExecution.cs
@@ -1224,7 +1224,7 @@ namespace UnrealBuildTool
 
 			internal bool GetModuleInfo(string module, out UInt128 hash, out long lastUpdateUtc)
 			{
-				if (_modules.TryGetValue(module, out var entry))
+				if (_modules.TryGetValue(module, out (UInt128 hash, long lastUpdateUtc) entry))
 				{
 					hash = entry.hash;
 					lastUpdateUtc = entry.lastUpdateUtc;
@@ -1250,7 +1250,7 @@ namespace UnrealBuildTool
 				using BinaryWriter writer = new(stream);
 				writer.Write(CurrentVersion);
 				writer.Write(_modules.Count);
-				foreach (var kv in _modules)
+				foreach (KeyValuePair<string, (UInt128 hash, long lastUpdateUtc)> kv in _modules)
 				{
 					writer.Write(kv.Key);
 					writer.Write((ulong)kv.Value.hash);
@@ -1281,7 +1281,7 @@ namespace UnrealBuildTool
 				}
 				using (ProgressWriter Progress = new("Generating code...", false, Logger))
 				{
-					using var uhtManifestScope = Timeline.ScopeEvent("UHTManifest");
+					using ITimelineEvent uhtManifestScope = Timeline.ScopeEvent("UHTManifest");
 
 					string RootLocalPath = Unreal.RootDirectory.FullName;
 
@@ -1404,7 +1404,7 @@ namespace UnrealBuildTool
 
 						Logger.LogDebug("  Running Internal UnrealHeaderTool {CmdLine}", CmdLine);
 
-						using var uhtExecute = Timeline.ScopeEvent("UHT");
+						using ITimelineEvent uhtExecute = Timeline.ScopeEvent("UHT");
 
 						// Run UHT
 						IScope Timer = GlobalTracer.Instance.BuildSpan("Executing UnrealHeaderTool").StartActive();
@@ -1480,7 +1480,7 @@ namespace UnrealBuildTool
 				{
 					using (GlobalTracer.Instance.BuildSpan("UHTExecution.UpdateTimestamps()").StartActive())
 					{
-						using var uhtExecute = Timeline.ScopeEvent("UHTTimestamps");
+						using ITimelineEvent uhtExecute = Timeline.ScopeEvent("UHTTimestamps");
 						timestampsFile.Update(Makefile.UObjectModules);
 						timestampsFile.Save(Logger);
 					}
@@ -1499,7 +1499,7 @@ namespace UnrealBuildTool
 		/// <param name="Makefile"></param>
 		public static void WaitForExecution(Task executionTask, TargetMakefile Makefile)
 		{
-			using var uhtManifestScope = Timeline.ScopeEvent("UHTWait");
+			using ITimelineEvent uhtManifestScope = Timeline.ScopeEvent("UHTWait");
 			executionTask.Wait();
 			ResetCachedHeaderInfo(Makefile.UObjectModules);
 		}
diff --git a/Engine/Source/Programs/UnrealBuildTool/System/Utils.cs b/Engine/Source/Programs/UnrealBuildTool/System/Utils.cs
--- a/Engine/Source/Programs/UnrealBuildTool/System/Utils.cs
+++ b/Engine/Source/Programs/UnrealBuildTool/System/Utils.cs
@@ -1452,7 +1452,7 @@ namespace UnrealBuildTool
 		{
 			foreach (FileReference ScriptFile in ScriptFiles)
 			{
-				using var _ = Timeline.ScopeEvent(ScriptFile.GetFileName());
+				using ITimelineEvent _ = Timeline.ScopeEvent(ScriptFile.GetFileName());
 
 				ProcessStartInfo StartInfo = new ProcessStartInfo();
 				StartInfo.FileName = BuildHostPlatform.Current.Shell.FullName;
@@ -1475,7 +1475,7 @@ namespace UnrealBuildTool
 
 			if (ScriptFiles.Length > 0)
 			{
-				using var _ = Timeline.ScopeEvent("ResetCachedInfo");
+				using ITimelineEvent _ = Timeline.ScopeEvent("ResetCachedInfo");
 
 				// We have to invalidate all cached file info after running the scripts, because we don't know what may have changed.
 				DirectoryItem.ResetAllCachedInfo_SLOW();
diff --git a/Engine/Source/Programs/UnrealBuildTool/ToolChain/ClangToolChain.cs b/Engine/Source/Programs/UnrealBuildTool/ToolChain/ClangToolChain.cs
--- a/Engine/Source/Programs/UnrealBuildTool/ToolChain/ClangToolChain.cs
+++ b/Engine/Source/Programs/UnrealBuildTool/ToolChain/ClangToolChain.cs
@@ -785,7 +785,7 @@ namespace UnrealBuildTool
 			if (CompileEnvironment.Architecture == UnrealArch.X64 && CompileEnvironment.MinCpuArchX64 != MinimumCpuArchitectureX64.None)
 			{
 				// The binary created will be targeting AVX instructions. Machines without AVX support will crash on any AVX instructions if they run this compilation unit.
-				var MinCpuArchX64 = CompileEnvironment.MinCpuArchX64;
+				MinimumCpuArchitectureX64 MinCpuArchX64 = CompileEnvironment.MinCpuArchX64;
 				VersionNumber ClangVersion = new(Info.ClangVersion.Major, Info.ClangVersion.Minor, Info.ClangVersion.Build);
 				// Downgrade AVX if unsupported by Clang version
 				if (ClangVersion < new VersionNumber(20, 1) && MinCpuArchX64 >= MinimumCpuArchitectureX64.AVX10_2)
