[UBT]
* Fixed so xcode project generator creates a lldbinit file containing mappings when vfs is enabled. Note, this only affect solution generation when vfs is enabled in ubt

diff --git a/Engine/Source/Programs/UnrealBuildTool/ProjectFiles/Xcode/XcodeProject.Scheme.cs b/Engine/Source/Programs/UnrealBuildTool/ProjectFiles/Xcode/XcodeProject.Scheme.cs
--- a/Engine/Source/Programs/UnrealBuildTool/ProjectFiles/Xcode/XcodeProject.Scheme.cs
+++ b/Engine/Source/Programs/UnrealBuildTool/ProjectFiles/Xcode/XcodeProject.Scheme.cs
@@ -58,6 +58,19 @@ namespace UnrealBuildTool.XcodeProjectXcconfig
 					}
 				}
 
+				FileReference lldbInitFile = FileReference.Combine(SchemeFilePath.Directory, "LLDBInit");
+
+				List<string> sourceMappings = [];
+				if (UnrealData.RootPaths != null)
+				{
+					foreach ((uint id, DirectoryReference vfs, DirectoryReference local) in UnrealData.RootPaths)
+					{
+						sourceMappings.Add($"settings append target.source-map \"{vfs}\" \"{local}\"");
+					}
+
+					FileReference.WriteAllLines(lldbInitFile, sourceMappings);
+				}
+
 				IEnumerable<string> ScriptLines = new string[]
 				{
 					"# if a UBTGenerated plist doesn&apos;t exist, create it now for Xcode&apos;s dependency needs",
@@ -160,6 +173,10 @@ namespace UnrealBuildTool.XcodeProjectXcconfig
 				Content.WriteLine("   </TestAction>");
 				Content.WriteLine("   <LaunchAction");
 				Content.WriteLine("      buildConfiguration = \"" + DefaultConfiguration + "\"");
+				if (sourceMappings.Count > 0)
+				{
+					Content.WriteLine("      customLLDBInitFile = \"" + lldbInitFile + "\"");
+				}
 				Content.WriteLine("      selectedDebuggerIdentifier = \"Xcode.DebuggerFoundation.Debugger.LLDB\"");
 				Content.WriteLine("      selectedLauncherIdentifier = \"Xcode.DebuggerFoundation.Launcher.LLDB\"");
 				Content.WriteLine("      launchStyle = \"0\"");
diff --git a/Engine/Source/Programs/UnrealBuildTool/ProjectFiles/Xcode/XcodeProject.cs b/Engine/Source/Programs/UnrealBuildTool/ProjectFiles/Xcode/XcodeProject.cs
--- a/Engine/Source/Programs/UnrealBuildTool/ProjectFiles/Xcode/XcodeProject.cs
+++ b/Engine/Source/Programs/UnrealBuildTool/ProjectFiles/Xcode/XcodeProject.cs
@@ -360,6 +360,9 @@ namespace UnrealBuildTool.XcodeProjectXcconfig
 		public IOSProjectSettings? TVOSProjectSettings;
 		public IOSProjectSettings? VisionOSProjectSettings;
 
+		// Root paths for creating lldbinit file with source mappings vfs
+		public CppRootPaths? RootPaths = null;
+
 		// Name of the product (usually the project name, but UE5.xcodeproj is actually UnrealGame product)
 		public string ProductName;
 
@@ -670,6 +673,8 @@ namespace UnrealBuildTool.XcodeProjectXcconfig
 			FileBatch.bEnableRTTI = CompileEnvironment.bUseRTTI;
 			FileBatch.SystemIncludePaths.UnionWith(CompileEnvironment.SystemIncludePaths);
 			FileBatch.UserIncludePaths.UnionWith(CompileEnvironment.UserIncludePaths);
+
+			RootPaths = CompileEnvironment.RootPaths;
 		}
 
 		private List<UnrealBuildConfig> GetSupportedBuildConfigs(List<UnrealTargetPlatform> Platforms, List<UnrealTargetConfiguration> Configurations, ILogger Logger)
