[UBT]
* Fixed so failing to save FileHashCache is not an error
* Fixed so trying to load FileHashCache from non-existing directory is not an error

diff --git a/Engine/Source/Programs/UnrealBuildTool/Modes/BuildMode.cs b/Engine/Source/Programs/UnrealBuildTool/Modes/BuildMode.cs
--- a/Engine/Source/Programs/UnrealBuildTool/Modes/BuildMode.cs
+++ b/Engine/Source/Programs/UnrealBuildTool/Modes/BuildMode.cs
@@ -281,7 +281,7 @@ namespace UnrealBuildTool
 				// Save all the caches
 				using (Timeline.ScopeEvent("Saving caches"))
 				{
-					fileHashCache.SaveAll();
+					fileHashCache.SaveAll(Logger);
 					SourceFileMetadataCache.SaveAll();
 					cppDependencyCache.SaveAll();
 				}
diff --git a/Engine/Source/Programs/UnrealBuildTool/ProjectFiles/ProjectFileGenerator.cs b/Engine/Source/Programs/UnrealBuildTool/ProjectFiles/ProjectFileGenerator.cs
--- a/Engine/Source/Programs/UnrealBuildTool/ProjectFiles/ProjectFileGenerator.cs
+++ b/Engine/Source/Programs/UnrealBuildTool/ProjectFiles/ProjectFileGenerator.cs
@@ -1496,7 +1496,7 @@ namespace UnrealBuildTool
 					WriteProjectFiles(PlatformProjectGenerators, Logger);
 				}
 
-				_fileHashCache.SaveAll();
+				_fileHashCache.SaveAll(Logger);
 				_fileHashCache = null;
 
 				Logger.LogDebug("Project generation complete ({GeneratedProjectFilesCount} generated, {OtherProjectFilesCount} imported)", GeneratedProjectFiles.Count, OtherProjectFiles.Count);
diff --git a/Engine/Source/Programs/UnrealBuildTool/System/FileHashCache.cs b/Engine/Source/Programs/UnrealBuildTool/System/FileHashCache.cs
--- a/Engine/Source/Programs/UnrealBuildTool/System/FileHashCache.cs
+++ b/Engine/Source/Programs/UnrealBuildTool/System/FileHashCache.cs
@@ -187,7 +187,7 @@ namespace UnrealBuildTool
 		/// <summary>
 		/// Save all partitions
 		/// </summary>
-		public void SaveAll()
+		public void SaveAll(ILogger logger)
 		{
 			foreach ((DirectoryReference root, List<(string filter, Partition partition)>) root in _roots)
 			{
@@ -201,27 +201,33 @@ namespace UnrealBuildTool
 					List<KeyValuePair<FileItem, FileEntry>> files = new(partition.Files);
 					files.Sort(static (a, b) => a.Key.Location.CompareTo(b.Key.Location));
 
-					//DirectoryReference.CreateDirectory(partition.Location.Directory);
-					using FileStream stream = File.Open(partition.Location.FullName, FileMode.Create, FileAccess.Write, FileShare.Read);
-					using BinaryWriter writer = new(stream, new UTF8Encoding(false), leaveOpen: false);
-					writer.Write(Partition.CurrentVersion);
-					writer.Write(partition.Files.Count);
-
-					string lastFullName = "";
-
-					foreach (KeyValuePair<FileItem, FileEntry> pair in files)
+					try
 					{
-						ReadOnlySpan<char> last = lastFullName;
-						ReadOnlySpan<char> curr = pair.Key.FullName.AsSpan();
-						UInt16 matchLength = (UInt16)MemoryExtensions.CommonPrefixLength(last, curr);
-						writer.Write(matchLength);
-						writer.Write(curr.Slice(matchLength).ToString());
-						lastFullName = pair.Key.FullName;
+						DirectoryReference.CreateDirectory(partition.Location.Directory);
+						using FileStream stream = File.Open(partition.Location.FullName, FileMode.Create, FileAccess.Write, FileShare.Read);
+						using BinaryWriter writer = new(stream, new UTF8Encoding(false), leaveOpen: false);
+						writer.Write(Partition.CurrentVersion);
+						writer.Write(partition.Files.Count);
 
-						writer.Write((ulong)pair.Value.hash);
-						writer.Write((ulong)(pair.Value.hash >> 64));
-						writer.Write(pair.Value.lastWriteTimeUtc);
+						string lastFullName = "";
 
+						foreach (KeyValuePair<FileItem, FileEntry> pair in files)
+						{
+							ReadOnlySpan<char> last = lastFullName;
+							ReadOnlySpan<char> curr = pair.Key.FullName.AsSpan();
+							UInt16 matchLength = (UInt16)MemoryExtensions.CommonPrefixLength(last, curr);
+							writer.Write(matchLength);
+							writer.Write(curr.Slice(matchLength).ToString());
+							lastFullName = pair.Key.FullName;
+
+							writer.Write((ulong)pair.Value.hash);
+							writer.Write((ulong)(pair.Value.hash >> 64));
+							writer.Write(pair.Value.lastWriteTimeUtc);
+						}
+					}
+					catch (Exception ex)
+					{
+						logger.LogInformation("Failed to save {Location} - {Message}", partition.Location, ex.Message);
 					}
 				}
 			}
@@ -446,7 +452,7 @@ namespace UnrealBuildTool
 
 					if (!isLocked)
 					{
-						logger.LogInformation($"Get exception trying to write {fileItem} ({e.Message}). Will wait 1 second and try again");
+						logger.LogInformation("Get exception trying to write {Location} ({Message}). Will wait 1 second and try again", fileItem.Location, e.Message);
 						Monitor.Enter(fileItem);
 						isLocked = true;
 					}
@@ -615,10 +621,12 @@ namespace UnrealBuildTool
 			catch (FileNotFoundException)
 			{
 			}
+			catch (DirectoryNotFoundException)
+			{
+			}
 			catch (Exception Ex)
 			{
-				logger.LogWarning("Unable to read {Location}. See log for additional information.", partition.Location);
-				logger.LogDebug("{Ex}", ExceptionUtils.FormatExceptionDetails(Ex));
+				logger.LogInformation("Unable to read {Location} - {Message}", partition.Location, Ex.Message);
 			}
 		}
 
