[UBT]
* Added /dxifcSuppressDllImportTransform to arguments when compiling with header units to fix lots of warnings when linking

diff --git a/Engine/Source/Programs/UnrealBuildTool/Platform/Windows/VCToolChain.cs b/Engine/Source/Programs/UnrealBuildTool/Platform/Windows/VCToolChain.cs
--- a/Engine/Source/Programs/UnrealBuildTool/Platform/Windows/VCToolChain.cs
+++ b/Engine/Source/Programs/UnrealBuildTool/Platform/Windows/VCToolChain.cs
@@ -316,6 +316,7 @@ namespace UnrealBuildTool
 				Arguments.Add($"/translateInclude");
 				Arguments.Add($"/dxifcInlineFunctions");
 				Arguments.Add($"/ifcOutput {CreatePchFilePath}");
+				Arguments.Add($"/dxifcSuppressDllImportTransform");
 			}
 			else
 			{
@@ -333,6 +334,7 @@ namespace UnrealBuildTool
 				Arguments.Add("/wd4127"); // with header units pragma warning disable is not propagated and this warning is very spammy so explicitly turned off here. Hoping ms will change so warning disable is propgated from HU
 				Arguments.Add($"/translateInclude");
 				Arguments.Add($"/headerUnit:quote {PchThroughHeaderFilePath}={UsingPchFilePath}");
+				Arguments.Add($"/dxifcSuppressDllImportTransform");
 			}
 			else
 			{
