[UBT]
* Fixed rare situation where targets with same code gen folder is built multiple time in same run and code gen differs between the targets. Like running -target="UnrealGame win64 development" -target="UnrealGame win64 debuggame -allmodules". Both targets will use same codegen folder but -allmodules will add more uht modules..

diff --git a/Engine/Source/Programs/UnrealBuildTool/Configuration/UEBuildTarget.cs b/Engine/Source/Programs/UnrealBuildTool/Configuration/UEBuildTarget.cs
--- a/Engine/Source/Programs/UnrealBuildTool/Configuration/UEBuildTarget.cs
+++ b/Engine/Source/Programs/UnrealBuildTool/Configuration/UEBuildTarget.cs
@@ -2814,58 +2814,80 @@ namespace UnrealBuildTool
 
 				if (lastRunContext != null)
 				{
+					bool rerunCodeGen = false;
+
 					if (context.TargetName != lastRunContext.TargetName)
 					{
-						throw new Exception("Different targetname!.. run after each other!");
+						Logger.LogDebug("Different targetname!.. will rerun code gen!");
+						rerunCodeGen = true;
 					}
 
-					if (lastRunUObjectModules == null)
+					if (!rerunCodeGen)
 					{
-						lastRunUObjectModules = lastRunContext.Makefile.UObjectModules.ToDictionary(k => k.ModuleName, k => k);
-					}
+						if (lastRunUObjectModules == null)
+						{
+							lastRunUObjectModules = lastRunContext.Makefile.UObjectModules.ToDictionary(k => k.ModuleName, k => k);
+						}
 
-					foreach (UHTModuleInfo module in makefile.UObjectModules)
-					{
-						if (lastRunUObjectModules.TryGetValue(module.ModuleName, out UHTModuleInfo? lastModule))
+						foreach (UHTModuleInfo module in makefile.UObjectModules)
 						{
-							int CountHeaders(UHTModuleInfo info) => info.PublicUObjectClassesHeaders.Count + info.PublicUObjectHeaders.Count + info.InternalUObjectHeaders.Count + info.PrivateUObjectHeaders.Count;
-							if (CountHeaders(module) != CountHeaders(lastModule))
+							if (lastRunUObjectModules.TryGetValue(module.ModuleName, out UHTModuleInfo? lastModule))
 							{
-								throw new Exception($"Target {GetTargetName(context.Descriptor)} module {module.ModuleName} has different number of headers than previous target {GetTargetName(lastRunContext.Descriptor)} did not have!.. run after each other!");
+								int CountHeaders(UHTModuleInfo info) => info.PublicUObjectClassesHeaders.Count + info.PublicUObjectHeaders.Count + info.InternalUObjectHeaders.Count + info.PrivateUObjectHeaders.Count;
+								if (CountHeaders(module) != CountHeaders(lastModule))
+								{
+									Logger.LogDebug($"Target {GetTargetName(context.Descriptor)} module {module.ModuleName} has different number of headers than previous target {GetTargetName(lastRunContext.Descriptor)} did not have!.. will rerun code gen");
+									rerunCodeGen = true;
+									break;
+								}
+							}
+							else
+							{
+								Logger.LogDebug($"Target {GetTargetName(context.Descriptor)} has uht module {module.ModuleName} that previous target {GetTargetName(lastRunContext.Descriptor)} did not have!.. will rerun code gen");
+								rerunCodeGen = true;
+								break;
 							}
-						}
-						else
-						{
-							throw new Exception($"Target {GetTargetName(context.Descriptor)} has uht module {module.ModuleName} that previous target {GetTargetName(lastRunContext.Descriptor)} did not have!.. run after each other!");
 						}
 					}
 
 #if __VPROJECT_AVAILABLE__
-					HashSet<string>? lastRunVNIModules = null;
-
-					if (lastRunVNIModules == null)
+					if (!rerunCodeGen)
 					{
-						lastRunVNIModules = new(lastRunContext.Makefile.VNIModules.Select(m => m.ModuleName));
-					}
+						HashSet<string>? lastRunVNIModules = null;
 
-					foreach (VNIModuleInfo module in makefile.VNIModules)
-					{
-						if (!lastRunVNIModules.Contains(module.ModuleName))
+						if (lastRunVNIModules == null)
 						{
-							throw new Exception($"Target {GetTargetName(context.Descriptor)} has vni modules that previous target {GetTargetName(lastRunContext.Descriptor)} did not have!.. run after each other!");
+							lastRunVNIModules = new(lastRunContext.Makefile.VNIModules.Select(m => m.ModuleName));
+						}
+
+						foreach (VNIModuleInfo module in makefile.VNIModules)
+						{
+							if (!lastRunVNIModules.Contains(module.ModuleName))
+							{
+								Logger.LogDebug($"Target {GetTargetName(context.Descriptor)} has vni modules that previous target {GetTargetName(lastRunContext.Descriptor)} did not have!.. will rerun code gen");
+								rerunCodeGen = true;
+								break;
+							}
 						}
 					}
 #endif
+					if (!rerunCodeGen)
+					{
+						continue;
+					}
+				}
 
-					//if (context.Makefile.EnabledUhtPlugins!.Length != lastRunContext.Makefile.EnabledUhtPlugins!.Length)
-					//	throw new Exception("DIFFERENT THIS NEEDS TO BE HANDLED!.. run after each other!");
-
-					// No need to set these since they are set by a previous context
-					//context.uhtRun = lastRunContext.uhtRun;
-					//context.uhtTimestamp = lastRunContext.uhtTimestamp;
-					//context.vniRun = lastRunContext.vniRun;
-					continue;
+				// This should be very rare to rerun code gen.. can be reproed with having for example both these being built in one go:
+				// -Target="UnrealGame Win64 Development"
+				// -Target="UnrealGame Win64 DebugGame -allmodules"
 
+				if (lastRunContext != null)
+				{
+					lastRunContext.uhtRun?.Wait();
+					lastRunContext.uhtTimestamp?.Wait();
+#if __VPROJECT_AVAILABLE__
+					lastRunContext.vniRun?.Wait();
+#endif
 				}
 
 				TargetDescriptor descriptor = context.Descriptor;
