[UBT]
* Moved pch logging to be behind a define that is manually enabled to not get all the crazy much spam when running in debug. (info only useful when actually looking at pch logic)

diff --git a/Engine/Source/Programs/UnrealBuildTool/Configuration/UEBuildModuleCPP.cs b/Engine/Source/Programs/UnrealBuildTool/Configuration/UEBuildModuleCPP.cs
--- a/Engine/Source/Programs/UnrealBuildTool/Configuration/UEBuildModuleCPP.cs
+++ b/Engine/Source/Programs/UnrealBuildTool/Configuration/UEBuildModuleCPP.cs
@@ -1,5 +1,7 @@
 // Copyright Epic Games, Inc. All Rights Reserved.
 
+//#define LOG_PCH_PARSING
+
 using EpicGames.Core;
 using Microsoft.CodeAnalysis;
 using Microsoft.Extensions.Logging;
@@ -1474,19 +1476,19 @@ namespace UnrealBuildTool
 						{
 							if (!TransformedHeaderInclude.Contains('.'))
 							{
-#if DEBUG
+#if LOG_PCH_PARSING
 								Logger.LogDebug("{0} SharedPCH - Skipping '{1}' found in '{2}' because it doesn't appear to be a module header.", Name, TransformedHeaderInclude, FileToSearch.Location);
 #endif
 							}
 							else if (TransformedHeaderInclude.EndsWith(".generated.h"))
 							{
-#if DEBUG
+#if LOG_PCH_PARSING
 								Logger.LogDebug("{0} SharedPCH - Skipping '{1}' found in '{2}' because it appears to be a generated UHT header.", Name, TransformedHeaderInclude, FileToSearch.Location);
 #endif
 							}
 							else
 							{
-#if DEBUG
+#if LOG_PCH_PARSING
 								Logger.LogDebug("{0} SharedPCH - Could not find include directory for '{1}' found in '{2}'.", Name, TransformedHeaderInclude, FileToSearch.Location);
 #endif
 								IncludesNotFound.Add(TransformedHeaderInclude);
@@ -1507,7 +1509,7 @@ namespace UnrealBuildTool
 					UEBuildModule? FoundModule = ModuleDeps.FirstOrDefault(Module => Module.ContainsFile(IncludeFile.Location));
 					if (FoundModule != null)
 					{
-#if DEBUG
+#if LOG_PCH_PARSING
 						if (ModuleToIncludePathsOnlyFlag[FoundModule])
 						{
 							Logger.LogDebug("{0} SharedPCH - '{1}' is exporting types but '{2}' isn't declared as a public dependency.", Name, IncludeFile.Location, FoundModule);
@@ -1518,7 +1520,7 @@ namespace UnrealBuildTool
 					}
 					else
 					{
-#if DEBUG
+#if LOG_PCH_PARSING
 						Logger.LogDebug("{0} SharedPCH - '{1}' is exporting types but the module this file belongs to isn't declared as a public dependency or include.", Name, IncludeFile.Location);
 #endif
 						FoundAllModules = false;
@@ -1526,7 +1528,7 @@ namespace UnrealBuildTool
 				}
 				else
 				{
-#if DEBUG
+#if LOG_PCH_PARSING
 					Logger.LogDebug("{0} SharedPCH - '{1}' is not exporting types so we are ignoring the dependency", Name, IncludeFile.Location);
 #endif
 				}
@@ -1534,7 +1536,7 @@ namespace UnrealBuildTool
 
 			if (!FoundAllModules)
 			{
-#if DEBUG
+#if LOG_PCH_PARSING
 				Logger.LogDebug("{0} SharedPCH - Is missing public dependencies. To be safe, this shared PCH will fall back to use all the module dependencies. Note that this could affect compile times.", Name);
 #endif
 				return AllModuleDeps;
