[UBT]
* Fixed so build target BuildPlugins is sorted from start and removed the orderby called. This also caused Binary Modules array to be non-deterministic which resulted in rsp file not being deterministic
* Fixed so UnboundReferences are added in a deterministic way to binary (Can't use UEBuildTarget.Modules since it is a dictionary and not deterministic)

diff --git a/Engine/Source/Programs/UnrealBuildTool/Configuration/UEBuildTarget.cs b/Engine/Source/Programs/UnrealBuildTool/Configuration/UEBuildTarget.cs
--- a/Engine/Source/Programs/UnrealBuildTool/Configuration/UEBuildTarget.cs
+++ b/Engine/Source/Programs/UnrealBuildTool/Configuration/UEBuildTarget.cs
@@ -1452,7 +1452,7 @@ namespace UnrealBuildTool
 				return false;
 			}
 
-			foreach (UEBuildPlugin plugin in BuildPlugins.OrderBy(x => x.Name))
+			foreach (UEBuildPlugin plugin in BuildPlugins)
 			{
 				UEBuildPlatform.GetBuildPlatform(Platform).ValidatePlugin(plugin, Rules);
 			}
@@ -4438,13 +4438,17 @@ namespace UnrealBuildTool
 			}
 
 			// Add the external and non-C++ referenced modules to the binaries that reference them.
-			foreach (UEBuildModuleCPP Module in Modules.Values.OfType<UEBuildModuleCPP>())
+			// Note, we want to add them sorted since we want binary modules to be sorted
+			foreach (UEBuildBinary binary in Binaries)
 			{
-				if (Module.Binary != null)
+				for (int i=0, e= binary.Modules.Count; i!=e; ++i)
 				{
-					foreach (UEBuildModule ReferencedModule in Module.GetUnboundReferences())
+					if (binary.Modules[i] is UEBuildModuleCPP Module)
 					{
-						Module.Binary.AddModule(ReferencedModule);
+						foreach (UEBuildModule ReferencedModule in Module.GetUnboundReferences())
+						{
+							binary.AddModule(ReferencedModule);
+						}
 					}
 				}
 			}
@@ -5561,7 +5565,8 @@ namespace UnrealBuildTool
 			}
 
 			// Set the list of plugins that should be built
-			BuildPlugins = new List<UEBuildPlugin>(NameToInstance.Values);
+			BuildPlugins = new(NameToInstance.Values);
+			BuildPlugins.SortBy(x => x.Name);
 
 			// Setup the UHT plugins
 			SetupUbtPlugins(Logger);
@@ -5583,7 +5588,7 @@ namespace UnrealBuildTool
 				{
 					// Compile in the build plugins list for the runtime plugin manager to filter optional plugin references against. Only plugins that were actually built
 					// should be enabled when another plugin has an optional dependency on the plugin.
-					IEnumerable<string> BuildPluginStrings = BuildPlugins?.OrderBy(x => x.Name)?.Select(x => $"TEXT(\"{x}\")") ?? Array.Empty<string>();
+					IEnumerable<string> BuildPluginStrings = BuildPlugins?.Select(x => $"TEXT(\"{x}\")") ?? Array.Empty<string>();
 					moduleRules.PrivateDefinitions.Add($"UBT_TARGET_BUILD_PLUGINS={String.Join(", ", BuildPluginStrings)}");
 				}
 			}
