UnrealBuildTool: Convert resource files to objects manually so link.exe doesn't need to prevent unnecessary reprocessing when resource files are shared
Remove LinkEnvironment.CommonResourceFiles as it is unused

#rnx
#FailedValidations Preflight

diff --git a/Engine/Source/Programs/UnrealBuildTool/Configuration/UEBuildBinary.cs b/Engine/Source/Programs/UnrealBuildTool/Configuration/UEBuildBinary.cs
--- a/Engine/Source/Programs/UnrealBuildTool/Configuration/UEBuildBinary.cs
+++ b/Engine/Source/Programs/UnrealBuildTool/Configuration/UEBuildBinary.cs
@@ -1,10 +1,10 @@
 // Copyright Epic Games, Inc. All Rights Reserved.
 
+using EpicGames.Core;
+using Microsoft.Extensions.Logging;
 using System;
 using System.Collections.Generic;
 using System.Linq;
-using EpicGames.Core;
-using Microsoft.Extensions.Logging;
 using UnrealBuildBase;
 
 namespace UnrealBuildTool
@@ -741,6 +741,7 @@ namespace UnrealBuildTool
 			using Pooled<HashSet<FileItem>> rented2 = HashSetPool<FileItem>.Rent();
 			HashSet<FileItem> InputFilesLookup = rented2.Value;
 
+			List<UEBuildModule> CompiledModules = [];
 			foreach (UEBuildModule Module in UEBuildModule.StableTopologicalSort(Modules))
 			{
 				if (Module.Binary == null || Module.Binary == this)
@@ -752,11 +753,12 @@ namespace UnrealBuildTool
 					Logger.LogTrace("Compile module: {ModuleName}", Module.Name);
 					List<FileItem> LinkInputFiles = Module.Compile(Target, ToolChain, BinaryCompileEnvironment, WorkingSet, Graph, Logger);
 					InputObjects.UnionWith(LinkInputFiles.Where(x => x.HasExtension(".obj") || x.HasExtension(".o")));
+					CompiledModules.Add(Module);
 
 					// Save the module outputs. In monolithic builds, this is just the object files.
 					if (Target.LinkType == TargetLinkType.Monolithic)
 					{
-						Graph.SetOutputItemsForModule(Module.Name, LinkInputFiles.ToArray());
+						Graph.SetOutputItemsForModule(Module.Name, [.. LinkInputFiles]);
 					}
 
 					// Sometimes the same link input file is added multiple times, so make sure to filter out the dups
@@ -777,7 +779,7 @@ namespace UnrealBuildTool
 					// Force a reference to initialize module for this binary
 					if (Module.Rules.bRequiresImplementModule ?? true)
 					{
-						BinaryLinkEnvironment.IncludeFunctions.Add(String.Format("IMPLEMENT_MODULE_{0}", Module.Name));
+						BinaryLinkEnvironment.IncludeFunctions.Add($"IMPLEMENT_MODULE_{Module.Name}");
 					}
 				}
 				else
@@ -802,7 +804,7 @@ namespace UnrealBuildTool
 			}
 
 			// Set the link output file.
-			BinaryLinkEnvironment.OutputFilePaths = OutputFilePaths.ToList();
+			BinaryLinkEnvironment.OutputFilePaths = [.. OutputFilePaths];
 
 			// Remember the link type
 			BinaryLinkEnvironment.LinkType = Target.LinkType;
@@ -820,6 +822,9 @@ namespace UnrealBuildTool
 			BinaryLinkEnvironment.bIsBuildingDLL = IsBuildingDll(Type);
 			BinaryLinkEnvironment.bIsBuildingLibrary = IsBuildingLibrary(Type);
 
+			// Compile any resource files
+			CompileResourceFiles(Target, ToolChain, BinaryCompileEnvironment, BinaryLinkEnvironment, CompiledModules.SelectMany(x => x.ResourceFiles), Graph, Logger);
+
 			// Setup object export stripping
 			if (Target.bStripExports)
 			{
@@ -834,71 +839,75 @@ namespace UnrealBuildTool
 			// Code coverage inherited from compile environment
 			BinaryLinkEnvironment.bCodeCoverage = CompileEnvironment.bCodeCoverage;
 
+			// Provide list of runtime dependencies to link environment.
+			BinaryLinkEnvironment.RuntimeDependencies = RuntimeDependencies;
+
+			// Allow the platform to modify the binary link environment for platform-specific resources etc.
+			UEBuildPlatform.GetBuildPlatform(Target.Platform).ModifyBinaryLinkEnvironment(BinaryLinkEnvironment, BinaryCompileEnvironment, Target, ToolChain, Graph);
+
+			return BinaryLinkEnvironment;
+		}
+
+		private void CompileResourceFiles(ReadOnlyTargetRules Target, UEToolChain ToolChain, CppCompileEnvironment BinaryCompileEnvironment, LinkEnvironment BinaryLinkEnvironment, IEnumerable<FileItem> InputFiles, IActionGraphBuilder Graph, ILogger Logger)
+		{
+			// @todo: This should be in some Windows code somewhere...
+			if (!BinaryLinkEnvironment.Platform.IsInGroup(UnrealPlatformGroup.Windows))
+			{
+				return;
+			}
+
 			// If we don't have any resource file, use the default or compile a custom one for this module
-			if (BinaryLinkEnvironment.Platform.IsInGroup(UnrealPlatformGroup.Windows))
+			IEnumerable<FileItem> DefaultResourceFiles = [];
+			// Figure out if this binary has any custom resource files. Hacky check to ignore the resource file in the Launch module, since it contains dialogs that the engine needs and always needs to be included.
+			FileItem PCLaunchResource = FileItem.GetItemByFileReference(FileReference.Combine(Unreal.EngineSourceDirectory, "Runtime", "Launch", "Resources", "Windows", "PCLaunch.rc"));
+			if (!InputFiles.Any(x => x != PCLaunchResource))
 			{
-				// Figure out if this binary has any custom resource files. Hacky check to ignore the resource file in the Launch module, since it contains dialogs that the engine needs and always needs to be included.
-				FileItem[] CustomResourceFiles = BinaryLinkEnvironment.InputFiles.Where(x => x.Location.HasExtension(".res") && !x.Location.FullName.EndsWith("\\Launch\\PCLaunch.rc.res", StringComparison.OrdinalIgnoreCase)).ToArray();
-				if (CustomResourceFiles.Length == 0)
+				if (BinaryLinkEnvironment.DefaultResourceFiles.Count > 0)
 				{
-					if (BinaryLinkEnvironment.DefaultResourceFiles.Count > 0)
+					// If no resource files are found, use the precompiled default resource file for the target
+					if (!InputFiles.Any())
 					{
-						// Use the default resource file if possible
-						BinaryLinkEnvironment.InputFiles.AddRange(BinaryLinkEnvironment.DefaultResourceFiles);
+						BinaryLinkEnvironment.InputFiles.AddRange(BinaryLinkEnvironment.DefaultConvertedResourceFiles);
+						return;
 					}
-					else
-					{
-						// Get the intermediate directory
-						DirectoryReference ResourceIntermediateDirectory = BinaryLinkEnvironment.IntermediateDirectory;
-
-						// Place resource intermediate files for executables in a subfolder
-						if (!BinaryLinkEnvironment.bIsBuildingDLL && !BinaryLinkEnvironment.bIsBuildingLibrary)
-						{
-							ResourceIntermediateDirectory = DirectoryReference.Combine(ResourceIntermediateDirectory, OutputFilePaths[0].GetFileName());
-						}
-
-						// Create a compile environment for resource files
-						CppCompileEnvironment ResourceCompileEnvironment = new CppCompileEnvironment(BinaryCompileEnvironment);
+					DefaultResourceFiles = BinaryLinkEnvironment.DefaultResourceFiles;
+				}
+				else
+				{
+					// Otherwise compile the default resource file per-binary, so that it gets the correct ORIGINAL_FILE_NAME macro.
+					FileItem DefaultResourceFile = FileItem.GetItemByFileReference(FileReference.Combine(Unreal.EngineDirectory, "Build", "Windows", "Resources", "Default.rc2"));
+					InputFiles = [.. InputFiles, DefaultResourceFile];
+				}
+			}
 
-						// @todo: This should be in some Windows code somewhere...
-						// Set the original file name macro; used in Default.rc2 to set the binary metadata fields.
-						ResourceCompileEnvironment.Definitions.Add("ORIGINAL_FILE_NAME=\"" + OutputFilePaths[0].GetFileName() + "\"");
+			// Create a compile environment for resource files
+			CppCompileEnvironment ResourceCompileEnvironment = new(BinaryCompileEnvironment);
 
-						// Set the other version fields if requested
-						if (Target.WindowsPlatform.bSetResourceVersions)
-						{
-							if (Target.Version.Changelist != 0)
-							{
-								ResourceCompileEnvironment.Definitions.Add(String.Format("BUILT_FROM_CHANGELIST={0}", Target.Version.Changelist));
-							}
-							if (!String.IsNullOrEmpty(Target.Version.BranchName))
-							{
-								ResourceCompileEnvironment.Definitions.Add(String.Format("BRANCH_NAME={0}", Target.Version.BranchName));
-							}
-							if (!String.IsNullOrEmpty(Target.BuildVersion))
-							{
-								ResourceCompileEnvironment.Definitions.Add(String.Format("BUILD_VERSION={0}", Target.BuildVersion));
-							}
-						}
+			// Set the original file name macro; used in Default.rc2 to set the binary metadata fields.
+			ResourceCompileEnvironment.Definitions.Add($"ORIGINAL_FILE_NAME=\"{OutputFilePaths[0].GetFileName()}\"");
 
-						// Otherwise compile the default resource file per-binary, so that it gets the correct ORIGINAL_FILE_NAME macro.
-						FileItem DefaultResourceFile = FileItem.GetItemByFileReference(FileReference.Combine(Unreal.EngineDirectory, "Build", "Windows", "Resources", "Default.rc2"));
-						CPPOutput DefaultResourceOutput = ToolChain.CompileRCFiles(ResourceCompileEnvironment, new List<FileItem> { DefaultResourceFile }, ResourceIntermediateDirectory, Graph);
-						BinaryLinkEnvironment.InputFiles.AddRange(DefaultResourceOutput.ObjectFiles);
-					}
+			// Set the other version fields if requested
+			if (Target.WindowsPlatform.bSetResourceVersions)
+			{
+				if (Target.Version.Changelist != 0)
+				{
+					ResourceCompileEnvironment.Definitions.Add($"BUILT_FROM_CHANGELIST={Target.Version.Changelist}");
+				}
+				if (!String.IsNullOrEmpty(Target.Version.BranchName))
+				{
+					ResourceCompileEnvironment.Definitions.Add($"BRANCH_NAME={Target.Version.BranchName}");
+				}
+				if (!String.IsNullOrEmpty(Target.BuildVersion))
+				{
+					ResourceCompileEnvironment.Definitions.Add($"BUILD_VERSION={Target.BuildVersion}");
 				}
 			}
 
-			// Add all the common resource files
-			BinaryLinkEnvironment.InputFiles.AddRange(BinaryLinkEnvironment.CommonResourceFiles);
-
-			// Provide list of runtime dependencies to link environment.
-			BinaryLinkEnvironment.RuntimeDependencies = RuntimeDependencies;
-
-			// Allow the platform to modify the binary link environment for platform-specific resources etc.
-			UEBuildPlatform.GetBuildPlatform(Target.Platform).ModifyBinaryLinkEnvironment(BinaryLinkEnvironment, BinaryCompileEnvironment, Target, ToolChain, Graph);
+			// Compile the resource files
+			CPPOutput Result = ToolChain.CompileRCFiles(ResourceCompileEnvironment, InputFiles, IntermediateDirectory, Graph);
 
-			return BinaryLinkEnvironment;
+			// Optionally convert the resource files depending on toolchain settings
+			BinaryLinkEnvironment.InputFiles.AddRange(ToolChain.ConvertRCFiles(BinaryLinkEnvironment, [.. Result.ObjectFiles, .. DefaultResourceFiles], IntermediateDirectory, Graph).ObjectFiles);
 		}
 
 		/// <summary>
diff --git a/Engine/Source/Programs/UnrealBuildTool/Configuration/UEBuildModule.cs b/Engine/Source/Programs/UnrealBuildTool/Configuration/UEBuildModule.cs
--- a/Engine/Source/Programs/UnrealBuildTool/Configuration/UEBuildModule.cs
+++ b/Engine/Source/Programs/UnrealBuildTool/Configuration/UEBuildModule.cs
@@ -203,6 +203,11 @@ namespace UnrealBuildTool
 		/// </summary>
 		public bool bDependsOnVerse = false;
 
+		/// <summary>
+		/// Set of resource file paths
+		/// </summary>
+		public HashSet<FileItem> ResourceFiles = new();
+
 		/// <summary>
 		/// Set of debug visualizer paths
 		/// </summary>
@@ -341,7 +346,7 @@ namespace UnrealBuildTool
 				}
 			}
 
-			// Add any debug visualizers found in the module directories
+			// Add any resource and debug visualizers found in the module directories
 			foreach (DirectoryItem directory in ModuleDirectories)
 			{
 				NatvisFiles.UnionWith(directory.EnumerateFiles().Where(x => x.HasExtension(".natvis") || x.HasExtension(".natstepfilter") || x.HasExtension(".natjmc")));
diff --git a/Engine/Source/Programs/UnrealBuildTool/Configuration/UEBuildModuleCPP.cs b/Engine/Source/Programs/UnrealBuildTool/Configuration/UEBuildModuleCPP.cs
--- a/Engine/Source/Programs/UnrealBuildTool/Configuration/UEBuildModuleCPP.cs
+++ b/Engine/Source/Programs/UnrealBuildTool/Configuration/UEBuildModuleCPP.cs
@@ -1152,19 +1152,8 @@ namespace UnrealBuildTool
 				};
 			}
 
-			// Compile RC files. The resource compiler does not work with response files, and using the regular compile environment can easily result in the 
-			// command line length exceeding the OS limit. Use the binary compile environment to keep the size down, and require that all include paths
-			// must be specified relative to the resource file itself or Engine/Source.
-			if (InputFiles.RCFiles.Count > 0)
-			{
-				CppCompileEnvironment ResourceCompileEnvironment = new CppCompileEnvironment(BinaryCompileEnvironment);
-				if (Binary != null)
-				{
-					// @todo: This should be in some Windows code somewhere...
-					ResourceCompileEnvironment.Definitions.Add("ORIGINAL_FILE_NAME=\"" + Binary.OutputFilePaths[0].GetFileName() + "\"");
-				}
-				LinkInputFiles.AddRange(ToolChain.CompileRCFiles(ResourceCompileEnvironment, InputFiles.RCFiles, IntermediateDirectory, Graph).ObjectFiles);
-			}
+			// RC files for all modules in a binary will be compiled in UEBuildBinary.cs using the BinaryCompileEnvironment
+			ResourceFiles.UnionWith(InputFiles.RCFiles);
 
 			// Write the compiled manifest
 			if (Rules.bPrecompile && Target.LinkType == TargetLinkType.Monolithic)
diff --git a/Engine/Source/Programs/UnrealBuildTool/Configuration/UEBuildTarget.cs b/Engine/Source/Programs/UnrealBuildTool/Configuration/UEBuildTarget.cs
--- a/Engine/Source/Programs/UnrealBuildTool/Configuration/UEBuildTarget.cs
+++ b/Engine/Source/Programs/UnrealBuildTool/Configuration/UEBuildTarget.cs
@@ -2975,12 +2975,14 @@ namespace UnrealBuildTool
 				FileReference DefaultResourceLocation = FileReference.Combine(Unreal.EngineDirectory, "Build", "Windows", "Resources", "Default.rc2");
 				if (!Unreal.IsFileInstalled(DefaultResourceLocation))
 				{
-					CppCompileEnvironment DefaultResourceCompileEnvironment = new CppCompileEnvironment(GlobalCompileEnvironment);
+					CppCompileEnvironment DefaultResourceCompileEnvironment = new(GlobalCompileEnvironment);
 
 					FileItem DefaultResourceFile = FileItem.GetItemByFileReference(DefaultResourceLocation);
 
-					CPPOutput DefaultResourceOutput = TargetToolChain.CompileRCFiles(DefaultResourceCompileEnvironment, new List<FileItem> { DefaultResourceFile }, EngineIntermediateDirectory, MakefileBuilder);
+					CPPOutput DefaultResourceOutput = TargetToolChain.CompileRCFiles(DefaultResourceCompileEnvironment, [DefaultResourceFile], EngineIntermediateDirectory, MakefileBuilder);
 					GlobalLinkEnvironment.DefaultResourceFiles.AddRange(DefaultResourceOutput.ObjectFiles);
+					DefaultResourceOutput = TargetToolChain.ConvertRCFiles(GlobalLinkEnvironment, DefaultResourceOutput.ObjectFiles, EngineIntermediateDirectory, MakefileBuilder);
+					GlobalLinkEnvironment.DefaultConvertedResourceFiles.AddRange(DefaultResourceOutput.ObjectFiles);
 				}
 			}
 
diff --git a/Engine/Source/Programs/UnrealBuildTool/Platform/Windows/VCEnvironment.cs b/Engine/Source/Programs/UnrealBuildTool/Platform/Windows/VCEnvironment.cs
--- a/Engine/Source/Programs/UnrealBuildTool/Platform/Windows/VCEnvironment.cs
+++ b/Engine/Source/Programs/UnrealBuildTool/Platform/Windows/VCEnvironment.cs
@@ -96,6 +96,11 @@ namespace UnrealBuildTool
 		/// </summary>
 		public readonly FileReference ResourceCompilerPath;
 
+		/// <summary>
+		/// Path to the resource to object converter from the toolchain
+		/// </summary>
+		public readonly FileReference ResourceToObjectConverterPath;
+
 		/// <summary>
 		/// The path to the toolchain compiler for compiling code
 		/// </summary>
@@ -143,6 +148,7 @@ namespace UnrealBuildTool
 			LinkerPath = GetLinkerToolPath(Compiler, Architecture, CompilerDir, ToolChainDir);
 			LibraryManagerPath = GetLibraryLinkerToolPath(Compiler, Architecture, ToolChainDir);
 			ToolchainCompilerPath = GetCompilerToolPath(ToolChain, Architecture, ToolChainDir);
+			ResourceToObjectConverterPath = GetResourceToObjectConverterToolPath(ToolChain, Architecture, ToolChainDir);
 
 			// Get the resource compiler path from the Windows SDK
 			ResourceCompilerPath = GetResourceCompilerToolPath(WindowsSdkDir, WindowsSdkVersion, Logger);
@@ -276,6 +282,14 @@ namespace UnrealBuildTool
 			return FileReference.Combine(GetVCToolPath(ToochainDir, Architecture), "link.exe"); // We add /LIB to cmd line so we can use link.exe directly instead of going via lib.exe
 		}
 
+		/// <summary>
+		/// Gets the path to the resource to object converter.
+		/// </summary>
+		protected virtual FileReference GetResourceToObjectConverterToolPath(WindowsCompiler Compiler, UnrealArch Architecture, DirectoryReference ToochainDir)
+		{
+			return FileReference.Combine(GetVCToolPath(ToochainDir, Architecture), "cvtres.exe");
+		}
+
 		/// <summary>
 		/// Gets the path to the resource compiler.
 		/// </summary>
diff --git a/Engine/Source/Programs/UnrealBuildTool/Platform/Windows/VCToolChain.cs b/Engine/Source/Programs/UnrealBuildTool/Platform/Windows/VCToolChain.cs
--- a/Engine/Source/Programs/UnrealBuildTool/Platform/Windows/VCToolChain.cs
+++ b/Engine/Source/Programs/UnrealBuildTool/Platform/Windows/VCToolChain.cs
@@ -5,6 +5,7 @@ using Microsoft.CodeAnalysis;
 using Microsoft.Extensions.Logging;
 using System;
 using System.Collections.Generic;
+using System.Diagnostics;
 using System.IO;
 using System.Linq;
 using System.Runtime.Versioning;
@@ -36,6 +37,7 @@ namespace UnrealBuildTool
 			Logger.LogDebug("Linker: {Path}", EnvVars.LinkerPath);
 			Logger.LogDebug("Library Manager: {Path}", EnvVars.LibraryManagerPath);
 			Logger.LogDebug("Resource Compiler: {Path}", EnvVars.ResourceCompilerPath);
+			Logger.LogDebug("Resource Converter: {Path}", EnvVars.ResourceToObjectConverterPath);
 
 			if (Target.WindowsPlatform.ObjSrcMapFile != null)
 			{
@@ -2738,7 +2740,7 @@ namespace UnrealBuildTool
 
 		public override CPPOutput CompileRCFiles(CppCompileEnvironment CompileEnvironment, IEnumerable<FileItem> InputFiles, DirectoryReference OutputDir, IActionGraphBuilder Graph)
 		{
-			CPPOutput Result = new CPPOutput();
+			CPPOutput Result = new();
 
 			CppRootPaths RootPaths = CompileEnvironment.RootPaths;
 
@@ -2746,7 +2748,7 @@ namespace UnrealBuildTool
 			{
 				Action CompileAction = Graph.CreateAction(ActionType.Compile);
 				CompileAction.RootPaths = CompileEnvironment.RootPaths;
-				CompileAction.CommandDescription = "Resource";
+				CompileAction.CommandDescription = $"Compile Resource [{CompileEnvironment.Architecture.WindowsName}]";
 				CompileAction.WorkingDirectory = Unreal.EngineSourceDirectory;
 				CompileAction.CommandPath = EnvVars.ResourceCompilerPath;
 				CompileAction.StatusDescription = Path.GetFileName(RCFile.AbsolutePath);
@@ -2836,12 +2838,7 @@ namespace UnrealBuildTool
 				}
 
 				// Add the RES file to the produced item list.
-				FileItem CompiledResourceFile = FileItem.GetItemByFileReference(
-					FileReference.Combine(
-						OutputDir,
-						Path.GetFileName(RCFile.AbsolutePath) + ".res"
-						)
-					);
+				FileItem CompiledResourceFile = FileItem.GetItemByFileReference(FileReference.Combine(OutputDir, $"{Path.GetFileName(RCFile.AbsolutePath)}.res"));
 				CompileAction.ProducedItems.Add(CompiledResourceFile);
 				Arguments.Add($"/fo \"{NormalizeCommandLinePath(CompiledResourceFile, RootPaths)}\"");
 				Result.ObjectFiles.Add(CompiledResourceFile);
@@ -2849,22 +2846,11 @@ namespace UnrealBuildTool
 				// Add the RC file as a prerequisite of the action.
 				Arguments.Add($"\"{NormalizeCommandLinePath(RCFile, RootPaths)}\"");
 
-				// Create a response file for the resource compilier
+				// Create a response file for the resource compilier for informational purposes as it cannot be used
 				FileItem ResponseFile = FileItem.GetItemByFileReference(GetResponseFileName(CompileEnvironment, CompiledResourceFile));
 				Graph.CreateIntermediateTextFile(ResponseFile, Arguments);
 				CompileAction.PrerequisiteItems.Add(ResponseFile);
 
-				/* rc.exe currently errors when using a response file
-				string ResponseFileString = NormalizeCommandLinePath(ResponseFile);
-
-				// cl.exe can't handle response files with a path longer than 260 characters, and relative paths can push it over the limit
-				if (!System.IO.Path.IsPathRooted(ResponseFileString) && System.IO.Path.Combine(CompileAction.WorkingDirectory.FullName, ResponseFileString).Length > 260)
-				{
-					ResponseFileString = ResponseFile.FullName;
-				}
-
-				CompileAction.CommandArguments = $"@{Utils.MakePathSafeToUseWithCommandLine(ResponseFileString)}";
-				*/
 				CompileAction.CommandArguments = String.Join(" ", Arguments);
 
 				// Add the C++ source file and its included files to the prerequisite item list.
@@ -2874,6 +2860,50 @@ namespace UnrealBuildTool
 			return Result;
 		}
 
+		public override CPPOutput ConvertRCFiles(LinkEnvironment LinkEnvironment, IEnumerable<FileItem> InputFiles, DirectoryReference OutputDir, IActionGraphBuilder Graph)
+		{
+			// Converting resource files to objects when using the clang linker is not supported
+			if (EnvVars.Compiler.IsClang() && EnvVars.bAllowClangLinker)
+			{
+				return base.ConvertRCFiles(LinkEnvironment, InputFiles, OutputDir, Graph);
+			}
+
+			FileItem ConvertedResourceObjectFile = FileItem.GetItemByFileReference(FileReference.Combine(OutputDir, $"{InputFiles.First().Name}.obj"));
+
+			List<string> ConvertArguments = [
+				"/nologo",
+				"/folddups",
+				$"/machine:{LinkEnvironment.Architecture.WindowsName}",
+				$"/out:\"{NormalizeCommandLinePath(ConvertedResourceObjectFile, LinkEnvironment.RootPaths)}\"",
+				.. InputFiles.Select(x => $"\"{NormalizeCommandLinePath(x, LinkEnvironment.RootPaths)}\"")
+			];
+
+			Action ConvertAction = Graph.CreateAction(ActionType.Compile);
+			ConvertAction.RootPaths = LinkEnvironment.RootPaths;
+			ConvertAction.CommandDescription = $"Convert Resource [{LinkEnvironment.Architecture.WindowsName}]";
+			ConvertAction.WorkingDirectory = Unreal.EngineSourceDirectory;
+			ConvertAction.CommandPath = EnvVars.ResourceToObjectConverterPath;
+			ConvertAction.CommandArguments = String.Join(' ', ConvertArguments);
+			ConvertAction.StatusDescription = Path.GetFileName(InputFiles.First().AbsolutePath);
+			ConvertAction.PrerequisiteItems.UnionWith(InputFiles);
+			ConvertAction.ProducedItems.Add(ConvertedResourceObjectFile);
+			ConvertAction.ArtifactMode = ArtifactMode.Enabled;
+			ConvertAction.CacheBucket = GetCacheBucket(Target, LinkEnvironment);
+
+			// Convert Resource tool can run remotely if possible
+			ConvertAction.bCanExecuteRemotely = true;
+			ConvertAction.bCanExecuteRemotelyWithSNDBS = false; // no tool template for SN-DBS results in warnings
+
+			// Create a response file for the resource converter arguments for informational purposes as it cannot be used
+			FileItem ResponseFile = FileItem.GetItemByFileReference(GetResponseFileName(LinkEnvironment, ConvertedResourceObjectFile));
+			Graph.CreateIntermediateTextFile(ResponseFile, ConvertArguments);
+			ConvertAction.PrerequisiteItems.Add(ResponseFile);
+
+			CPPOutput Result = new();
+			Result.ObjectFiles.Add(ConvertedResourceObjectFile);
+			return Result;
+		}
+
 		public override void GenerateTypeLibraryHeader(CppCompileEnvironment CompileEnvironment, ModuleRules.TypeLibrary TypeLibrary, FileReference OutputFile, FileReference? OutputHeader, IActionGraphBuilder Graph)
 		{
 			CppRootPaths RootPaths = CompileEnvironment.RootPaths;
@@ -3203,6 +3233,7 @@ namespace UnrealBuildTool
 
 			// Add the input files to a response file, and pass the response file on the command-line.
 			List<string> InputFileNames = new List<string>();
+
 			foreach (FileItem InputFile in LinkEnvironment.InputFiles)
 			{
 				InputFileNames.Add($"\"{NormalizeCommandLinePath(InputFile, RootPaths)}\"");
diff --git a/Engine/Source/Programs/UnrealBuildTool/System/LinkEnvironment.cs b/Engine/Source/Programs/UnrealBuildTool/System/LinkEnvironment.cs
--- a/Engine/Source/Programs/UnrealBuildTool/System/LinkEnvironment.cs
+++ b/Engine/Source/Programs/UnrealBuildTool/System/LinkEnvironment.cs
@@ -330,9 +330,9 @@ namespace UnrealBuildTool
 		public List<FileItem> DefaultResourceFiles = [];
 
 		/// <summary>
-		/// Resource files which should be compiled into every binary
+		/// The default resource file to link in to every binary if a custom one is not provided (converted to obj)
 		/// </summary>
-		public List<FileItem> CommonResourceFiles = [];
+		public List<FileItem> DefaultConvertedResourceFiles = [];
 
 		/// <summary>
 		/// List of functions that should be exported from this module
@@ -438,7 +438,7 @@ namespace UnrealBuildTool
 			InstallName = Other.InstallName;
 			InputFiles.AddRange(Other.InputFiles);
 			DefaultResourceFiles.AddRange(Other.DefaultResourceFiles);
-			CommonResourceFiles.AddRange(Other.CommonResourceFiles);
+			DefaultConvertedResourceFiles.AddRange(Other.DefaultConvertedResourceFiles);
 			IncludeFunctions.AddRange(Other.IncludeFunctions);
 			DebuggerVisualizerFiles.AddRange(Other.DebuggerVisualizerFiles);
 			ModuleDefinitionFile = Other.ModuleDefinitionFile;
diff --git a/Engine/Source/Programs/UnrealBuildTool/ToolChain/UEToolChain.cs b/Engine/Source/Programs/UnrealBuildTool/ToolChain/UEToolChain.cs
--- a/Engine/Source/Programs/UnrealBuildTool/ToolChain/UEToolChain.cs
+++ b/Engine/Source/Programs/UnrealBuildTool/ToolChain/UEToolChain.cs
@@ -134,8 +134,14 @@ namespace UnrealBuildTool
 		}
 
 		public virtual CPPOutput CompileRCFiles(CppCompileEnvironment CompileEnvironment, IEnumerable<FileItem> InputFiles, DirectoryReference OutputDir, IActionGraphBuilder Graph)
+		{
+			return new();
+		}
+
+		public virtual CPPOutput ConvertRCFiles(LinkEnvironment LinkEnvironment, IEnumerable<FileItem> InputFiles, DirectoryReference OutputDir, IActionGraphBuilder Graph)
 		{
 			CPPOutput Result = new();
+			Result.ObjectFiles.AddRange(InputFiles);
 			return Result;
 		}
 
