UnrealBuildTool: Analysis fixes

#rnx

diff --git a/Engine/Source/Programs/UnrealBuildTool/Configuration/UEBuildModule.cs b/Engine/Source/Programs/UnrealBuildTool/Configuration/UEBuildModule.cs
--- a/Engine/Source/Programs/UnrealBuildTool/Configuration/UEBuildModule.cs
+++ b/Engine/Source/Programs/UnrealBuildTool/Configuration/UEBuildModule.cs
@@ -346,7 +346,7 @@ namespace UnrealBuildTool
 				}
 			}
 
-			// Add any resource and debug visualizers found in the module directories
+			// Add any debug visualizers found in the module directories
 			foreach (DirectoryItem directory in ModuleDirectories)
 			{
 				NatvisFiles.UnionWith(directory.EnumerateFiles().Where(x => x.HasExtension(".natvis") || x.HasExtension(".natstepfilter") || x.HasExtension(".natjmc")));
diff --git a/Engine/Source/Programs/UnrealBuildTool/Platform/Windows/VCToolChain.cs b/Engine/Source/Programs/UnrealBuildTool/Platform/Windows/VCToolChain.cs
--- a/Engine/Source/Programs/UnrealBuildTool/Platform/Windows/VCToolChain.cs
+++ b/Engine/Source/Programs/UnrealBuildTool/Platform/Windows/VCToolChain.cs
@@ -2740,6 +2740,11 @@ namespace UnrealBuildTool
 
 		public override CPPOutput CompileRCFiles(CppCompileEnvironment CompileEnvironment, IEnumerable<FileItem> InputFiles, DirectoryReference OutputDir, IActionGraphBuilder Graph)
 		{
+			if (!InputFiles.Any())
+			{
+				return new();
+			}
+
 			CPPOutput Result = new();
 
 			CppRootPaths RootPaths = CompileEnvironment.RootPaths;
@@ -2862,6 +2867,11 @@ namespace UnrealBuildTool
 
 		public override CPPOutput ConvertRCFiles(LinkEnvironment LinkEnvironment, IEnumerable<FileItem> InputFiles, DirectoryReference OutputDir, IActionGraphBuilder Graph)
 		{
+			if (!InputFiles.Any())
+			{
+				return new();
+			}
+
 			// Converting resource files to objects when using the clang linker is not supported
 			if (EnvVars.Compiler.IsClang() && EnvVars.bAllowClangLinker)
 			{
