[UBT]
* Fixed so UEBuildModuleCPP.Compile is called in parallel for monolithic targets. This function was called in parallel on modular builds since there is one module per binary.

diff --git a/Engine/Source/Programs/UnrealBuildTool/Configuration/UEBuildBinary.cs b/Engine/Source/Programs/UnrealBuildTool/Configuration/UEBuildBinary.cs
--- a/Engine/Source/Programs/UnrealBuildTool/Configuration/UEBuildBinary.cs
+++ b/Engine/Source/Programs/UnrealBuildTool/Configuration/UEBuildBinary.cs
@@ -746,45 +746,60 @@ namespace UnrealBuildTool
 			{
 				if (Module.Binary == null || Module.Binary == this)
 				{
-					// Add the module's AutoRTFM external mapping files to the compile environment
-					BinaryCompileEnvironment.AutoRTFMExternalMappingFiles.UnionWith(Module.PublicAutoRTFMExternalMappingFiles);
-
-					// Compile each module.
-					Logger.LogTrace("Compile module: {ModuleName}", Module.Name);
-					List<FileItem> LinkInputFiles = Module.Compile(Target, ToolChain, BinaryCompileEnvironment, WorkingSet, Graph, Logger);
-					InputObjects.UnionWith(LinkInputFiles.Where(x => x.HasExtension(".obj") || x.HasExtension(".o")));
 					CompiledModules.Add(Module);
+				}
+				else
+				{
+					BinaryDependencies.Add(Module.Binary);
 
-					// Save the module outputs. In monolithic builds, this is just the object files.
-					if (Target.LinkType == TargetLinkType.Monolithic)
-					{
-						Graph.SetOutputItemsForModule(Module.Name, [.. LinkInputFiles]);
-					}
+					// Allow the module to modify the link environment for the binary.
+					Module.SetupPrivateLinkEnvironment(this, BinaryLinkEnvironment, BinaryDependencies, LinkEnvironmentVisitedModules, ExeDir, Logger);
+				}
+			}
 
-					// Sometimes the same link input file is added multiple times, so make sure to filter out the dups
-					foreach (FileItem LinkInputFile in LinkInputFiles)
-					{
-						if (InputFilesLookup.Add(LinkInputFile))
-						{
-							BinaryLinkEnvironment.InputFiles.Add(LinkInputFile);
-						}
-					}
+			List<FileItem>[] LinkInputFilesArray = new List<FileItem>[CompiledModules.Count];
 
-					// Project root path should only be added if the binary needs it, to prevent violating shared environment actions
-					if (Target.ProjectFile != null && BinaryLinkEnvironment.InputFiles.Any(x => x.Location.IsUnderDirectory(Target.ProjectFile.Directory)))
-					{
-						BinaryLinkEnvironment.RootPaths[CppRootPathFolder.Project] = Target.ProjectFile.Directory;
-					}
+			// Compile call is expensive so do it in parallel
+			Parallel2.For(0, CompiledModules.Count, index =>
+			{
+					LinkInputFilesArray[index] = CompiledModules[index].Compile(Target, ToolChain, BinaryCompileEnvironment, WorkingSet, Graph, Logger);
+			});
+
+			for (int index=0, e= CompiledModules.Count; index != e; ++index)
+			{
+				UEBuildModule Module = CompiledModules[index];
+
+				// Add the module's AutoRTFM external mapping files to the compile environment
+				BinaryCompileEnvironment.AutoRTFMExternalMappingFiles.UnionWith(Module.PublicAutoRTFMExternalMappingFiles);
+
+				List<FileItem> LinkInputFiles = LinkInputFilesArray[index];
+				InputObjects.UnionWith(LinkInputFiles.Where(x => x.HasExtension(".obj") || x.HasExtension(".o")));
+
+				// Save the module outputs. In monolithic builds, this is just the object files.
+				if (Target.LinkType == TargetLinkType.Monolithic)
+				{
+					Graph.SetOutputItemsForModule(Module.Name, [.. LinkInputFiles]);
+				}
 
-					// Force a reference to initialize module for this binary
-					if (Module.Rules.bRequiresImplementModule ?? true)
+				// Sometimes the same link input file is added multiple times, so make sure to filter out the dups
+				foreach (FileItem LinkInputFile in LinkInputFiles)
+				{
+					if (InputFilesLookup.Add(LinkInputFile))
 					{
-						BinaryLinkEnvironment.IncludeFunctions.Add($"IMPLEMENT_MODULE_{Module.Name}");
+						BinaryLinkEnvironment.InputFiles.Add(LinkInputFile);
 					}
 				}
-				else
+
+				// Project root path should only be added if the binary needs it, to prevent violating shared environment actions
+				if (Target.ProjectFile != null && BinaryLinkEnvironment.InputFiles.Any(x => x.Location.IsUnderDirectory(Target.ProjectFile.Directory)))
 				{
-					BinaryDependencies.Add(Module.Binary);
+					BinaryLinkEnvironment.RootPaths[CppRootPathFolder.Project] = Target.ProjectFile.Directory;
+				}
+
+				// Force a reference to initialize module for this binary
+				if (Module.Rules.bRequiresImplementModule ?? true)
+				{
+					BinaryLinkEnvironment.IncludeFunctions.Add($"IMPLEMENT_MODULE_{Module.Name}");
 				}
 
 				// Allow the module to modify the link environment for the binary.
diff --git a/Engine/Source/Programs/UnrealBuildTool/Configuration/UEBuildModule.cs b/Engine/Source/Programs/UnrealBuildTool/Configuration/UEBuildModule.cs
--- a/Engine/Source/Programs/UnrealBuildTool/Configuration/UEBuildModule.cs
+++ b/Engine/Source/Programs/UnrealBuildTool/Configuration/UEBuildModule.cs
@@ -1646,7 +1646,7 @@ namespace UnrealBuildTool
 		/// remain in their same relative order within the original Nodes sequence.
 		/// </summary>
 		/// <param name="NodeList">The list of nodes to sort.</param>
-		public static List<UEBuildModule> StableTopologicalSort(List<UEBuildModule> NodeList)
+		public static UEBuildModule[] StableTopologicalSort(List<UEBuildModule> NodeList)
 		{
 			int nodeCount = NodeList.Count;
 			if (nodeCount <= 1)
@@ -1709,94 +1709,12 @@ namespace UnrealBuildTool
 				outIdx.Insert(insertAt, i);
 			}
 
-			List<UEBuildModule> ordered = new (nodeCount);
+			UEBuildModule[] ordered = new UEBuildModule[nodeCount];
+			int counter = 0;
 			foreach (int idx in outIdx)
-				ordered.Add(NodeList[idx]);
-
-			#if false
-			// THIS IS THE OLD CODE
 			{
-				int NodeCount = NodeList.Count;
-
-				// For each Node in NodeList, populated with the full circular dependency list from
-				// Node.GetAllDependencyModules()
-				List<Task<HashSet<UEBuildModule>>> NodeDependencies = new List<Task<HashSet<UEBuildModule>>>(NodeCount);
-				// Used to populate an element of NodeDependencies
-				HashSet<UEBuildModule> FetchDependencies(int NodeIndex)
-				{
-					HashSet<UEBuildModule> Dependencies = new HashSet<UEBuildModule>();
-					NodeList[NodeIndex].GetAllDependencyModules(new List<UEBuildModule>(), Dependencies, true, true, false);
-					return Dependencies;
-				}
-
-				// For each Node in NodeList, populated with the nodes with a lower index in NodeList that Node depends on
-				List<Task<HashSet<UEBuildModule>>> PrecedingDependents = new List<Task<HashSet<UEBuildModule>>>(NodeCount);
-				HashSet<UEBuildModule> ComputePrecedingDependents(int NodeIndex)
-				{
-					HashSet<UEBuildModule> Results = new HashSet<UEBuildModule>();
-
-					UEBuildModule Node = NodeList[NodeIndex];
-					HashSet<UEBuildModule> Dependencies = NodeDependencies[NodeIndex].Result;
-
-					for (int I = 0; I < NodeIndex; ++I)
-					{
-						if (NodeDependencies[I].Result.Contains(Node) && !Dependencies.Contains(NodeList[I]))
-						{
-							Results.Add(NodeList[I]);
-						}
-					}
-
-					return Results;
-				}
-
-				for (int I = 0; I < NodeCount; ++I)
-				{
-					int LocalI = I;
-					NodeDependencies.Add(Task.Run(() => FetchDependencies(LocalI)));
-					PrecedingDependents.Add(Task.Run(() => ComputePrecedingDependents(LocalI)));
-				}
-
-				List<UEBuildModule> Out = new List<UEBuildModule>(NodeCount);
-				// Write the ordered output			
-				for (int Index1 = 0; Index1 != NodeCount; ++Index1)
-				{
-					UEBuildModule Node1 = NodeList[Index1];
-					HashSet<UEBuildModule> NodesThatDependOnNode1 = PrecedingDependents[Index1].Result;
-					Out.Add(Node1);
-
-					if (NodesThatDependOnNode1.Count == 0)
-					{
-						continue;
-					}
-
-					for (int Index2 = 0; Index2 != Index1; ++Index2)
-					{
-						UEBuildModule Node2 = Out[Index2];
-
-						if (NodesThatDependOnNode1.Contains(Node2))
-						{
-							// Rotate element at Index1 into position at Index2
-							for (int Index3 = Index1; Index3 != Index2;)
-							{
-								--Index3;
-								Out[Index3 + 1] = Out[Index3];
-							}
-
-							Out[Index2] = Node1;
-
-							// Break out of this loop, because this iteration must have covered all existing cases
-							// involving the node formerly at position Index1
-							break;
-						}
-					}
-				}
-				if (!ordered.SequenceEqual(Out))
-				{
-					throw new Exception("Sorting diff!!!");
-				}
+				ordered[counter++] = NodeList[idx];
 			}
-			#endif
-
 			return ordered;
 		}
 	};
