[UBT]
* Attempt at fixing compile error related to rc files. It seems like rc files set in BinaryLinkEnvironment.DefaultResourceFiles can not have resource versions set because then rc.exe fails. Don't really know why but this change makes it more similar to how it was before

diff --git a/Engine/Source/Programs/UnrealBuildTool/Configuration/UEBuildBinary.cs b/Engine/Source/Programs/UnrealBuildTool/Configuration/UEBuildBinary.cs
--- a/Engine/Source/Programs/UnrealBuildTool/Configuration/UEBuildBinary.cs
+++ b/Engine/Source/Programs/UnrealBuildTool/Configuration/UEBuildBinary.cs
@@ -871,6 +871,8 @@ namespace UnrealBuildTool
 				return;
 			}
 
+			FileItem? DefaultResourceFile = null;
+
 			// If we don't have any resource file, use the default or compile a custom one for this module
 			IEnumerable<FileItem> DefaultResourceFiles = [];
 			// Figure out if this binary has any custom resource files. Hacky check to ignore the resource file in the Launch module, since it contains dialogs that the engine needs and always needs to be included.
@@ -890,8 +892,8 @@ namespace UnrealBuildTool
 				else
 				{
 					// Otherwise compile the default resource file per-binary, so that it gets the correct ORIGINAL_FILE_NAME macro.
-					FileItem DefaultResourceFile = FileItem.GetItemByFileReference(FileReference.Combine(Unreal.EngineDirectory, "Build", "Windows", "Resources", "Default.rc2"));
-					InputFiles = [.. InputFiles, DefaultResourceFile];
+					DefaultResourceFile = FileItem.GetItemByFileReference(FileReference.Combine(Unreal.EngineDirectory, "Build", "Windows", "Resources", "Default.rc2"));
+					//InputFiles = [.. InputFiles, DefaultResourceFile];
 				}
 			}
 
@@ -901,28 +903,34 @@ namespace UnrealBuildTool
 			// Set the original file name macro; used in Default.rc2 to set the binary metadata fields.
 			ResourceCompileEnvironment.Definitions.Add($"ORIGINAL_FILE_NAME=\"{OutputFilePaths[0].GetFileName()}\"");
 
-			// Set the other version fields if requested
-			if (Target.WindowsPlatform.bSetResourceVersions)
+			// Compile the resource files
+			List<FileItem> ObjectFiles = ToolChain.CompileRCFiles(ResourceCompileEnvironment, InputFiles, IntermediateDirectory, Graph).ObjectFiles;
+
+			// Compile the default resource file if it exists
+			if (DefaultResourceFile != null)
 			{
-				if (Target.Version.Changelist != 0)
-				{
-					ResourceCompileEnvironment.Definitions.Add($"BUILT_FROM_CHANGELIST={Target.Version.Changelist}");
-				}
-				if (!String.IsNullOrEmpty(Target.Version.BranchName))
-				{
-					ResourceCompileEnvironment.Definitions.Add($"BRANCH_NAME={Target.Version.BranchName}");
-				}
-				if (!String.IsNullOrEmpty(Target.BuildVersion))
+				// Set the other version fields if requested
+				if (Target.WindowsPlatform.bSetResourceVersions)
 				{
-					ResourceCompileEnvironment.Definitions.Add($"BUILD_VERSION={Target.BuildVersion}");
+					if (Target.Version.Changelist != 0)
+					{
+						ResourceCompileEnvironment.Definitions.Add($"BUILT_FROM_CHANGELIST={Target.Version.Changelist}");
+					}
+					if (!String.IsNullOrEmpty(Target.Version.BranchName))
+					{
+						ResourceCompileEnvironment.Definitions.Add($"BRANCH_NAME={Target.Version.BranchName}");
+					}
+					if (!String.IsNullOrEmpty(Target.BuildVersion))
+					{
+						ResourceCompileEnvironment.Definitions.Add($"BUILD_VERSION={Target.BuildVersion}");
+					}
 				}
-			}
 
-			// Compile the resource files
-			CPPOutput Result = ToolChain.CompileRCFiles(ResourceCompileEnvironment, InputFiles, IntermediateDirectory, Graph);
+				ObjectFiles.AddRange(ToolChain.CompileRCFiles(ResourceCompileEnvironment, [DefaultResourceFile], IntermediateDirectory, Graph).ObjectFiles);
+			}
 
 			// Optionally convert the resource files depending on toolchain settings
-			BinaryLinkEnvironment.InputFiles.AddRange(ToolChain.ConvertRCFiles(BinaryLinkEnvironment, [.. Result.ObjectFiles, .. DefaultResourceFiles], IntermediateDirectory, Graph).ObjectFiles);
+			BinaryLinkEnvironment.InputFiles.AddRange(ToolChain.ConvertRCFiles(BinaryLinkEnvironment, [.. ObjectFiles, .. DefaultResourceFiles], IntermediateDirectory, Graph).ObjectFiles);
 		}
 
 		/// <summary>
