[Backout] - CL46791144
#rnx
Original CL Desc
-----------------------------------------------------------------
UnrealBuildTool: Convert resource files to objects manually so link.exe doesn't need to prevent unnecessary reprocessing when resource files are shared
Remove LinkEnvironment.CommonResourceFiles as it is unused

#rnx

#FailedValidations Preflight

diff --git a/Engine/Source/Programs/UnrealBuildTool/Configuration/UEBuildBinary.cs b/Engine/Source/Programs/UnrealBuildTool/Configuration/UEBuildBinary.cs
--- a/Engine/Source/Programs/UnrealBuildTool/Configuration/UEBuildBinary.cs
+++ b/Engine/Source/Programs/UnrealBuildTool/Configuration/UEBuildBinary.cs
@@ -1,10 +1,10 @@
 // Copyright Epic Games, Inc. All Rights Reserved.
 
-using EpicGames.Core;
-using Microsoft.Extensions.Logging;
 using System;
 using System.Collections.Generic;
 using System.Linq;
+using EpicGames.Core;
+using Microsoft.Extensions.Logging;
 using UnrealBuildBase;
 
 namespace UnrealBuildTool
@@ -741,65 +741,48 @@ namespace UnrealBuildTool
 			using Pooled<HashSet<FileItem>> rented2 = HashSetPool<FileItem>.Rent();
 			HashSet<FileItem> InputFilesLookup = rented2.Value;
 
-			List<UEBuildModule> CompiledModules = [];
 			foreach (UEBuildModule Module in UEBuildModule.StableTopologicalSort(Modules))
 			{
 				if (Module.Binary == null || Module.Binary == this)
 				{
-					CompiledModules.Add(Module);
-				}
-				else
-				{
-					BinaryDependencies.Add(Module.Binary);
-
-					// Allow the module to modify the link environment for the binary.
-					Module.SetupPrivateLinkEnvironment(this, BinaryLinkEnvironment, BinaryDependencies, LinkEnvironmentVisitedModules, ExeDir, Logger);
-				}
-			}
+					// Add the module's AutoRTFM external mapping files to the compile environment
+					BinaryCompileEnvironment.AutoRTFMExternalMappingFiles.UnionWith(Module.PublicAutoRTFMExternalMappingFiles);
 
-			List<FileItem>[] LinkInputFilesArray = new List<FileItem>[CompiledModules.Count];
+					// Compile each module.
+					Logger.LogTrace("Compile module: {ModuleName}", Module.Name);
+					List<FileItem> LinkInputFiles = Module.Compile(Target, ToolChain, BinaryCompileEnvironment, WorkingSet, Graph, Logger);
+					InputObjects.UnionWith(LinkInputFiles.Where(x => x.HasExtension(".obj") || x.HasExtension(".o")));
 
-			// Compile call is expensive so do it in parallel
-			Parallel2.For(0, CompiledModules.Count, index =>
-			{
-					LinkInputFilesArray[index] = CompiledModules[index].Compile(Target, ToolChain, BinaryCompileEnvironment, WorkingSet, Graph, Logger);
-			});
-
-			for (int index=0, e= CompiledModules.Count; index != e; ++index)
-			{
-				UEBuildModule Module = CompiledModules[index];
-
-				// Add the module's AutoRTFM external mapping files to the compile environment
-				BinaryCompileEnvironment.AutoRTFMExternalMappingFiles.UnionWith(Module.PublicAutoRTFMExternalMappingFiles);
-
-				List<FileItem> LinkInputFiles = LinkInputFilesArray[index];
-				InputObjects.UnionWith(LinkInputFiles.Where(x => x.HasExtension(".obj") || x.HasExtension(".o")));
+					// Save the module outputs. In monolithic builds, this is just the object files.
+					if (Target.LinkType == TargetLinkType.Monolithic)
+					{
+						Graph.SetOutputItemsForModule(Module.Name, LinkInputFiles.ToArray());
+					}
 
-				// Save the module outputs. In monolithic builds, this is just the object files.
-				if (Target.LinkType == TargetLinkType.Monolithic)
-				{
-					Graph.SetOutputItemsForModule(Module.Name, [.. LinkInputFiles]);
-				}
+					// Sometimes the same link input file is added multiple times, so make sure to filter out the dups
+					foreach (FileItem LinkInputFile in LinkInputFiles)
+					{
+						if (InputFilesLookup.Add(LinkInputFile))
+						{
+							BinaryLinkEnvironment.InputFiles.Add(LinkInputFile);
+						}
+					}
 
-				// Sometimes the same link input file is added multiple times, so make sure to filter out the dups
-				foreach (FileItem LinkInputFile in LinkInputFiles)
-				{
-					if (InputFilesLookup.Add(LinkInputFile))
+					// Project root path should only be added if the binary needs it, to prevent violating shared environment actions
+					if (Target.ProjectFile != null && BinaryLinkEnvironment.InputFiles.Any(x => x.Location.IsUnderDirectory(Target.ProjectFile.Directory)))
 					{
-						BinaryLinkEnvironment.InputFiles.Add(LinkInputFile);
+						BinaryLinkEnvironment.RootPaths[CppRootPathFolder.Project] = Target.ProjectFile.Directory;
 					}
-				}
 
-				// Project root path should only be added if the binary needs it, to prevent violating shared environment actions
-				if (Target.ProjectFile != null && BinaryLinkEnvironment.InputFiles.Any(x => x.Location.IsUnderDirectory(Target.ProjectFile.Directory)))
-				{
-					BinaryLinkEnvironment.RootPaths[CppRootPathFolder.Project] = Target.ProjectFile.Directory;
+					// Force a reference to initialize module for this binary
+					if (Module.Rules.bRequiresImplementModule ?? true)
+					{
+						BinaryLinkEnvironment.IncludeFunctions.Add(String.Format("IMPLEMENT_MODULE_{0}", Module.Name));
+					}
 				}
-
-				// Force a reference to initialize module for this binary
-				if (Module.Rules.bRequiresImplementModule ?? true)
+				else
 				{
-					BinaryLinkEnvironment.IncludeFunctions.Add($"IMPLEMENT_MODULE_{Module.Name}");
+					BinaryDependencies.Add(Module.Binary);
 				}
 
 				// Allow the module to modify the link environment for the binary.
@@ -819,7 +802,7 @@ namespace UnrealBuildTool
 			}
 
 			// Set the link output file.
-			BinaryLinkEnvironment.OutputFilePaths = [.. OutputFilePaths];
+			BinaryLinkEnvironment.OutputFilePaths = OutputFilePaths.ToList();
 
 			// Remember the link type
 			BinaryLinkEnvironment.LinkType = Target.LinkType;
@@ -837,9 +820,6 @@ namespace UnrealBuildTool
 			BinaryLinkEnvironment.bIsBuildingDLL = IsBuildingDll(Type);
 			BinaryLinkEnvironment.bIsBuildingLibrary = IsBuildingLibrary(Type);
 
-			// Compile any resource files
-			CompileResourceFiles(Target, ToolChain, BinaryCompileEnvironment, BinaryLinkEnvironment, CompiledModules.SelectMany(x => x.ResourceFiles), Graph, Logger);
-
 			// Setup object export stripping
 			if (Target.bStripExports)
 			{
@@ -854,83 +834,71 @@ namespace UnrealBuildTool
 			// Code coverage inherited from compile environment
 			BinaryLinkEnvironment.bCodeCoverage = CompileEnvironment.bCodeCoverage;
 
-			// Provide list of runtime dependencies to link environment.
-			BinaryLinkEnvironment.RuntimeDependencies = RuntimeDependencies;
-
-			// Allow the platform to modify the binary link environment for platform-specific resources etc.
-			UEBuildPlatform.GetBuildPlatform(Target.Platform).ModifyBinaryLinkEnvironment(BinaryLinkEnvironment, BinaryCompileEnvironment, Target, ToolChain, Graph);
-
-			return BinaryLinkEnvironment;
-		}
-
-		private void CompileResourceFiles(ReadOnlyTargetRules Target, UEToolChain ToolChain, CppCompileEnvironment BinaryCompileEnvironment, LinkEnvironment BinaryLinkEnvironment, IEnumerable<FileItem> InputFiles, IActionGraphBuilder Graph, ILogger Logger)
-		{
-			// @todo: This should be in some Windows code somewhere...
-			if (!BinaryLinkEnvironment.Platform.IsInGroup(UnrealPlatformGroup.Windows))
-			{
-				return;
-			}
-
-			FileItem? DefaultResourceFile = null;
-
 			// If we don't have any resource file, use the default or compile a custom one for this module
-			IEnumerable<FileItem> DefaultResourceFiles = [];
-			// Figure out if this binary has any custom resource files. Hacky check to ignore the resource file in the Launch module, since it contains dialogs that the engine needs and always needs to be included.
-			FileItem PCLaunchResource = FileItem.GetItemByFileReference(FileReference.Combine(Unreal.EngineSourceDirectory, "Runtime", "Launch", "Resources", "Windows", "PCLaunch.rc"));
-			if (!InputFiles.Any(x => x != PCLaunchResource))
+			if (BinaryLinkEnvironment.Platform.IsInGroup(UnrealPlatformGroup.Windows))
 			{
-				if (BinaryLinkEnvironment.DefaultResourceFiles.Count > 0)
+				// Figure out if this binary has any custom resource files. Hacky check to ignore the resource file in the Launch module, since it contains dialogs that the engine needs and always needs to be included.
+				FileItem[] CustomResourceFiles = BinaryLinkEnvironment.InputFiles.Where(x => x.Location.HasExtension(".res") && !x.Location.FullName.EndsWith("\\Launch\\PCLaunch.rc.res", StringComparison.OrdinalIgnoreCase)).ToArray();
+				if (CustomResourceFiles.Length == 0)
 				{
-					// If no resource files are found, use the precompiled default resource file for the target
-					if (!InputFiles.Any())
+					if (BinaryLinkEnvironment.DefaultResourceFiles.Count > 0)
 					{
-						BinaryLinkEnvironment.InputFiles.AddRange(BinaryLinkEnvironment.DefaultConvertedResourceFiles);
-						return;
+						// Use the default resource file if possible
+						BinaryLinkEnvironment.InputFiles.AddRange(BinaryLinkEnvironment.DefaultResourceFiles);
 					}
-					DefaultResourceFiles = BinaryLinkEnvironment.DefaultResourceFiles;
-				}
-				else
-				{
-					// Otherwise compile the default resource file per-binary, so that it gets the correct ORIGINAL_FILE_NAME macro.
-					DefaultResourceFile = FileItem.GetItemByFileReference(FileReference.Combine(Unreal.EngineDirectory, "Build", "Windows", "Resources", "Default.rc2"));
-					//InputFiles = [.. InputFiles, DefaultResourceFile];
-				}
-			}
+					else
+					{
+						// Get the intermediate directory
+						DirectoryReference ResourceIntermediateDirectory = BinaryLinkEnvironment.IntermediateDirectory;
 
-			// Create a compile environment for resource files
-			CppCompileEnvironment ResourceCompileEnvironment = new(BinaryCompileEnvironment);
+						// Place resource intermediate files for executables in a subfolder
+						if (!BinaryLinkEnvironment.bIsBuildingDLL && !BinaryLinkEnvironment.bIsBuildingLibrary)
+						{
+							ResourceIntermediateDirectory = DirectoryReference.Combine(ResourceIntermediateDirectory, OutputFilePaths[0].GetFileName());
+						}
 
-			// Set the original file name macro; used in Default.rc2 to set the binary metadata fields.
-			ResourceCompileEnvironment.Definitions.Add($"ORIGINAL_FILE_NAME=\"{OutputFilePaths[0].GetFileName()}\"");
+						// Create a compile environment for resource files
+						CppCompileEnvironment ResourceCompileEnvironment = new CppCompileEnvironment(BinaryCompileEnvironment);
 
-			// Compile the resource files
-			List<FileItem> ObjectFiles = ToolChain.CompileRCFiles(ResourceCompileEnvironment, InputFiles, IntermediateDirectory, Graph).ObjectFiles;
+						// @todo: This should be in some Windows code somewhere...
+						// Set the original file name macro; used in Default.rc2 to set the binary metadata fields.
+						ResourceCompileEnvironment.Definitions.Add("ORIGINAL_FILE_NAME=\"" + OutputFilePaths[0].GetFileName() + "\"");
 
-			// Compile the default resource file if it exists
-			if (DefaultResourceFile != null)
-			{
-				// Set the other version fields if requested
-				if (Target.WindowsPlatform.bSetResourceVersions)
-				{
-					if (Target.Version.Changelist != 0)
-					{
-						ResourceCompileEnvironment.Definitions.Add($"BUILT_FROM_CHANGELIST={Target.Version.Changelist}");
-					}
-					if (!String.IsNullOrEmpty(Target.Version.BranchName))
-					{
-						ResourceCompileEnvironment.Definitions.Add($"BRANCH_NAME={Target.Version.BranchName}");
-					}
-					if (!String.IsNullOrEmpty(Target.BuildVersion))
-					{
-						ResourceCompileEnvironment.Definitions.Add($"BUILD_VERSION={Target.BuildVersion}");
+						// Set the other version fields if requested
+						if (Target.WindowsPlatform.bSetResourceVersions)
+						{
+							if (Target.Version.Changelist != 0)
+							{
+								ResourceCompileEnvironment.Definitions.Add(String.Format("BUILT_FROM_CHANGELIST={0}", Target.Version.Changelist));
+							}
+							if (!String.IsNullOrEmpty(Target.Version.BranchName))
+							{
+								ResourceCompileEnvironment.Definitions.Add(String.Format("BRANCH_NAME={0}", Target.Version.BranchName));
+							}
+							if (!String.IsNullOrEmpty(Target.BuildVersion))
+							{
+								ResourceCompileEnvironment.Definitions.Add(String.Format("BUILD_VERSION={0}", Target.BuildVersion));
+							}
+						}
+
+						// Otherwise compile the default resource file per-binary, so that it gets the correct ORIGINAL_FILE_NAME macro.
+						FileItem DefaultResourceFile = FileItem.GetItemByFileReference(FileReference.Combine(Unreal.EngineDirectory, "Build", "Windows", "Resources", "Default.rc2"));
+						CPPOutput DefaultResourceOutput = ToolChain.CompileRCFiles(ResourceCompileEnvironment, new List<FileItem> { DefaultResourceFile }, ResourceIntermediateDirectory, Graph);
+						BinaryLinkEnvironment.InputFiles.AddRange(DefaultResourceOutput.ObjectFiles);
 					}
 				}
-
-				ObjectFiles.AddRange(ToolChain.CompileRCFiles(ResourceCompileEnvironment, [DefaultResourceFile], IntermediateDirectory, Graph).ObjectFiles);
 			}
 
-			// Optionally convert the resource files depending on toolchain settings
-			BinaryLinkEnvironment.InputFiles.AddRange(ToolChain.ConvertRCFiles(BinaryLinkEnvironment, [.. ObjectFiles, .. DefaultResourceFiles], IntermediateDirectory, Graph).ObjectFiles);
+			// Add all the common resource files
+			BinaryLinkEnvironment.InputFiles.AddRange(BinaryLinkEnvironment.CommonResourceFiles);
+
+			// Provide list of runtime dependencies to link environment.
+			BinaryLinkEnvironment.RuntimeDependencies = RuntimeDependencies;
+
+			// Allow the platform to modify the binary link environment for platform-specific resources etc.
+			UEBuildPlatform.GetBuildPlatform(Target.Platform).ModifyBinaryLinkEnvironment(BinaryLinkEnvironment, BinaryCompileEnvironment, Target, ToolChain, Graph);
+
+			return BinaryLinkEnvironment;
 		}
 
 		/// <summary>
diff --git a/Engine/Source/Programs/UnrealBuildTool/Configuration/UEBuildModule.cs b/Engine/Source/Programs/UnrealBuildTool/Configuration/UEBuildModule.cs
--- a/Engine/Source/Programs/UnrealBuildTool/Configuration/UEBuildModule.cs
+++ b/Engine/Source/Programs/UnrealBuildTool/Configuration/UEBuildModule.cs
@@ -203,11 +203,6 @@ namespace UnrealBuildTool
 		/// </summary>
 		public bool bDependsOnVerse = false;
 
-		/// <summary>
-		/// Set of resource file paths
-		/// </summary>
-		public HashSet<FileItem> ResourceFiles = new();
-
 		/// <summary>
 		/// Set of debug visualizer paths
 		/// </summary>
@@ -1646,7 +1641,7 @@ namespace UnrealBuildTool
 		/// remain in their same relative order within the original Nodes sequence.
 		/// </summary>
 		/// <param name="NodeList">The list of nodes to sort.</param>
-		public static UEBuildModule[] StableTopologicalSort(List<UEBuildModule> NodeList)
+		public static List<UEBuildModule> StableTopologicalSort(List<UEBuildModule> NodeList)
 		{
 			int nodeCount = NodeList.Count;
 			if (nodeCount <= 1)
@@ -1709,12 +1704,94 @@ namespace UnrealBuildTool
 				outIdx.Insert(insertAt, i);
 			}
 
-			UEBuildModule[] ordered = new UEBuildModule[nodeCount];
-			int counter = 0;
+			List<UEBuildModule> ordered = new (nodeCount);
 			foreach (int idx in outIdx)
+				ordered.Add(NodeList[idx]);
+
+			#if false
+			// THIS IS THE OLD CODE
 			{
-				ordered[counter++] = NodeList[idx];
+				int NodeCount = NodeList.Count;
+
+				// For each Node in NodeList, populated with the full circular dependency list from
+				// Node.GetAllDependencyModules()
+				List<Task<HashSet<UEBuildModule>>> NodeDependencies = new List<Task<HashSet<UEBuildModule>>>(NodeCount);
+				// Used to populate an element of NodeDependencies
+				HashSet<UEBuildModule> FetchDependencies(int NodeIndex)
+				{
+					HashSet<UEBuildModule> Dependencies = new HashSet<UEBuildModule>();
+					NodeList[NodeIndex].GetAllDependencyModules(new List<UEBuildModule>(), Dependencies, true, true, false);
+					return Dependencies;
+				}
+
+				// For each Node in NodeList, populated with the nodes with a lower index in NodeList that Node depends on
+				List<Task<HashSet<UEBuildModule>>> PrecedingDependents = new List<Task<HashSet<UEBuildModule>>>(NodeCount);
+				HashSet<UEBuildModule> ComputePrecedingDependents(int NodeIndex)
+				{
+					HashSet<UEBuildModule> Results = new HashSet<UEBuildModule>();
+
+					UEBuildModule Node = NodeList[NodeIndex];
+					HashSet<UEBuildModule> Dependencies = NodeDependencies[NodeIndex].Result;
+
+					for (int I = 0; I < NodeIndex; ++I)
+					{
+						if (NodeDependencies[I].Result.Contains(Node) && !Dependencies.Contains(NodeList[I]))
+						{
+							Results.Add(NodeList[I]);
+						}
+					}
+
+					return Results;
+				}
+
+				for (int I = 0; I < NodeCount; ++I)
+				{
+					int LocalI = I;
+					NodeDependencies.Add(Task.Run(() => FetchDependencies(LocalI)));
+					PrecedingDependents.Add(Task.Run(() => ComputePrecedingDependents(LocalI)));
+				}
+
+				List<UEBuildModule> Out = new List<UEBuildModule>(NodeCount);
+				// Write the ordered output			
+				for (int Index1 = 0; Index1 != NodeCount; ++Index1)
+				{
+					UEBuildModule Node1 = NodeList[Index1];
+					HashSet<UEBuildModule> NodesThatDependOnNode1 = PrecedingDependents[Index1].Result;
+					Out.Add(Node1);
+
+					if (NodesThatDependOnNode1.Count == 0)
+					{
+						continue;
+					}
+
+					for (int Index2 = 0; Index2 != Index1; ++Index2)
+					{
+						UEBuildModule Node2 = Out[Index2];
+
+						if (NodesThatDependOnNode1.Contains(Node2))
+						{
+							// Rotate element at Index1 into position at Index2
+							for (int Index3 = Index1; Index3 != Index2;)
+							{
+								--Index3;
+								Out[Index3 + 1] = Out[Index3];
+							}
+
+							Out[Index2] = Node1;
+
+							// Break out of this loop, because this iteration must have covered all existing cases
+							// involving the node formerly at position Index1
+							break;
+						}
+					}
+				}
+				if (!ordered.SequenceEqual(Out))
+				{
+					throw new Exception("Sorting diff!!!");
+				}
 			}
+			#endif
+
 			return ordered;
 		}
 	};
diff --git a/Engine/Source/Programs/UnrealBuildTool/Configuration/UEBuildModuleCPP.cs b/Engine/Source/Programs/UnrealBuildTool/Configuration/UEBuildModuleCPP.cs
--- a/Engine/Source/Programs/UnrealBuildTool/Configuration/UEBuildModuleCPP.cs
+++ b/Engine/Source/Programs/UnrealBuildTool/Configuration/UEBuildModuleCPP.cs
@@ -1152,8 +1152,19 @@ namespace UnrealBuildTool
 				};
 			}
 
-			// RC files for all modules in a binary will be compiled in UEBuildBinary.cs using the BinaryCompileEnvironment
-			ResourceFiles.UnionWith(InputFiles.RCFiles);
+			// Compile RC files. The resource compiler does not work with response files, and using the regular compile environment can easily result in the 
+			// command line length exceeding the OS limit. Use the binary compile environment to keep the size down, and require that all include paths
+			// must be specified relative to the resource file itself or Engine/Source.
+			if (InputFiles.RCFiles.Count > 0)
+			{
+				CppCompileEnvironment ResourceCompileEnvironment = new CppCompileEnvironment(BinaryCompileEnvironment);
+				if (Binary != null)
+				{
+					// @todo: This should be in some Windows code somewhere...
+					ResourceCompileEnvironment.Definitions.Add("ORIGINAL_FILE_NAME=\"" + Binary.OutputFilePaths[0].GetFileName() + "\"");
+				}
+				LinkInputFiles.AddRange(ToolChain.CompileRCFiles(ResourceCompileEnvironment, InputFiles.RCFiles, IntermediateDirectory, Graph).ObjectFiles);
+			}
 
 			// Write the compiled manifest
 			if (Rules.bPrecompile && Target.LinkType == TargetLinkType.Monolithic)
diff --git a/Engine/Source/Programs/UnrealBuildTool/Configuration/UEBuildTarget.cs b/Engine/Source/Programs/UnrealBuildTool/Configuration/UEBuildTarget.cs
--- a/Engine/Source/Programs/UnrealBuildTool/Configuration/UEBuildTarget.cs
+++ b/Engine/Source/Programs/UnrealBuildTool/Configuration/UEBuildTarget.cs
@@ -2975,14 +2975,12 @@ namespace UnrealBuildTool
 				FileReference DefaultResourceLocation = FileReference.Combine(Unreal.EngineDirectory, "Build", "Windows", "Resources", "Default.rc2");
 				if (!Unreal.IsFileInstalled(DefaultResourceLocation))
 				{
-					CppCompileEnvironment DefaultResourceCompileEnvironment = new(GlobalCompileEnvironment);
+					CppCompileEnvironment DefaultResourceCompileEnvironment = new CppCompileEnvironment(GlobalCompileEnvironment);
 
 					FileItem DefaultResourceFile = FileItem.GetItemByFileReference(DefaultResourceLocation);
 
-					CPPOutput DefaultResourceOutput = TargetToolChain.CompileRCFiles(DefaultResourceCompileEnvironment, [DefaultResourceFile], EngineIntermediateDirectory, MakefileBuilder);
+					CPPOutput DefaultResourceOutput = TargetToolChain.CompileRCFiles(DefaultResourceCompileEnvironment, new List<FileItem> { DefaultResourceFile }, EngineIntermediateDirectory, MakefileBuilder);
 					GlobalLinkEnvironment.DefaultResourceFiles.AddRange(DefaultResourceOutput.ObjectFiles);
-					DefaultResourceOutput = TargetToolChain.ConvertRCFiles(GlobalLinkEnvironment, DefaultResourceOutput.ObjectFiles, EngineIntermediateDirectory, MakefileBuilder);
-					GlobalLinkEnvironment.DefaultConvertedResourceFiles.AddRange(DefaultResourceOutput.ObjectFiles);
 				}
 			}
 
diff --git a/Engine/Source/Programs/UnrealBuildTool/Platform/Windows/VCEnvironment.cs b/Engine/Source/Programs/UnrealBuildTool/Platform/Windows/VCEnvironment.cs
--- a/Engine/Source/Programs/UnrealBuildTool/Platform/Windows/VCEnvironment.cs
+++ b/Engine/Source/Programs/UnrealBuildTool/Platform/Windows/VCEnvironment.cs
@@ -96,11 +96,6 @@ namespace UnrealBuildTool
 		/// </summary>
 		public readonly FileReference ResourceCompilerPath;
 
-		/// <summary>
-		/// Path to the resource to object converter from the toolchain
-		/// </summary>
-		public readonly FileReference ResourceToObjectConverterPath;
-
 		/// <summary>
 		/// The path to the toolchain compiler for compiling code
 		/// </summary>
@@ -148,7 +143,6 @@ namespace UnrealBuildTool
 			LinkerPath = GetLinkerToolPath(Compiler, Architecture, CompilerDir, ToolChainDir);
 			LibraryManagerPath = GetLibraryLinkerToolPath(Compiler, Architecture, ToolChainDir);
 			ToolchainCompilerPath = GetCompilerToolPath(ToolChain, Architecture, ToolChainDir);
-			ResourceToObjectConverterPath = GetResourceToObjectConverterToolPath(ToolChain, Architecture, ToolChainDir);
 
 			// Get the resource compiler path from the Windows SDK
 			ResourceCompilerPath = GetResourceCompilerToolPath(WindowsSdkDir, WindowsSdkVersion, Logger);
@@ -282,14 +276,6 @@ namespace UnrealBuildTool
 			return FileReference.Combine(GetVCToolPath(ToochainDir, Architecture), "link.exe"); // We add /LIB to cmd line so we can use link.exe directly instead of going via lib.exe
 		}
 
-		/// <summary>
-		/// Gets the path to the resource to object converter.
-		/// </summary>
-		protected virtual FileReference GetResourceToObjectConverterToolPath(WindowsCompiler Compiler, UnrealArch Architecture, DirectoryReference ToochainDir)
-		{
-			return FileReference.Combine(GetVCToolPath(ToochainDir, Architecture), "cvtres.exe");
-		}
-
 		/// <summary>
 		/// Gets the path to the resource compiler.
 		/// </summary>
diff --git a/Engine/Source/Programs/UnrealBuildTool/Platform/Windows/VCToolChain.cs b/Engine/Source/Programs/UnrealBuildTool/Platform/Windows/VCToolChain.cs
--- a/Engine/Source/Programs/UnrealBuildTool/Platform/Windows/VCToolChain.cs
+++ b/Engine/Source/Programs/UnrealBuildTool/Platform/Windows/VCToolChain.cs
@@ -5,7 +5,6 @@ using Microsoft.CodeAnalysis;
 using Microsoft.Extensions.Logging;
 using System;
 using System.Collections.Generic;
-using System.Diagnostics;
 using System.IO;
 using System.Linq;
 using System.Runtime.Versioning;
@@ -37,7 +36,6 @@ namespace UnrealBuildTool
 			Logger.LogDebug("Linker: {Path}", EnvVars.LinkerPath);
 			Logger.LogDebug("Library Manager: {Path}", EnvVars.LibraryManagerPath);
 			Logger.LogDebug("Resource Compiler: {Path}", EnvVars.ResourceCompilerPath);
-			Logger.LogDebug("Resource Converter: {Path}", EnvVars.ResourceToObjectConverterPath);
 
 			if (Target.WindowsPlatform.ObjSrcMapFile != null)
 			{
@@ -2740,12 +2738,7 @@ namespace UnrealBuildTool
 
 		public override CPPOutput CompileRCFiles(CppCompileEnvironment CompileEnvironment, IEnumerable<FileItem> InputFiles, DirectoryReference OutputDir, IActionGraphBuilder Graph)
 		{
-			if (!InputFiles.Any())
-			{
-				return new();
-			}
-
-			CPPOutput Result = new();
+			CPPOutput Result = new CPPOutput();
 
 			CppRootPaths RootPaths = CompileEnvironment.RootPaths;
 
@@ -2753,7 +2746,7 @@ namespace UnrealBuildTool
 			{
 				Action CompileAction = Graph.CreateAction(ActionType.Compile);
 				CompileAction.RootPaths = CompileEnvironment.RootPaths;
-				CompileAction.CommandDescription = $"Compile Resource [{CompileEnvironment.Architecture.WindowsName}]";
+				CompileAction.CommandDescription = "Resource";
 				CompileAction.WorkingDirectory = Unreal.EngineSourceDirectory;
 				CompileAction.CommandPath = EnvVars.ResourceCompilerPath;
 				CompileAction.StatusDescription = Path.GetFileName(RCFile.AbsolutePath);
@@ -2843,7 +2836,12 @@ namespace UnrealBuildTool
 				}
 
 				// Add the RES file to the produced item list.
-				FileItem CompiledResourceFile = FileItem.GetItemByFileReference(FileReference.Combine(OutputDir, $"{Path.GetFileName(RCFile.AbsolutePath)}.res"));
+				FileItem CompiledResourceFile = FileItem.GetItemByFileReference(
+					FileReference.Combine(
+						OutputDir,
+						Path.GetFileName(RCFile.AbsolutePath) + ".res"
+						)
+					);
 				CompileAction.ProducedItems.Add(CompiledResourceFile);
 				Arguments.Add($"/fo \"{NormalizeCommandLinePath(CompiledResourceFile, RootPaths)}\"");
 				Result.ObjectFiles.Add(CompiledResourceFile);
@@ -2851,11 +2849,22 @@ namespace UnrealBuildTool
 				// Add the RC file as a prerequisite of the action.
 				Arguments.Add($"\"{NormalizeCommandLinePath(RCFile, RootPaths)}\"");
 
-				// Create a response file for the resource compilier for informational purposes as it cannot be used
+				// Create a response file for the resource compilier
 				FileItem ResponseFile = FileItem.GetItemByFileReference(GetResponseFileName(CompileEnvironment, CompiledResourceFile));
 				Graph.CreateIntermediateTextFile(ResponseFile, Arguments);
 				CompileAction.PrerequisiteItems.Add(ResponseFile);
 
+				/* rc.exe currently errors when using a response file
+				string ResponseFileString = NormalizeCommandLinePath(ResponseFile);
+
+				// cl.exe can't handle response files with a path longer than 260 characters, and relative paths can push it over the limit
+				if (!System.IO.Path.IsPathRooted(ResponseFileString) && System.IO.Path.Combine(CompileAction.WorkingDirectory.FullName, ResponseFileString).Length > 260)
+				{
+					ResponseFileString = ResponseFile.FullName;
+				}
+
+				CompileAction.CommandArguments = $"@{Utils.MakePathSafeToUseWithCommandLine(ResponseFileString)}";
+				*/
 				CompileAction.CommandArguments = String.Join(" ", Arguments);
 
 				// Add the C++ source file and its included files to the prerequisite item list.
@@ -2865,55 +2874,6 @@ namespace UnrealBuildTool
 			return Result;
 		}
 
-		public override CPPOutput ConvertRCFiles(LinkEnvironment LinkEnvironment, IEnumerable<FileItem> InputFiles, DirectoryReference OutputDir, IActionGraphBuilder Graph)
-		{
-			if (!InputFiles.Any())
-			{
-				return new();
-			}
-
-			// Converting resource files to objects when using the clang linker is not supported
-			if (EnvVars.Compiler.IsClang() && EnvVars.bAllowClangLinker)
-			{
-				return base.ConvertRCFiles(LinkEnvironment, InputFiles, OutputDir, Graph);
-			}
-
-			FileItem ConvertedResourceObjectFile = FileItem.GetItemByFileReference(FileReference.Combine(OutputDir, $"{InputFiles.First().Name}.obj"));
-
-			List<string> ConvertArguments = [
-				"/nologo",
-				"/folddups",
-				$"/machine:{LinkEnvironment.Architecture.WindowsName}",
-				$"/out:\"{NormalizeCommandLinePath(ConvertedResourceObjectFile, LinkEnvironment.RootPaths)}\"",
-				.. InputFiles.Select(x => $"\"{NormalizeCommandLinePath(x, LinkEnvironment.RootPaths)}\"")
-			];
-
-			Action ConvertAction = Graph.CreateAction(ActionType.Compile);
-			ConvertAction.RootPaths = LinkEnvironment.RootPaths;
-			ConvertAction.CommandDescription = $"Convert Resource [{LinkEnvironment.Architecture.WindowsName}]";
-			ConvertAction.WorkingDirectory = Unreal.EngineSourceDirectory;
-			ConvertAction.CommandPath = EnvVars.ResourceToObjectConverterPath;
-			ConvertAction.CommandArguments = String.Join(' ', ConvertArguments);
-			ConvertAction.StatusDescription = Path.GetFileName(InputFiles.First().AbsolutePath);
-			ConvertAction.PrerequisiteItems.UnionWith(InputFiles);
-			ConvertAction.ProducedItems.Add(ConvertedResourceObjectFile);
-			ConvertAction.ArtifactMode = ArtifactMode.Enabled;
-			ConvertAction.CacheBucket = GetCacheBucket(Target, LinkEnvironment);
-
-			// Convert Resource tool can run remotely if possible
-			ConvertAction.bCanExecuteRemotely = true;
-			ConvertAction.bCanExecuteRemotelyWithSNDBS = false; // no tool template for SN-DBS results in warnings
-
-			// Create a response file for the resource converter arguments for informational purposes as it cannot be used
-			FileItem ResponseFile = FileItem.GetItemByFileReference(GetResponseFileName(LinkEnvironment, ConvertedResourceObjectFile));
-			Graph.CreateIntermediateTextFile(ResponseFile, ConvertArguments);
-			ConvertAction.PrerequisiteItems.Add(ResponseFile);
-
-			CPPOutput Result = new();
-			Result.ObjectFiles.Add(ConvertedResourceObjectFile);
-			return Result;
-		}
-
 		public override void GenerateTypeLibraryHeader(CppCompileEnvironment CompileEnvironment, ModuleRules.TypeLibrary TypeLibrary, FileReference OutputFile, FileReference? OutputHeader, IActionGraphBuilder Graph)
 		{
 			CppRootPaths RootPaths = CompileEnvironment.RootPaths;
@@ -3243,7 +3203,6 @@ namespace UnrealBuildTool
 
 			// Add the input files to a response file, and pass the response file on the command-line.
 			List<string> InputFileNames = new List<string>();
-
 			foreach (FileItem InputFile in LinkEnvironment.InputFiles)
 			{
 				InputFileNames.Add($"\"{NormalizeCommandLinePath(InputFile, RootPaths)}\"");
diff --git a/Engine/Source/Programs/UnrealBuildTool/System/LinkEnvironment.cs b/Engine/Source/Programs/UnrealBuildTool/System/LinkEnvironment.cs
--- a/Engine/Source/Programs/UnrealBuildTool/System/LinkEnvironment.cs
+++ b/Engine/Source/Programs/UnrealBuildTool/System/LinkEnvironment.cs
@@ -330,9 +330,9 @@ namespace UnrealBuildTool
 		public List<FileItem> DefaultResourceFiles = [];
 
 		/// <summary>
-		/// The default resource file to link in to every binary if a custom one is not provided (converted to obj)
+		/// Resource files which should be compiled into every binary
 		/// </summary>
-		public List<FileItem> DefaultConvertedResourceFiles = [];
+		public List<FileItem> CommonResourceFiles = [];
 
 		/// <summary>
 		/// List of functions that should be exported from this module
@@ -438,7 +438,7 @@ namespace UnrealBuildTool
 			InstallName = Other.InstallName;
 			InputFiles.AddRange(Other.InputFiles);
 			DefaultResourceFiles.AddRange(Other.DefaultResourceFiles);
-			DefaultConvertedResourceFiles.AddRange(Other.DefaultConvertedResourceFiles);
+			CommonResourceFiles.AddRange(Other.CommonResourceFiles);
 			IncludeFunctions.AddRange(Other.IncludeFunctions);
 			DebuggerVisualizerFiles.AddRange(Other.DebuggerVisualizerFiles);
 			ModuleDefinitionFile = Other.ModuleDefinitionFile;
diff --git a/Engine/Source/Programs/UnrealBuildTool/ToolChain/UEToolChain.cs b/Engine/Source/Programs/UnrealBuildTool/ToolChain/UEToolChain.cs
--- a/Engine/Source/Programs/UnrealBuildTool/ToolChain/UEToolChain.cs
+++ b/Engine/Source/Programs/UnrealBuildTool/ToolChain/UEToolChain.cs
@@ -134,14 +134,8 @@ namespace UnrealBuildTool
 		}
 
 		public virtual CPPOutput CompileRCFiles(CppCompileEnvironment CompileEnvironment, IEnumerable<FileItem> InputFiles, DirectoryReference OutputDir, IActionGraphBuilder Graph)
-		{
-			return new();
-		}
-
-		public virtual CPPOutput ConvertRCFiles(LinkEnvironment LinkEnvironment, IEnumerable<FileItem> InputFiles, DirectoryReference OutputDir, IActionGraphBuilder Graph)
 		{
 			CPPOutput Result = new();
-			Result.ObjectFiles.AddRange(InputFiles);
 			return Result;
 		}
 
