UBT - Use Metal Shader syntax hightlight with ush, usf in generated Xcode projects

#tests local generation of Xcode project

diff --git a/Engine/Source/Programs/UnrealBuildTool/ProjectFiles/Xcode/XcodeProject.File.cs b/Engine/Source/Programs/UnrealBuildTool/ProjectFiles/Xcode/XcodeProject.File.cs
--- a/Engine/Source/Programs/UnrealBuildTool/ProjectFiles/Xcode/XcodeProject.File.cs
+++ b/Engine/Source/Programs/UnrealBuildTool/ProjectFiles/Xcode/XcodeProject.File.cs
@@ -165,7 +165,9 @@ namespace UnrealBuildTool.XcodeProjectXcconfig
 			{ ".storyboardc", "wrapper.storyboardc" },
 			{ ".storyboard", "file.storyboard" },
 			{ ".xcframework", "wrapper.xcframework" },
-			{ ".bundle", "\"wrapper.plug-in\"" }
+			{ ".bundle", "\"wrapper.plug-in\"" },
+			{ ".ush", "sourcecode.metal" },
+			{ ".usf", "sourcecode.metal" }
 		};
 
 		const string FileTypeDefault = "file.txt";
