Add target name to reference chains in UBT error reporting to make errors easier to understand.

diff --git a/Engine/Source/Programs/UnrealBuildTool/Configuration/UEBuildTarget.cs b/Engine/Source/Programs/UnrealBuildTool/Configuration/UEBuildTarget.cs
--- a/Engine/Source/Programs/UnrealBuildTool/Configuration/UEBuildTarget.cs
+++ b/Engine/Source/Programs/UnrealBuildTool/Configuration/UEBuildTarget.cs
@@ -4393,7 +4393,7 @@ namespace UnrealBuildTool
 
 			// Create all the modules referenced by the existing binaries
 			UEBuildModule CreateModuleDelegate(string Name, string RefChain) => FindOrCreateModuleByName(Name, RefChain, Logger);
-			StringBuilder ReferenceChain = new("Target");
+			StringBuilder ReferenceChain = new(TargetName);
 			foreach (UEBuildBinary Binary in Binaries)
 			{
 				foreach (UEBuildModule Module in Binary.Modules)
@@ -5043,7 +5043,7 @@ namespace UnrealBuildTool
 				ModuleRules? ModuleRules;
 				try
 				{
-					string PrecompileReferenceChain = "allmodules option";
+					string PrecompileReferenceChain = $"{TargetName} allmodules option";
 					if (ModulePluginSource.TryGetValue(FilteredModuleName, out PluginInfo? value))
 					{
 						PrecompileReferenceChain = $"{PrecompileReferenceChain} -> {value.File.GetFileName()}";
@@ -5077,7 +5077,7 @@ namespace UnrealBuildTool
 
 			foreach (string ModuleName in ValidModuleNames)
 			{
-				StringBuilder PrecompileReferenceChain = new("allmodules option");
+				StringBuilder PrecompileReferenceChain = new($"{TargetName} allmodules option");
 				if (ModulePluginSource.TryGetValue(ModuleName, out PluginInfo? value))
 				{
 					PrecompileReferenceChain.Append(" -> ").Append(value.File.GetFileName());
@@ -5439,7 +5439,7 @@ namespace UnrealBuildTool
 			if (ForeignPlugin != null)
 			{
 				PluginReferenceDescriptor PluginReference = new PluginReferenceDescriptor(ForeignPlugin.GetFileNameWithoutExtension(), null, true);
-				AddPlugin(PluginReference, "command line", ExcludeFolders, NameToInstance, NameToInfos, Logger);
+				AddPlugin(PluginReference, $"{TargetName} command line", ExcludeFolders, NameToInstance, NameToInfos, Logger);
 			}
 
 			// Configure plugins explicitly enabled via target settings
@@ -5448,7 +5448,7 @@ namespace UnrealBuildTool
 				if (ReferencedNames.Add(PluginName))
 				{
 					PluginReferenceDescriptor PluginReference = new PluginReferenceDescriptor(PluginName, null, true);
-					AddPlugin(PluginReference, "target settings", ExcludeFolders, NameToInstance, NameToInfos, Logger);
+					AddPlugin(PluginReference, $"{TargetName} EnablePlugins", ExcludeFolders, NameToInstance, NameToInfos, Logger);
 				}
 			}
 
@@ -5458,7 +5458,7 @@ namespace UnrealBuildTool
 				if (ReferencedNames.Add(PluginName))
 				{
 					PluginReferenceDescriptor PluginReference = new PluginReferenceDescriptor(PluginName, null, false);
-					AddPlugin(PluginReference, "target settings", ExcludeFolders, NameToInstance, NameToInfos, Logger);
+					AddPlugin(PluginReference, $"{TargetName} DisablePlugins", ExcludeFolders, NameToInstance, NameToInfos, Logger);
 				}
 			}
 
@@ -5469,7 +5469,7 @@ namespace UnrealBuildTool
 				{
 					PluginReferenceDescriptor PluginReference = new PluginReferenceDescriptor(PluginName, null, true);
 					PluginReference.bOptional = true;
-					AddPlugin(PluginReference, "target settings", ExcludeFolders, NameToInstance, NameToInfos, Logger);
+					AddPlugin(PluginReference, $"{TargetName} OptionalPlugins", ExcludeFolders, NameToInstance, NameToInfos, Logger);
 				}
 			}
 
@@ -5527,7 +5527,7 @@ namespace UnrealBuildTool
 							PluginReference.PlatformAllowList = PluginReference.SupportedTargetPlatforms; //synthesize allow list if it must be explicit
 						}
 
-						AddPlugin(PluginReference, "default plugins", ExcludeFolders, NameToInstance, NameToInfos, Logger);
+						AddPlugin(PluginReference, $"{TargetName} default plugins", ExcludeFolders, NameToInstance, NameToInfos, Logger);
 					}
 				}
 			}
@@ -5545,7 +5545,7 @@ namespace UnrealBuildTool
 						if (ReferencedNames.Add(PluginName))
 						{
 							PluginReferenceDescriptor PluginReference = new PluginReferenceDescriptor(PluginName, null, true);
-							AddPlugin(PluginReference, "DefaultEngine.ini", ExcludeFolders, NameToInstance, NameToInfos, Logger);
+							AddPlugin(PluginReference, $"{TargetName} DefaultEngine.ini", ExcludeFolders, NameToInstance, NameToInfos, Logger);
 						}
 					}
 				}
@@ -5560,7 +5560,7 @@ namespace UnrealBuildTool
 				if (ReferencedNames.Add(PluginName))
 				{
 					PluginReferenceDescriptor PluginReference = new PluginReferenceDescriptor(PluginName, null, true);
-					AddPlugin(PluginReference, "target settings", ExcludeFolders, NameToInstance, NameToInfos, Logger);
+					AddPlugin(PluginReference, $"{TargetName} BuildPlugins", ExcludeFolders, NameToInstance, NameToInfos, Logger);
 				}
 			}
 
