Disable PGOOptimize when profiledata is not available

diff --git a/Engine/Source/Programs/UnrealBuildTool/ToolChain/ClangToolChain.Link.cs b/Engine/Source/Programs/UnrealBuildTool/ToolChain/ClangToolChain.Link.cs
--- a/Engine/Source/Programs/UnrealBuildTool/ToolChain/ClangToolChain.Link.cs
+++ b/Engine/Source/Programs/UnrealBuildTool/ToolChain/ClangToolChain.Link.cs
@@ -3,6 +3,7 @@
 using EpicGames.Core;
 using System;
 using System.Collections.Generic;
+using UnrealBuildBase;
 
 namespace UnrealBuildTool
 {
@@ -14,6 +15,27 @@ namespace UnrealBuildTool
 			bool bThinLTOEnabled = Options.HasFlag(ClangToolChainOptions.EnableThinLTO);
 
 			// Profile Guided Optimization (PGO)
+			// Disable PGOOptimize when no profdata available
+			if (linkEnvironment.bPGOOptimize)
+			{
+				if (linkEnvironment.PGODirectory != null && linkEnvironment.PGOFilenamePrefix != null)
+				{
+					DirectoryReference PGODir = new(linkEnvironment.PGODirectory);
+					FileItem PGOFile = FileItem.GetItemByFileReference(FileReference.Combine(PGODir, linkEnvironment.PGOFilenamePrefix));
+
+					if (!PGOFile.Exists)
+					{
+						Log.TraceInformationOnce("Missing ProfileData, Profile Guided Optimization Disabled (PGO).");
+						linkEnvironment.bPGOOptimize = false;
+					}
+				}
+				else
+				{
+					Log.TraceInformationOnce("Missing ProfileData, Profile Guided Optimization Disabled (PGO).");
+					linkEnvironment.bPGOOptimize = false;
+				}
+			}
+
 			if (linkEnvironment.bPGOOptimize && linkEnvironment.PGODirectory != null && linkEnvironment.PGOFilenamePrefix != null)
 			{
 				DirectoryReference PGODir = new(linkEnvironment.PGODirectory);
diff --git a/Engine/Source/Programs/UnrealBuildTool/ToolChain/ClangToolChain.cs b/Engine/Source/Programs/UnrealBuildTool/ToolChain/ClangToolChain.cs
--- a/Engine/Source/Programs/UnrealBuildTool/ToolChain/ClangToolChain.cs
+++ b/Engine/Source/Programs/UnrealBuildTool/ToolChain/ClangToolChain.cs
@@ -753,6 +753,31 @@ namespace UnrealBuildTool
 			}
 
 			// Profile Guided Optimization (PGO) and Link Time Optimization (LTO)
+			// Disable PGOOptimize when no profdata available
+			if (CompileEnvironment.bPGOOptimize)
+			{
+				CompileEnvironment.bPGOOptimize = false;
+
+				if (CompileEnvironment.PGODirectory != null && CompileEnvironment.PGOFilenamePrefix != null)
+				{
+					DirectoryReference? PGODir = DirectoryReference.FromString(CompileEnvironment.PGODirectory!);
+					if (PGODir != null)
+					{
+						FileItem PGOFile = FileItem.GetItemByFileReference(FileReference.Combine(PGODir, CompileEnvironment.PGOFilenamePrefix));
+
+						if (PGOFile.Exists)
+						{
+							CompileEnvironment.bPGOOptimize = true;
+						}
+					}
+				}
+				
+				if (!CompileEnvironment.bPGOOptimize)
+				{
+					Log.TraceInformationOnce("Missing ProfileData, Profile Guided Optimization Disabled (PGO).");
+				}
+			}
+
 			if (CompileEnvironment.bPGOOptimize)
 			{
 				Log.TraceInformationOnce("Enabling Profile Guided Optimization (PGO). Linking will take a while.");
