Use legacy rsync 2.6.9 for Mac Remote builds. MacOS Sequoia transitioned from rsync to openrsync causing some incompatibilities of the --delete* arguments resulting in erroneous file deletions on the invoking Windows PC.

diff --git a/Engine/Source/Programs/UnrealBuildTool/ToolChain/RemoteMac.cs b/Engine/Source/Programs/UnrealBuildTool/ToolChain/RemoteMac.cs
--- a/Engine/Source/Programs/UnrealBuildTool/ToolChain/RemoteMac.cs
+++ b/Engine/Source/Programs/UnrealBuildTool/ToolChain/RemoteMac.cs
@@ -13,6 +13,22 @@ using UnrealBuildBase;
 
 namespace UnrealBuildTool
 {
+	/// <summary>
+	/// Stores information about how a local directory maps to a remote directory
+	/// </summary>
+	[DebuggerDisplay("{LocalDirectory}")]
+	class RemoteMapping
+	{
+		public DirectoryReference LocalDirectory;
+		public string RemoteDirectory;
+
+		public RemoteMapping(DirectoryReference LocalDirectory, string RemoteDirectory)
+		{
+			this.LocalDirectory = LocalDirectory;
+			this.RemoteDirectory = RemoteDirectory;
+		}
+	}
+
 	/// <summary>
 	/// Handles uploading and building on a remote Mac
 	/// </summary>
@@ -83,27 +99,6 @@ namespace UnrealBuildTool
 		/// </summary>
 		private string RemoteBaseDir;
 
-		/// <summary>
-		/// Stores information about how a local directory maps to a remote directory
-		/// </summary>
-		[DebuggerDisplay("{LocalDirectory}")]
-		private class RemoteMapping
-		{
-			public DirectoryReference LocalDirectory;
-			public string RemoteDirectory;
-
-			public RemoteMapping(DirectoryReference LocalDirectory, string RemoteDirectory)
-			{
-				this.LocalDirectory = LocalDirectory;
-				this.RemoteDirectory = RemoteDirectory;
-			}
-		}
-
-		private RemoteMapping CreateRemoteMapping(DirectoryReference LocalDirectory)
-		{
-			return new RemoteMapping(LocalDirectory, GetRemotePath(LocalDirectory));
-		}
-
 		/// <summary>
 		/// Mappings from local directories to remote directories
 		/// </summary>
@@ -119,16 +114,25 @@ namespace UnrealBuildTool
 		/// </summary>
 		private List<string> BasicRsyncArguments;
 
-		/// <summary>
-		/// Arguments that are used by directory Rsync call without deletions
-		/// </summary>
-		private List<string> NoDeleteRsyncArguments;
-
 		/// <summary>
 		/// Arguments that are used by directory Rsync call
 		/// </summary>
 		private List<string> CommonRsyncArguments;
 
+		private List<string> RsyncArgumentsForDownload(List<string> BaseArguments)
+		{
+			List<string> Arguments = new List<string>(BaseArguments);
+			Arguments.Add(String.Format("--rsync-path=\"{0}/rsync\"", RemoteBaseDir));
+			return Arguments;
+		}
+
+		private List<string> RsyncArgumentsForUpload(List<string> BaseArguments, string RemoteDirectory)
+		{
+			List<string> Arguments = new List<string>(BaseArguments);
+			Arguments.Add(String.Format("--rsync-path=\"mkdir -p {0} && {1}/rsync\"", RemoteDirectory, RemoteBaseDir));
+			return Arguments;
+		}
+
 		private string? IniBundleIdentifier = "";
 
 		/// <summary>
@@ -141,8 +145,9 @@ namespace UnrealBuildTool
 		/// <param name="bPrepareForSecondaryMac">Added arguments when preparing a secondary Mac for debug. False by default.</param>
 		public RemoteMac(FileReference? ProjectFile, ILogger Logger, string? RemoteIniPath = null, bool bIsPrimary = true, bool bPrepareForSecondaryMac = false)
 		{
-			RsyncExe = FileReference.Combine(Unreal.EngineDirectory, "Extras", "ThirdPartyNotUE", "cwrsync", "bin", "rsync.exe");
-			SshExe = FileReference.Combine(Unreal.EngineDirectory, "Extras", "ThirdPartyNotUE", "cwrsync", "bin", "ssh.exe");
+			DirectoryReference Cwrsync = DirectoryReference.Combine(Unreal.EngineDirectory, "Extras", "ThirdPartyNotUE", "cwrsync");
+			RsyncExe = FileReference.Combine(Cwrsync, "bin", "rsync.exe");
+			SshExe = FileReference.Combine(Cwrsync, "bin", "ssh.exe");
 			this.ProjectFile = ProjectFile;
 			if (ProjectFile != null)
 			{
@@ -197,12 +202,12 @@ namespace UnrealBuildTool
 			int PortIdx = ServerName.LastIndexOf(':');
 			if (PortIdx != -1)
 			{
-				string Port = ServerName.Substring(PortIdx + 1);
+				string Port = ServerName[(PortIdx + 1)..];
 				if (!Int32.TryParse(Port, out ServerPort))
 				{
 					throw new BuildException("Unable to parse port number from '{0}'", ServerName);
 				}
-				ServerName = ServerName.Substring(0, PortIdx);
+				ServerName = ServerName[..PortIdx];
 			}
 
 			// If a user name is not set, use the current user
@@ -276,21 +281,18 @@ namespace UnrealBuildTool
 			BasicRsyncArguments.Add(String.Format("--rsh=\"{0} -p {1}\"", RsyncAuthentication, ServerPort));
 			BasicRsyncArguments.Add("--chmod=ugo=rwx");
 
-			// Build a list of arguments for Rsync filters without destination deletions
-			NoDeleteRsyncArguments = new List<string>(BasicRsyncArguments);
-			NoDeleteRsyncArguments.Add("--copy-links");
-			NoDeleteRsyncArguments.Add("--recursive");
-			NoDeleteRsyncArguments.Add("--times"); // Preserve modification times
-			NoDeleteRsyncArguments.Add("--omit-dir-times"); // Ignore modification times for directories
-			NoDeleteRsyncArguments.Add("--prune-empty-dirs"); // Remove empty directories from the file list
-
 			// Build a list of arguments for Rsync filters
-			CommonRsyncArguments = new List<string>(NoDeleteRsyncArguments);
+			CommonRsyncArguments = new List<string>(BasicRsyncArguments);
+			CommonRsyncArguments.Add("--copy-links");
+			CommonRsyncArguments.Add("--recursive");
 			if (!bPrepareForSecondaryMac)
 			{
 				CommonRsyncArguments.Add("--delete"); // Delete anything not in the source directory
 				CommonRsyncArguments.Add("--delete-excluded"); // Delete anything not in the source directory
 			}
+			CommonRsyncArguments.Add("--times"); // Preserve modification times
+			CommonRsyncArguments.Add("--omit-dir-times"); // Ignore modification times for directories
+			CommonRsyncArguments.Add("--prune-empty-dirs"); // Remove empty directories from the file list
 
 			// Get the remote base directory
 			string RemoteServerOverrideBuildPath;
@@ -312,10 +314,10 @@ namespace UnrealBuildTool
 
 			// Build the list of directory mappings between the local and remote machines
 			Mappings = new List<RemoteMapping>();
-			Mappings.Add(CreateRemoteMapping(Unreal.EngineDirectory));
+			Mappings.Add(new RemoteMapping(Unreal.EngineDirectory, GetRemotePath(Unreal.EngineDirectory)));
 			if (ProjectFile != null && !ProjectFile.IsUnderDirectory(Unreal.EngineDirectory))
 			{
-				Mappings.Add(CreateRemoteMapping(ProjectFile.Directory));
+				Mappings.Add(new RemoteMapping(ProjectFile.Directory, GetRemotePath(ProjectFile.Directory)));
 			}
 			if (AdditionalPaths != null && ProjectFile != null)
 			{
@@ -324,10 +326,26 @@ namespace UnrealBuildTool
 					if (!AdditionalPath.IsUnderDirectory(Unreal.EngineDirectory) &&
 						!AdditionalPath.IsUnderDirectory(ProjectFile.Directory))
 					{
-						Mappings.Add(CreateRemoteMapping(AdditionalPath));
+						Mappings.Add(new RemoteMapping(AdditionalPath, GetRemotePath(AdditionalPath)));
 					}
 				}
 			}
+
+			string RemoteFile = GetRemotePath("rsync");
+
+			List<string> Arguments = new List<string>(CommonRsyncArguments);
+			// use MacOS default system rsync for the upload of the 2.6.9 rsync binary
+			Arguments.Add(String.Format("--rsync-path=\"mkdir -p {0} && rsync\"", RemoteBaseDir));
+			Arguments.Add(String.Format("\"{0}\"", GetLocalCygwinPath(FileReference.Combine(Cwrsync, "mac", "rsync"))));
+			Arguments.Add(String.Format("\"{0}@{1}\":'{2}'", UserName, ServerName, RemoteFile));
+			Arguments.Add("-q");
+
+			int Result = Rsync(String.Join(" ", Arguments), Logger);
+			if (Result != 0)
+			{
+				throw new BuildException("Error while running Rsync (exit code {0})", Result);
+			}
+			Execute(RemoteBaseDir, "xattr -d com.apple.quarantine rsync", Logger);
 		}
 
 		/// <summary>
@@ -478,7 +496,7 @@ namespace UnrealBuildTool
 			}
 			if (!bLogIsMapped)
 			{
-				Mappings.Add(CreateRemoteMapping(RemoteLogFile.Directory));
+				Mappings.Add(new RemoteMapping(RemoteLogFile.Directory, GetRemotePath(RemoteLogFile.Directory)));
 			}
 
 			// Path to the local manifest file. This has to be translated from the remote format after the build is complete.
@@ -675,8 +693,8 @@ namespace UnrealBuildTool
 					continue;
 				}
 
-				string Key = LocalArgument.Substring(0, EqualsIdx);
-				string Value = LocalArgument.Substring(EqualsIdx + 1);
+				string Key = LocalArgument[..EqualsIdx];
+				string Value = LocalArgument[(EqualsIdx + 1)..];
 
 				if (Key.Equals("-Log", StringComparison.InvariantCultureIgnoreCase))
 				{
@@ -810,7 +828,7 @@ namespace UnrealBuildTool
 			{
 				if (RemotePath.StartsWith(Mapping.RemoteDirectory, StringComparison.InvariantCultureIgnoreCase) && RemotePath.Length > Mapping.RemoteDirectory.Length && RemotePath[Mapping.RemoteDirectory.Length] == '/')
 				{
-					return FileReference.Combine(Mapping.LocalDirectory, RemotePath.Substring(Mapping.RemoteDirectory.Length + 1));
+					return FileReference.Combine(Mapping.LocalDirectory, RemotePath[(Mapping.RemoteDirectory.Length + 1)..]);
 				}
 			}
 			throw new BuildException("Unable to map remote path '{0}' to local path", RemotePath);
@@ -847,7 +865,7 @@ namespace UnrealBuildTool
 			{
 				throw new BuildException("Invalid local path for converting to cygwin format ({0}).", InPath);
 			}
-			return String.Format("/cygdrive/{0}{1}", InPath.FullName.Substring(0, 1), InPath.FullName.Substring(2).Replace('\\', '/'));
+			return String.Format("/cygdrive/{0}{1}", InPath.FullName[..1], InPath.FullName[2..].Replace('\\', '/'));
 		}
 
 		/// <summary>
@@ -870,8 +888,7 @@ namespace UnrealBuildTool
 			string RemoteFile = GetRemotePath(LocalFile);
 			string RemoteDirectory = GetRemotePath(LocalFile.Directory);
 
-			List<string> Arguments = new List<string>(CommonRsyncArguments);
-			Arguments.Add(String.Format("--rsync-path=\"mkdir -p {0} && rsync\"", RemoteDirectory));
+			List<string> Arguments = RsyncArgumentsForUpload(CommonRsyncArguments, RemoteDirectory);
 			Arguments.Add(String.Format("\"{0}\"", GetLocalCygwinPath(LocalFile)));
 			Arguments.Add(String.Format("\"{0}@{1}\":'{2}'", UserName, ServerName, RemoteFile));
 			Arguments.Add("-q");
@@ -892,8 +909,7 @@ namespace UnrealBuildTool
 		/// <param name="Logger">Logger for output</param>
 		void UploadFiles(DirectoryReference LocalDirectory, string RemoteDirectory, FileReference LocalFileList, ILogger Logger)
 		{
-			List<string> Arguments = new List<string>(BasicRsyncArguments);
-			Arguments.Add(String.Format("--rsync-path=\"mkdir -p {0} && rsync\"", RemoteDirectory));
+			List<string> Arguments = RsyncArgumentsForUpload(BasicRsyncArguments, RemoteDirectory);
 			Arguments.Add(String.Format("--files-from=\"{0}\"", GetLocalCygwinPath(LocalFileList)));
 			Arguments.Add(String.Format("\"{0}/\"", GetLocalCygwinPath(LocalDirectory)));
 			Arguments.Add(String.Format("\"{0}@{1}\":'{2}/'", UserName, ServerName, RemoteDirectory));
@@ -915,8 +931,7 @@ namespace UnrealBuildTool
 		{
 			string RemoteDirectory = GetRemotePath(LocalDirectory);
 
-			List<string> Arguments = new List<string>(CommonRsyncArguments);
-			Arguments.Add(String.Format("--rsync-path=\"mkdir -p {0} && rsync\"", RemoteDirectory));
+			List<string> Arguments = RsyncArgumentsForUpload(CommonRsyncArguments, RemoteDirectory);
 			Arguments.Add(String.Format("\"{0}/\"", GetLocalCygwinPath(LocalDirectory)));
 			Arguments.Add(String.Format("\"{0}@{1}\":'{2}/'", UserName, ServerName, RemoteDirectory));
 			Arguments.Add("-q");
@@ -937,8 +952,7 @@ namespace UnrealBuildTool
 		/// <param name="Logger">Logger for output</param>
 		void UploadDirectory(DirectoryReference LocalDirectory, string RemoteDirectory, List<FileReference> FilterLocations, ILogger Logger)
 		{
-			List<string> Arguments = new List<string>(CommonRsyncArguments);
-			Arguments.Add(String.Format("--rsync-path=\"mkdir -p {0} && rsync\"", RemoteDirectory));
+			List<string> Arguments = RsyncArgumentsForUpload(CommonRsyncArguments, RemoteDirectory);
 			foreach (FileReference FilterLocation in FilterLocations)
 			{
 				Arguments.Add(String.Format("--filter=\"merge {0}\"", GetLocalCygwinPath(FilterLocation)));
@@ -1106,7 +1120,7 @@ namespace UnrealBuildTool
 				throw new BuildException("File for download '{0}' is not under any mapped directory.", LocalFile);
 			}
 
-			List<string> Arguments = new List<string>(CommonRsyncArguments);
+			List<string> Arguments = RsyncArgumentsForDownload(CommonRsyncArguments);
 			Arguments.Add(String.Format("\"{0}@{1}\":'{2}/{3}'", UserName, ServerName, Mapping.RemoteDirectory, LocalFile.MakeRelativeTo(Mapping.LocalDirectory).Replace('\\', '/')));
 			Arguments.Add(String.Format("\"{0}/\"", GetLocalCygwinPath(LocalFile.Directory)));
 			Arguments.Add("-q");
@@ -1147,7 +1161,7 @@ namespace UnrealBuildTool
 					DirectoryReference.CreateDirectory(DownloadListLocation.Directory);
 					FileReference.WriteAllLines(DownloadListLocation, FileGroups[Idx].Select(x => x.MakeRelativeTo(Mappings[Idx].LocalDirectory).Replace('\\', '/')));
 
-					List<string> Arguments = new List<string>(NoDeleteRsyncArguments);
+					List<string> Arguments = RsyncArgumentsForDownload(CommonRsyncArguments);
 					Arguments.Add(String.Format("--files-from=\"{0}\"", GetLocalCygwinPath(DownloadListLocation)));
 					Arguments.Add(String.Format("\"{0}@{1}\":'{2}/'", UserName, ServerName, Mappings[Idx].RemoteDirectory));
 					Arguments.Add(String.Format("\"{0}/\"", GetLocalCygwinPath(Mappings[Idx].LocalDirectory)));
@@ -1184,7 +1198,7 @@ namespace UnrealBuildTool
 
 			string RemoteDirectory = GetRemotePath(LocalDirectory);
 
-			List<string> Arguments = new List<string>(CommonRsyncArguments);
+			List<string> Arguments = RsyncArgumentsForDownload(CommonRsyncArguments);
 			Arguments.Add(String.Format("\"{0}@{1}\":'{2}/'", UserName, ServerName, RemoteDirectory));
 			Arguments.Add(String.Format("\"{0}/\"", GetLocalCygwinPath(LocalDirectory)));
 
@@ -1376,7 +1390,7 @@ namespace UnrealBuildTool
 				{
 					if (Path.StartsWith(Mapping.RemoteDirectory))
 					{
-						Path = Mapping.LocalDirectory + Path.Substring(Mapping.RemoteDirectory.Length).Replace('/', '\\');
+						Path = Mapping.LocalDirectory + Path[Mapping.RemoteDirectory.Length..].Replace('/', '\\');
 						break;
 					}
 				}
