[UBT]
* Fixed bug where iwyu mode could add ipp files twice

diff --git a/Engine/Source/Programs/UnrealBuildTool/Modes/IWYUMode.cs b/Engine/Source/Programs/UnrealBuildTool/Modes/IWYUMode.cs
--- a/Engine/Source/Programs/UnrealBuildTool/Modes/IWYUMode.cs
+++ b/Engine/Source/Programs/UnrealBuildTool/Modes/IWYUMode.cs
@@ -1534,6 +1534,7 @@ namespace UnrealBuildTool
 				{
 					if (SeenInclude.Printable.EndsWith(".ipp")) // Abuse this to make sure we never remove .ipp includes
 					{
+						IncludesToAdd.Remove($"\"{SeenInclude.Printable}\"");
 						continue;
 					}
 
