Make the warning for non-inlined .gen.cpp files display correctly as a warning or error in UBT.

diff --git a/Engine/Source/Programs/UnrealBuildTool/Configuration/UEBuildModuleCPP.cs b/Engine/Source/Programs/UnrealBuildTool/Configuration/UEBuildModuleCPP.cs
--- a/Engine/Source/Programs/UnrealBuildTool/Configuration/UEBuildModuleCPP.cs
+++ b/Engine/Source/Programs/UnrealBuildTool/Configuration/UEBuildModuleCPP.cs
@@ -772,7 +772,8 @@ namespace UnrealBuildTool
 					}
 					ModuleCompileEnvironment.FileInlineSourceMap = CppCompileEnvironment.FileInlineSourceMap;
 
-					if (Rules.CppCompileWarningSettings.NonInlinedGenCppWarningLevel >= WarningLevel.Warning)
+					WarningLevel NonInlinedGenCppWarningLevel = Rules.CppCompileWarningSettings.NonInlinedGenCppWarningLevel;
+					if (NonInlinedGenCppWarningLevel >= WarningLevel.Warning)
 					{
 						Dictionary<string, FileItem> CPPFilesLookup = new Dictionary<string, FileItem>();
 						foreach (FileItem CPPFile in CPPFiles)
@@ -788,10 +789,16 @@ namespace UnrealBuildTool
 							}
 
 							string NameWithoutPrefix = Name.Substring(4);
-							string NameWithoutPrefixWithSuffix = NameWithoutPrefix + ".cpp";
+							string NameWithoutPrefixWithSuffix = $"{NameWithoutPrefix}.cpp";
 							if (CPPFilesLookup.TryGetValue(NameWithoutPrefixWithSuffix, out FileItem? Item))
 							{
-								Logger.LogWarning("{0}(1): .gen.cpp not inlined. Add '#include UE_INLINE_GENERATED_CPP_BY_NAME({1})' after header includes.", Item.AbsolutePath, NameWithoutPrefix);
+								LogLevel LogLevel = NonInlinedGenCppWarningLevel switch
+								{
+									WarningLevel.Warning => LogLevel.Warning,
+									WarningLevel.Error => LogLevel.Error,
+									_ => throw new BuildException("Unhandled warning level"),
+								};
+								Logger.Log(LogLevel, "{CPPFileItem}(1): {LogLevel}: .gen.cpp not inlined. Add '#include UE_INLINE_GENERATED_CPP_BY_NAME({CPPFileName})' after header includes.", Item.AbsolutePath, LogLevel, NameWithoutPrefix);
 							}
 						}
 					}
