Safely remove quarantine extended attribute so as to not add benign error to the build log if the attribute was already removed.

diff --git a/Engine/Source/Programs/UnrealBuildTool/ToolChain/RemoteMac.cs b/Engine/Source/Programs/UnrealBuildTool/ToolChain/RemoteMac.cs
--- a/Engine/Source/Programs/UnrealBuildTool/ToolChain/RemoteMac.cs
+++ b/Engine/Source/Programs/UnrealBuildTool/ToolChain/RemoteMac.cs
@@ -345,7 +345,8 @@ namespace UnrealBuildTool
 			{
 				throw new BuildException("Error while running Rsync (exit code {0})", Result);
 			}
-			Execute(RemoteBaseDir, "xattr -d com.apple.quarantine rsync", Logger);
+			// Remove quaratine extended attribute if remote machine set it.
+			Execute(RemoteBaseDir, "if xattr -p com.apple.quarantine rsync > /dev/null 2>&1; then xattr -d com.apple.quarantine rsync; fi", Logger);
 		}
 
 		/// <summary>
