[UBT]
* Added null check for UBAAgentCoordinatorHorde _updateStatus since it can be null if horde is disabled

diff --git a/Engine/Source/Programs/UnrealBuildTool/Executors/UnrealBuildAccelerator/UBAAgentCoordinatorHorde.cs b/Engine/Source/Programs/UnrealBuildTool/Executors/UnrealBuildAccelerator/UBAAgentCoordinatorHorde.cs
--- a/Engine/Source/Programs/UnrealBuildTool/Executors/UnrealBuildAccelerator/UBAAgentCoordinatorHorde.cs
+++ b/Engine/Source/Programs/UnrealBuildTool/Executors/UnrealBuildAccelerator/UBAAgentCoordinatorHorde.cs
@@ -883,9 +883,9 @@ namespace UnrealBuildTool
 				(hordeSession, reason) = _hordeSessionTask.Result;
 			}
 
-			if (hordeSession == null)
+			if (hordeSession == null && _updateStatus != null)
 			{
-				_updateStatus!(0, 6, reason, EpicGames.UBA.LogEntryType.Info, null);
+				_updateStatus(0, 6, reason, EpicGames.UBA.LogEntryType.Info, null);
 				return;
 			}
 
