[UBT]
* Fixed change in UBAAgentCoordinatorHorde where null guard was placed in the wrong place

diff --git a/Engine/Source/Programs/UnrealBuildTool/Executors/UnrealBuildAccelerator/UBAAgentCoordinatorHorde.cs b/Engine/Source/Programs/UnrealBuildTool/Executors/UnrealBuildAccelerator/UBAAgentCoordinatorHorde.cs
--- a/Engine/Source/Programs/UnrealBuildTool/Executors/UnrealBuildAccelerator/UBAAgentCoordinatorHorde.cs
+++ b/Engine/Source/Programs/UnrealBuildTool/Executors/UnrealBuildAccelerator/UBAAgentCoordinatorHorde.cs
@@ -883,9 +883,12 @@ namespace UnrealBuildTool
 				(hordeSession, reason) = _hordeSessionTask.Result;
 			}
 
-			if (hordeSession == null && _updateStatus != null)
+			if (hordeSession == null)
 			{
-				_updateStatus(0, 6, reason, EpicGames.UBA.LogEntryType.Info, null);
+				if (_updateStatus != null)
+				{
+					_updateStatus(0, 6, reason, EpicGames.UBA.LogEntryType.Info, null);
+				}
 				return;
 			}
 
