[UBT]
* Changed so header units is always disabled if compiler is not visual studio 2026 (which is the first one that properly supports it)

diff --git a/Engine/Source/Programs/UnrealBuildTool/Configuration/TargetRules.cs b/Engine/Source/Programs/UnrealBuildTool/Configuration/TargetRules.cs
--- a/Engine/Source/Programs/UnrealBuildTool/Configuration/TargetRules.cs
+++ b/Engine/Source/Programs/UnrealBuildTool/Configuration/TargetRules.cs
@@ -3273,6 +3273,13 @@ namespace UnrealBuildTool
 				// bAlwaysUseUnityForGeneratedFiles = false; // This should be used but adds too many extra files atm.
 			}
 
+			// Header units only work for vs 2026 (tested with unreleased 14.50.35615)
+			// (For future I suspect we want bUseHeaderUnitsForPch to default to true with 2026)
+			if (WindowsPlatform.Compiler != WindowsCompiler.VisualStudio2026)
+			{
+				bUseHeaderUnitsForPch = false;
+			}
+
 			// Set the default build version
 			if (String.IsNullOrEmpty(BuildVersion))
 			{
