UnrealBuildTool: Move version defines for resource files to VCToolChain.cs
Pass deterministic flag when linking import libs

#rnx

diff --git a/Engine/Source/Programs/UnrealBuildTool/Configuration/UEBuildBinary.cs b/Engine/Source/Programs/UnrealBuildTool/Configuration/UEBuildBinary.cs
--- a/Engine/Source/Programs/UnrealBuildTool/Configuration/UEBuildBinary.cs
+++ b/Engine/Source/Programs/UnrealBuildTool/Configuration/UEBuildBinary.cs
@@ -909,23 +909,6 @@ namespace UnrealBuildTool
 
 			// Set the original file name macro; used in Default.rc2 to set the binary metadata fields.
 			ResourceCompileEnvironment.Definitions.Add($"ORIGINAL_FILE_NAME=\"{OutputFilePaths[0].GetFileName()}\"");
-
-			// Set the other version fields if requested
-			if (Target.WindowsPlatform.bSetResourceVersions)
-			{
-				if (Target.Version.Changelist != 0)
-				{
-					ResourceCompileEnvironment.Definitions.Add($"BUILT_FROM_CHANGELIST={Target.Version.Changelist}");
-				}
-				if (!String.IsNullOrEmpty(Target.Version.BranchName))
-				{
-					ResourceCompileEnvironment.Definitions.Add($"BRANCH_NAME={Target.Version.BranchName}");
-				}
-				if (!String.IsNullOrEmpty(Target.BuildVersion))
-				{
-					ResourceCompileEnvironment.Definitions.Add($"BUILD_VERSION={Target.BuildVersion}");
-				}
-			}
 			
 			// Compile the resource files
 			List<FileItem> ObjectFiles = ToolChain.CompileRCFiles(ResourceCompileEnvironment, InputFiles, IntermediateDirectory, Graph).ObjectFiles;
diff --git a/Engine/Source/Programs/UnrealBuildTool/Platform/Windows/VCToolChain.cs b/Engine/Source/Programs/UnrealBuildTool/Platform/Windows/VCToolChain.cs
--- a/Engine/Source/Programs/UnrealBuildTool/Platform/Windows/VCToolChain.cs
+++ b/Engine/Source/Programs/UnrealBuildTool/Platform/Windows/VCToolChain.cs
@@ -1954,6 +1954,12 @@ namespace UnrealBuildTool
 				// Use link-time code generation.
 				Arguments.Add("/LTCG");
 			}
+
+			// Always build import libraries in deterministic mode, as no codegen is run this does not add any additional time
+			if (!EnvVars.Compiler.IsClang() || !EnvVars.bAllowClangLinker)
+			{
+				Arguments.Add("/experimental:deterministic");
+			}
 		}
 
 		private VCCompileAction CreateBaseCompileAction(CppCompileEnvironment CompileEnvironment)
@@ -2744,7 +2750,6 @@ namespace UnrealBuildTool
 				return new();
 			}
 
-
 			CPPOutput Result = new();
 			CppRootPaths RootPaths = CompileEnvironment.RootPaths;
 
@@ -2809,6 +2814,23 @@ namespace UnrealBuildTool
 					}
 				}
 
+				// Set the other version fields if requested
+				if (Target.WindowsPlatform.bSetResourceVersions)
+				{
+					if (Target.Version.Changelist != 0)
+					{
+						AddDefinition(Arguments, $"BUILT_FROM_CHANGELIST={Target.Version.Changelist}");
+					}
+					if (!String.IsNullOrEmpty(Target.Version.BranchName))
+					{
+						AddDefinition(Arguments, $"BRANCH_NAME={Target.Version.BranchName}");
+					}
+					if (!String.IsNullOrEmpty(Target.BuildVersion))
+					{
+						AddDefinition(Arguments, $"BUILD_VERSION={Target.BuildVersion}");
+					}
+				}
+
 				// Figure the icon to use. We can only use a custom icon when compiling to a project-specific intermediate directory (and not for the shared editor executable, for example).
 				FileItem IconFile = FileItem.GetItemByFileReference(GetApplicationIcon(!CompileEnvironment.bUseSharedBuildEnvironment ? Target.ProjectFile : null));
 				CompileAction.PrerequisiteItems.Add(IconFile);
