[UBT]
* If target type is monolithic and we are building a dll it means that this link step will output final binaries of the target which is a dynamic library. In this case we also want additional files (.lib etc) to end up in the same directory as the binary. This is solved by using LocalShadowDirectory which is already used by some toolchains for this use case

diff --git a/Engine/Source/Programs/UnrealBuildTool/Configuration/UEBuildBinary.cs b/Engine/Source/Programs/UnrealBuildTool/Configuration/UEBuildBinary.cs
--- a/Engine/Source/Programs/UnrealBuildTool/Configuration/UEBuildBinary.cs
+++ b/Engine/Source/Programs/UnrealBuildTool/Configuration/UEBuildBinary.cs
@@ -329,6 +329,13 @@ namespace UnrealBuildTool
 				}
 			}
 
+			// If target type is monolithic and we are building a dll it means that this link step will output final binaries which is a dynamic library
+			// In this case we also want additional files (.lib etc) to end up in the same directory as the binary
+			if (Target.LinkType == TargetLinkType.Monolithic && BinaryLinkEnvironment.bIsBuildingDLL)
+			{
+				BinaryLinkEnvironment.LocalShadowDirectory = BinaryLinkEnvironment.OutputFilePath.Directory;
+			}
+
 			// Link the binary.
 			FileItem[] Executables = ToolChain.LinkAllFiles(BinaryLinkEnvironment, false, Graph);
 			OutputFiles.AddRange(Executables);
