Remove usages of the third-party Newtonsoft.Json dependency from UBT, replacing them with the inbuilt System.Text.Json library. Leave the libraries alone for a subsequent submit to address.

diff --git a/Engine/Source/Programs/UnrealBuildTool/Platform/Windows/PVSToolChain.cs b/Engine/Source/Programs/UnrealBuildTool/Platform/Windows/PVSToolChain.cs
--- a/Engine/Source/Programs/UnrealBuildTool/Platform/Windows/PVSToolChain.cs
+++ b/Engine/Source/Programs/UnrealBuildTool/Platform/Windows/PVSToolChain.cs
@@ -3,7 +3,6 @@
 using EpicGames.Core;
 using Microsoft.Extensions.Logging;
 using Microsoft.Extensions.Logging.Abstractions;
-using Newtonsoft.Json;
 using System;
 using System.Collections.Generic;
 using System.Diagnostics;
@@ -11,6 +10,7 @@ using System.IO;
 using System.Linq;
 using System.Reflection;
 using System.Text;
+using System.Text.Json;
 using System.Text.RegularExpressions;
 using System.Threading.Tasks;
 using System.Xml.Serialization;
@@ -402,6 +402,13 @@ namespace UnrealBuildTool
 
 		CppRootPaths _rootPaths = new();
 
+		private readonly JsonSerializerOptions _jsonSerializerOptions = new()
+		{
+			// PVS-Studio output is camelCased, so we require insensitive property names to serialize to our
+			// standard C# property names.
+			PropertyNameCaseInsensitive = true,
+		};
+
 		/// <summary>
 		/// Execute the command
 		/// </summary>
@@ -482,7 +489,7 @@ namespace UnrealBuildTool
 		{
 			try
 			{
-				PVSErrorInfo errorInfo = JsonConvert.DeserializeObject<PVSErrorInfo>(line) ?? throw new FormatException();
+				PVSErrorInfo errorInfo = JsonSerializer.Deserialize<PVSErrorInfo>(line, _jsonSerializerOptions) ?? throw new FormatException();
 				FileReference? fileReference = errorInfo.Positions.FirstOrDefault()?.UpdateFilePath(_rootPaths);
 				return IsFileIgnored(fileReference) ? null : errorInfo;
 			}
@@ -496,11 +503,9 @@ namespace UnrealBuildTool
 
 	class PVSPosition
 	{
-		[JsonProperty(Required = Required.Always)]
-		public required string File;
+		public required string File { get; set; }
 
-		[JsonProperty(Required = Required.Always)]
-		public required int[] Lines;
+		public required int[] Lines { get; init; }
 
 		public FileReference? UpdateFilePath(CppRootPaths rootPaths)
 		{
@@ -516,20 +521,15 @@ namespace UnrealBuildTool
 
 	class PVSErrorInfo
 	{
-		[JsonProperty(Required = Required.Always)]
-		public required string Code;
+		public required string Code { get; init; }
 
-		[JsonProperty(Required = Required.Always)]
-		public required bool FalseAlarm;
+		public required bool FalseAlarm { get; init; }
 
-		[JsonProperty(Required = Required.Always)]
-		public required int Level;
+		public required int Level { get; init; }
 
-		[JsonProperty(Required = Required.Always)]
-		public required string Message;
+		public required string Message { get; init; }
 
-		[JsonProperty(Required = Required.Always)]
-		public required PVSPosition[] Positions;
+		public required PVSPosition[] Positions { get; init; }
 	}
 
 	class PVSToolChain : ISPCToolChain
