Fix a comment to correctly refer to *de*serialization in UBT's PVS toolchain.

diff --git a/Engine/Source/Programs/UnrealBuildTool/Platform/Windows/PVSToolChain.cs b/Engine/Source/Programs/UnrealBuildTool/Platform/Windows/PVSToolChain.cs
--- a/Engine/Source/Programs/UnrealBuildTool/Platform/Windows/PVSToolChain.cs
+++ b/Engine/Source/Programs/UnrealBuildTool/Platform/Windows/PVSToolChain.cs
@@ -404,7 +404,7 @@ namespace UnrealBuildTool
 
 		private readonly JsonSerializerOptions _jsonSerializerOptions = new()
 		{
-			// PVS-Studio output is camelCased, so we require insensitive property names to serialize to our
+			// PVS-Studio output is camelCased, so we require insensitive property names to deserialize to our
 			// standard C# property names.
 			PropertyNameCaseInsensitive = true,
 		};
