[Android] Build fix for TSan on NDK >= 27

[RN]

diff --git a/Engine/Source/Programs/UnrealBuildTool/Platform/Android/AndroidToolChain.cs b/Engine/Source/Programs/UnrealBuildTool/Platform/Android/AndroidToolChain.cs
--- a/Engine/Source/Programs/UnrealBuildTool/Platform/Android/AndroidToolChain.cs
+++ b/Engine/Source/Programs/UnrealBuildTool/Platform/Android/AndroidToolChain.cs
@@ -347,10 +347,10 @@ namespace UnrealBuildTool
 			return NDKVersionInt;
 		}
 
-		public ulong GetNdkVersionInt()
-		{
-			return NDKVersionInt;
-		}
+		public ulong GetNdkVersionInt() => NDKVersionInt;
+
+		public int GetClangVersionMajor() => ClangVersionMajor;
+
 
 		class MinMax
 		{
@@ -360,6 +360,7 @@ namespace UnrealBuildTool
 		}
 		static ConcurrentDictionary<string, Lazy<MinMax>> CachedMinMax = [];
 
+
 		private bool ReadMinMaxPlatforms(string PlatformsFilename, out int MinPlatform, out int MaxPlatform)
 		{
 			MinMax minMax = CachedMinMax.GetOrAdd(PlatformsFilename, new Lazy<MinMax>(() =>
diff --git a/Engine/Source/Programs/UnrealBuildTool/Platform/Android/UEDeployAndroid.cs b/Engine/Source/Programs/UnrealBuildTool/Platform/Android/UEDeployAndroid.cs
--- a/Engine/Source/Programs/UnrealBuildTool/Platform/Android/UEDeployAndroid.cs
+++ b/Engine/Source/Programs/UnrealBuildTool/Platform/Android/UEDeployAndroid.cs
@@ -3,6 +3,7 @@
 using System;
 using System.Collections.Generic;
 using System.Diagnostics;
+using System.Globalization;
 using System.IO;
 using System.IO.Compression;
 using System.Linq;
@@ -1715,10 +1716,11 @@ namespace UnrealBuildTool
 
 			string PlatformHostName = GetPlatformNDKHostName();
 
-			string VersionFileName = Path.Combine(Environment.ExpandEnvironmentVariables("%NDKROOT%"), "toolchains", "llvm", "prebuilt", PlatformHostName, "AndroidVersion.txt");
-			System.IO.StreamReader VersionFile = new System.IO.StreamReader(VersionFileName);
-			string LibsVersion = VersionFile.ReadLine()!;
-			VersionFile.Close();
+			string LibsVersion = ToolChain.GetClangVersionString();
+			if (ToolChain.GetNdkVersionInt() >= 270000)
+			{
+				LibsVersion = ToolChain.GetClangVersionMajor().ToString(CultureInfo.InvariantCulture);
+			}
 
 			string SanitizerLib = String.IsNullOrEmpty(SanitizerFullLibName) ?
 				String.Empty :
@@ -1732,12 +1734,12 @@ namespace UnrealBuildTool
 					Directory.CreateDirectory(LibDestDir);
 
 					string LibDestFilePath = Path.Combine(LibDestDir, SanitizerFullLibName);
-					Logger.LogInformation("Copying asan lib from {SanitizerLib} to {LibDestFilePath}", SanitizerLib, LibDestFilePath);
+					Logger.LogInformation("Copying {LibName} lib from {SanitizerLib} to {LibDestFilePath}", LibName, SanitizerLib, LibDestFilePath);
 					File.Copy(SanitizerLib, LibDestFilePath, true);
 				}
 				else
 				{
-					throw new BuildException("No asan lib found in {0}", SanitizerLib);
+					throw new BuildException("No {0} lib found in {1}", LibName, SanitizerLib);
 				}
 			}
 
