UnrealBuildTool: resolve enumerable for windows resource files to prevent reenumeration

#rnx

diff --git a/Engine/Source/Programs/UnrealBuildTool/Configuration/UEBuildBinary.cs b/Engine/Source/Programs/UnrealBuildTool/Configuration/UEBuildBinary.cs
--- a/Engine/Source/Programs/UnrealBuildTool/Configuration/UEBuildBinary.cs
+++ b/Engine/Source/Programs/UnrealBuildTool/Configuration/UEBuildBinary.cs
@@ -847,7 +847,7 @@ namespace UnrealBuildTool
 			BinaryLinkEnvironment.bIsBuildingLibrary = IsBuildingLibrary(Type);
 
 			// Compile any resource files
-			CompileResourceFiles(Target, ToolChain, BinaryCompileEnvironment, BinaryLinkEnvironment, CompiledModules.SelectMany(x => x.ResourceFiles), Graph, Logger);
+			CompileResourceFiles(Target, ToolChain, BinaryCompileEnvironment, BinaryLinkEnvironment, [.. CompiledModules.SelectMany(x => x.ResourceFiles)], Graph, Logger);
 
 			// Setup object export stripping
 			if (Target.bStripExports)
