UnrealBuildTool: Use framework type

#rnx

diff --git a/Engine/Source/Programs/UnrealBuildTool/Configuration/UEBuildModule.cs b/Engine/Source/Programs/UnrealBuildTool/Configuration/UEBuildModule.cs
--- a/Engine/Source/Programs/UnrealBuildTool/Configuration/UEBuildModule.cs
+++ b/Engine/Source/Programs/UnrealBuildTool/Configuration/UEBuildModule.cs
@@ -361,7 +361,7 @@ namespace UnrealBuildTool
 			// disable AutoRTFM on this function the compiler would spit nails. But it'll happily consume the C-style attribute at that location and do
 			// the right thing, so we use that instead.
 			// Caution: Hard coded constant '1' in this string must be kept in sync with Engine\Source\Runtime\AutoRTFM\Public\AutoRTFMConstants.h.
-			String AutoRTFMDisable = (Rules.Target.bUseAutoRTFMCompiler && Rules.bDisableAutoRTFMInstrumentation && Rules.bInjectAutoRTFMAttributeDisables) ?
+			string AutoRTFMDisable = (Rules.Target.bUseAutoRTFMCompiler && Rules.bDisableAutoRTFMInstrumentation && Rules.bInjectAutoRTFMAttributeDisables) ?
 				"__attribute__((autortfm(/*autortfm_mode_disable*/ 1)))" : String.Empty;
 
 			// We precalculate these because they are reused over and over again and causing tons of allocations
diff --git a/Engine/Source/Programs/UnrealBuildTool/Configuration/UEBuildTarget.cs b/Engine/Source/Programs/UnrealBuildTool/Configuration/UEBuildTarget.cs
--- a/Engine/Source/Programs/UnrealBuildTool/Configuration/UEBuildTarget.cs
+++ b/Engine/Source/Programs/UnrealBuildTool/Configuration/UEBuildTarget.cs
@@ -4220,7 +4220,7 @@ namespace UnrealBuildTool
 					{
 						if (dependency.Key.Modules.Any())
 						{
-							Logger.LogDebug("{0}: unmergeable (dependency of: {1})", dependency.Key, string.Join(",", dependency.Value));
+							Logger.LogDebug("{0}: unmergeable (dependency of: {1})", dependency.Key, String.Join(",", dependency.Value));
 						}
 					}
 				}
diff --git a/Engine/Source/Programs/UnrealBuildTool/Executors/UnrealBuildAccelerator/UBAAgentCoordinatorHorde.cs b/Engine/Source/Programs/UnrealBuildTool/Executors/UnrealBuildAccelerator/UBAAgentCoordinatorHorde.cs
--- a/Engine/Source/Programs/UnrealBuildTool/Executors/UnrealBuildAccelerator/UBAAgentCoordinatorHorde.cs
+++ b/Engine/Source/Programs/UnrealBuildTool/Executors/UnrealBuildAccelerator/UBAAgentCoordinatorHorde.cs
@@ -194,7 +194,7 @@ namespace UnrealBuildTool
 				Thread.Sleep(10);
 			}
 
-			var clusterId = new ClusterId(UnrealBuildAcceleratorHordeConfig.ClusterDefault);
+			ClusterId clusterId = new(UnrealBuildAcceleratorHordeConfig.ClusterDefault);
 			using ITimelineEvent timelineEvent = Timeline.ScopeEvent($"HordeGetCache");
 			UbaConfig ubaConfig = _client.AllocateUbaCacheServerAsync(clusterId, cancellationToken).Result;
 			timelineEvent.Finish();
diff --git a/Engine/Source/Programs/UnrealBuildTool/Executors/UnrealBuildAccelerator/UBAExecutor.cs b/Engine/Source/Programs/UnrealBuildTool/Executors/UnrealBuildAccelerator/UBAExecutor.cs
--- a/Engine/Source/Programs/UnrealBuildTool/Executors/UnrealBuildAccelerator/UBAExecutor.cs
+++ b/Engine/Source/Programs/UnrealBuildTool/Executors/UnrealBuildAccelerator/UBAExecutor.cs
@@ -893,7 +893,7 @@ namespace UnrealBuildTool
 				ProcessStartInfo startInfo = GetActionStartInfo(action);
 				startInfo.RootsHandle = GetActionRootsHandle(action);
 
-				var prereqActionsSortIndex = action.PrerequisiteActions.Select(a => a.SortIndex).ToArray();
+				int[] prereqActionsSortIndex = [.. action.PrerequisiteActions.Select(a => a.SortIndex)];
 				bool enableDetour = !ForceLocalNoDetour(action) && action.bCanExecuteInUBA;
 				bool canRunRemotely = CanRunRemotely(action);
 
diff --git a/Engine/Source/Programs/UnrealBuildTool/Platform/Android/AndroidProjectGenerator.cs b/Engine/Source/Programs/UnrealBuildTool/Platform/Android/AndroidProjectGenerator.cs
--- a/Engine/Source/Programs/UnrealBuildTool/Platform/Android/AndroidProjectGenerator.cs
+++ b/Engine/Source/Programs/UnrealBuildTool/Platform/Android/AndroidProjectGenerator.cs
@@ -22,7 +22,7 @@ namespace UnrealBuildTool
 		/// Whether Android Game Development Extension is installed in the system. See https://developer.android.com/games/agde for more details.
 		/// May be disabled by using -noagde on commandline
 		/// </summary>
-		private bool AGDEInstalled { get => AGDEInstalledTask.Result; }
+		private bool AGDEInstalled => AGDEInstalledTask.Result;
 
 		private Task<bool> AGDEInstalledTask;
 
diff --git a/Engine/Source/Programs/UnrealBuildTool/Platform/Android/AndroidToolChain.cs b/Engine/Source/Programs/UnrealBuildTool/Platform/Android/AndroidToolChain.cs
--- a/Engine/Source/Programs/UnrealBuildTool/Platform/Android/AndroidToolChain.cs
+++ b/Engine/Source/Programs/UnrealBuildTool/Platform/Android/AndroidToolChain.cs
@@ -728,7 +728,7 @@ namespace UnrealBuildTool
 
 		protected virtual string GetLinkArguments(LinkEnvironment LinkEnvironment, UnrealArch Architecture, IActionGraphBuilder Graph)
 		{
-			string Result = string.Join(' ', AdditionalPaths(LinkEnvironment.RootPaths));
+			string Result = String.Join(' ', AdditionalPaths(LinkEnvironment.RootPaths));
 
 			//Result += " -nostdlib";
 			Result += " -static-libstdc++";
diff --git a/Engine/Source/Programs/UnrealBuildTool/Platform/Android/UEDeployAndroid.cs b/Engine/Source/Programs/UnrealBuildTool/Platform/Android/UEDeployAndroid.cs
--- a/Engine/Source/Programs/UnrealBuildTool/Platform/Android/UEDeployAndroid.cs
+++ b/Engine/Source/Programs/UnrealBuildTool/Platform/Android/UEDeployAndroid.cs
@@ -4344,7 +4344,7 @@ namespace UnrealBuildTool
 			GradleBuildAdditionsContent.AppendLine("\t}");
 			GradleBuildAdditionsContent.AppendLine("}");
 
-			String BuildGradleAdditionsDefault = "";
+			string BuildGradleAdditionsDefault = "";
 			if (bJavaPerfettoMarkers)
 			{
 				BuildGradleAdditionsDefault = @"
diff --git a/Engine/Source/Programs/UnrealBuildTool/ProjectFiles/Xcode/XcodeProject.cs b/Engine/Source/Programs/UnrealBuildTool/ProjectFiles/Xcode/XcodeProject.cs
--- a/Engine/Source/Programs/UnrealBuildTool/ProjectFiles/Xcode/XcodeProject.cs
+++ b/Engine/Source/Programs/UnrealBuildTool/ProjectFiles/Xcode/XcodeProject.cs
@@ -1063,7 +1063,7 @@ namespace UnrealBuildTool.XcodeProjectXcconfig
 			}
 			MiscItems.Add($");");
 
-			//			string Script = string.Join("&#10", ScriptLines);
+			//			string Script = String.Join("&#10", ScriptLines);
 			string Script = String.Join("\\n", ScriptLines);
 			MiscItems.Add($"shellPath = /bin/sh;");
 			MiscItems.Add($"shellScript = \"{Script}\";");
@@ -1450,7 +1450,7 @@ namespace UnrealBuildTool.XcodeProjectXcconfig
 
 							// is it an extension we support? Grab the Plist file for it and check
 							// Also, only some AppExtensions need to embed
-							String AppExtPlist = DI.FullName + "/Info.plist";
+							string AppExtPlist = DI.FullName + "/Info.plist";
 							PListHelper PlistDoc = PListHelper.CreateFromFile(AppExtPlist);
 							
 							XcodeCopyAppExtFilesBuildPhase EmbedAppExt = new XcodeCopyAppExtFilesBuildPhase(Project.FileCollection, ExtConfig);
diff --git a/Engine/Source/Programs/UnrealBuildTool/System/FileHashCache.cs b/Engine/Source/Programs/UnrealBuildTool/System/FileHashCache.cs
--- a/Engine/Source/Programs/UnrealBuildTool/System/FileHashCache.cs
+++ b/Engine/Source/Programs/UnrealBuildTool/System/FileHashCache.cs
@@ -215,7 +215,7 @@ namespace UnrealBuildTool
 						{
 							ReadOnlySpan<char> last = lastFullName;
 							ReadOnlySpan<char> curr = pair.Key.FullName.AsSpan();
-							UInt16 matchLength = (UInt16)MemoryExtensions.CommonPrefixLength(last, curr);
+							ushort matchLength = (ushort)MemoryExtensions.CommonPrefixLength(last, curr);
 							writer.Write(matchLength);
 							writer.Write(curr.Slice(matchLength).ToString());
 							lastFullName = pair.Key.FullName;
diff --git a/Engine/Source/Programs/UnrealBuildTool/System/NativeProjects.cs b/Engine/Source/Programs/UnrealBuildTool/System/NativeProjects.cs
--- a/Engine/Source/Programs/UnrealBuildTool/System/NativeProjects.cs
+++ b/Engine/Source/Programs/UnrealBuildTool/System/NativeProjects.cs
@@ -256,7 +256,7 @@ namespace UnrealBuildTool
 					string? InnerReason;
 					if (RequiresTempTarget(UProjectFile, bHasCode, Platform, Configuration, TargetType.Game, out InnerReason, Logger))
 					{
-						string PlatformNames = string.Join(", ", Platforms.Select(x => x.ToString()));
+						string PlatformNames = String.Join(", ", Platforms.Select(x => x.ToString()));
 						Reason = $"{UProjectFile.GetFileName()} is has no code, but is being treated as a code-based project for platforms {PlatformNames} because: {InnerReason}.";
 						return true;
 					}
diff --git a/Engine/Source/Programs/UnrealBuildTool/ToolChain/AppleToolChain.cs b/Engine/Source/Programs/UnrealBuildTool/ToolChain/AppleToolChain.cs
--- a/Engine/Source/Programs/UnrealBuildTool/ToolChain/AppleToolChain.cs
+++ b/Engine/Source/Programs/UnrealBuildTool/ToolChain/AppleToolChain.cs
@@ -95,7 +95,7 @@ namespace UnrealBuildTool
 		protected AppleToolChainSettings(string OSPrefix, string? SimulatorOSPrefix, string TargetOSName, string TargetOSVersion, bool bVerbose, ILogger Logger)
 		{
 			this.TargetOSName = TargetOSName;
-			this.TargetSDKName = OSPrefix;
+			TargetSDKName = OSPrefix;
 			PlatformDirName = OSPrefix.ToLower();
 
 			// set up directories
@@ -674,7 +674,7 @@ namespace UnrealBuildTool
 			// pass the UE definitions to C++ land for the imported obj-c header
 			Arguments.AddRange(CompileEnvironment.Definitions.Select(x => $"-Xcc -D{x}"));
 
-			Console.WriteLine("Arguments: {0}", string.Join(", ", Arguments));
+			Console.WriteLine("Arguments: {0}", String.Join(", ", Arguments));
 		}
 
 		private Action CompileSwiftFiles(CppCompileEnvironment CompileEnvironment, IEnumerable<FileItem> InputFiles, DirectoryReference OutputDir, string ModuleName, IActionGraphBuilder Graph, CPPOutput CompileResult)
@@ -723,7 +723,7 @@ namespace UnrealBuildTool
 				string ReadableArch = ArchConfig.ConvertToReadableArchitecture(CompileEnvironment.Architecture);
 				CompileAction.CommandDescription += $" [{ReadableArch}]";
 			}
-			CompileAction.StatusDescription = string.Join(" ", InputFiles.Select(x => x.Name));
+			CompileAction.StatusDescription = String.Join(" ", InputFiles.Select(x => x.Name));
 			CompileAction.bIsClangCompiler = true;
 			CompileAction.bCanExecuteRemotely = true;
 			CompileAction.bCanExecuteInUBA = false;
@@ -1015,7 +1015,7 @@ namespace UnrealBuildTool
 					else
 					{
 						string DestPlatform = AppleExports.GetDestinationPlatform(Platform, Architectures);
-						DestinationParam = string.Join(" ", DestinationIds.Select(x => $"-destination \"platform={DestPlatform},{IdOrNameParam(x)}\""));
+						DestinationParam = String.Join(" ", DestinationIds.Select(x => $"-destination \"platform={DestPlatform},{IdOrNameParam(x)}\""));
 					}
 					List<string> Arguments = new()
 					{
