[UBT]
* Changed ProgressWriter initial state to -1 to make sure default value 0 creates an update. This will fix the "-- Compiling C++ source code..." to pop up straight away when compiling starts

diff --git a/Engine/Source/Programs/UnrealBuildTool/System/ProgressWriter.cs b/Engine/Source/Programs/UnrealBuildTool/System/ProgressWriter.cs
--- a/Engine/Source/Programs/UnrealBuildTool/System/ProgressWriter.cs
+++ b/Engine/Source/Programs/UnrealBuildTool/System/ProgressWriter.cs
@@ -34,7 +34,7 @@ namespace UnrealBuildTool
 		/// <summary>
 		/// The current progress message
 		/// </summary>
-		int CurrentProgressValue = 0;
+		int CurrentProgressValue = -1;
 
 		/// <summary>
 		/// Constructor
