UnrealBuildTool: Add more info to UBA cache connection message

#rnx

diff --git a/Engine/Source/Programs/UnrealBuildTool/Executors/UnrealBuildAccelerator/UBAExecutor.cs b/Engine/Source/Programs/UnrealBuildTool/Executors/UnrealBuildAccelerator/UBAExecutor.cs
--- a/Engine/Source/Programs/UnrealBuildTool/Executors/UnrealBuildAccelerator/UBAExecutor.cs
+++ b/Engine/Source/Programs/UnrealBuildTool/Executors/UnrealBuildAccelerator/UBAExecutor.cs
@@ -339,7 +339,8 @@ namespace UnrealBuildTool
 					bool cacheSuccess = cacheClient.Connect(nameAndPort[0], port, cacheConfig.DesiredConnectionCount);
 					long totalMs = stopwatch.ElapsedMilliseconds;
 					string successText = cacheSuccess ? "Connected to" : "Failed to connect to";
-					logger.LogInformation("UbaCache - {SuccessText} {Name}:{Port} ({Seconds}.{Milliseconds}s)", successText, nameAndPort[0], port, totalMs / 1000, totalMs % 1000);
+					logger.LogInformation("UbaCache - {SuccessText} {Name}:{Port}, CanWrite: {CanWrite}, RequireVFS: {Vfs}, Encrypted: {Encrypted} ({Seconds}.{Milliseconds}s)",
+						successText, nameAndPort[0], port, cacheConfig.CanWrite, cacheConfig.bRequireVfs, !String.IsNullOrEmpty(cacheConfig.Crypto), totalMs / 1000, totalMs % 1000);
 					if (cacheSuccess)
 					{
 						_cacheClients.Add(cacheConfig, cacheClient);
