Fix issue CA2253 (Named placeholders should not be numeric values) in UBT.

diff --git a/Engine/Source/Programs/UnrealBuildTool/Configuration/UEBuildTarget.cs b/Engine/Source/Programs/UnrealBuildTool/Configuration/UEBuildTarget.cs
--- a/Engine/Source/Programs/UnrealBuildTool/Configuration/UEBuildTarget.cs
+++ b/Engine/Source/Programs/UnrealBuildTool/Configuration/UEBuildTarget.cs
@@ -4101,7 +4101,7 @@ namespace UnrealBuildTool
 						pluginGroups.Add(name, self);
 						pluginGroupDependencies.Add(name, dependencies);
 
-						Logger.LogInformation("{Name}: setup for merging with {0} total dependencies", name, dependencies.Count());
+						Logger.LogInformation("{Name}: setup for merging with {NumDependencies} total dependencies", name, dependencies.Count());
 					}
 				}
 
@@ -4117,7 +4117,7 @@ namespace UnrealBuildTool
 				{
 					void AddPluginToGroup(UEBuildPlugin plugin, string groupName)
 					{
-						Logger.LogDebug("{0}: added to plugin group {1}", plugin.Name, groupName);
+						Logger.LogDebug("{PluginName}: added to plugin group {GroupName}", plugin.Name, groupName);
 
 						pluginGroups[groupName].Add(plugin);
 
@@ -4158,7 +4158,7 @@ namespace UnrealBuildTool
 									}
 									else
 									{
-										Logger.LogDebug("{0}: unmergeable", plugin.Name);
+										Logger.LogDebug("{PluginName}: unmergeable", plugin.Name);
 									}
 								}
 							}
@@ -4220,12 +4220,12 @@ namespace UnrealBuildTool
 					{
 						if (dependency.Key.Modules.Any())
 						{
-							Logger.LogDebug("{0}: unmergeable (dependency of: {1})", dependency.Key, String.Join(",", dependency.Value));
+							Logger.LogDebug("{PluginName}: unmergeable (dependency of: {DependantPluginNames})", dependency.Key, String.Join(",", dependency.Value));
 						}
 					}
 				}
 
-				Logger.LogInformation("Found a total of {0} plugins to merge", pluginGroups.SelectMany(x => x.Value).Count());
+				Logger.LogInformation("Found a total of {NumPlugins} plugins to merge", pluginGroups.SelectMany(x => x.Value).Count());
 
 				// Remove unmergable plugin modules from plugin modules and create a set of all plugin modules
 				HashSet<UEBuildModuleCPP> allPluginModules = new HashSet<UEBuildModuleCPP>();
@@ -5844,7 +5844,7 @@ namespace UnrealBuildTool
 						// Ignore any plugin dependency which is programmatically filtered out by the target
 						if (Rules.ShouldIgnorePluginDependency(Info, NextReference))
 						{
-							Logger.LogDebug("Ignoring plugin '{0}' (referenced via {1}) due to it being filtered out by the target rules.", NextReference.Name, Info.Name);
+							Logger.LogDebug("Ignoring plugin '{PluginName}' (referenced via {ReferencingPluginName}) due to it being filtered out by the target rules.", NextReference.Name, Info.Name);
 							continue;
 						}
 
diff --git a/Engine/Source/Programs/UnrealBuildTool/Modes/ServerMode.cs b/Engine/Source/Programs/UnrealBuildTool/Modes/ServerMode.cs
--- a/Engine/Source/Programs/UnrealBuildTool/Modes/ServerMode.cs
+++ b/Engine/Source/Programs/UnrealBuildTool/Modes/ServerMode.cs
@@ -391,7 +391,7 @@ namespace UnrealBuildTool
 			}
 			catch (Exception e)
 			{
-				Logger.LogError("Caught exception setting up target: {0}", e);
+				Logger.LogError(e, "Failed to set up target");
 			}
 
 			CurrentBrowseConfiguration.Paths = BrowseConfigurationFolders.ToList();
diff --git a/Engine/Source/Programs/UnrealBuildTool/Platform/Android/UEDeployAndroid.cs b/Engine/Source/Programs/UnrealBuildTool/Platform/Android/UEDeployAndroid.cs
--- a/Engine/Source/Programs/UnrealBuildTool/Platform/Android/UEDeployAndroid.cs
+++ b/Engine/Source/Programs/UnrealBuildTool/Platform/Android/UEDeployAndroid.cs
@@ -2620,7 +2620,7 @@ namespace UnrealBuildTool
 				{
 					if (!File.Exists(LandscapeFilename))
 					{
-						Logger.LogWarning("Warning: Landscape splash screen source image {0} not available, Oculus splash screen will not function properly!", LandscapeFilename);
+						Logger.LogWarning("Warning: Landscape splash screen source image {LandscapeFilename} not available, Oculus splash screen will not function properly!", LandscapeFilename);
 					}
 					else if (FilesAreDifferent(LandscapeFilename, OculusSplashTargetPath))
 					{
@@ -2628,7 +2628,7 @@ namespace UnrealBuildTool
 						MakeDirectoryIfRequired(OculusSplashTargetPath);
 						File.Copy(LandscapeFilename, OculusSplashTargetPath, true);
 						File.SetLastWriteTimeUtc(OculusSplashTargetPath, File.GetLastWriteTimeUtc(LandscapeFilename));
-						Logger.LogInformation("Copying {0} to {1} for Oculus splash", LandscapeFilename, OculusSplashTargetPath);
+						Logger.LogInformation("Copying {LandscapeFilename} to {OculusSplashTargetPath} for Oculus splash", LandscapeFilename, OculusSplashTargetPath);
 					}
 				}
 				else
diff --git a/Engine/Source/Programs/UnrealBuildTool/Platform/IOS/IOSCodeSigning.cs b/Engine/Source/Programs/UnrealBuildTool/Platform/IOS/IOSCodeSigning.cs
--- a/Engine/Source/Programs/UnrealBuildTool/Platform/IOS/IOSCodeSigning.cs
+++ b/Engine/Source/Programs/UnrealBuildTool/Platform/IOS/IOSCodeSigning.cs
@@ -186,7 +186,7 @@ namespace UnrealBuildTool
 
 			if (!Directory.Exists(ProvisionDirectory))
 			{
-				_logger?.LogError("Could not find provision directory '{0}'.", ProvisionDirectory);
+				_logger?.LogError("Could not find provision directory '{ProvisionDirectory}'.", ProvisionDirectory);
 				MatchedProvision = null;
 				return Provisions;
 			}
@@ -293,14 +293,14 @@ namespace UnrealBuildTool
 				}
 				if (p != null)
 				{
-					_logger?.LogInformation("PROVISION-File:{0},Name:{1},Validity:{2},StartDate:{3},EndDate:{4},Type:{5}", Path.GetFileName(ProvisionFile), p.ProvisionName, Validity, EffectiveDate.ToString(), ExpirationDate.ToString(), bDistribution ? "DISTRIBUTION" : "DEVELOPMENT");
+					_logger?.LogInformation("PROVISION-File:{Filename},Name:{Name},Validity:{Validity},StartDate:{StartDate},EndDate:{EndDate},Type:{Type}", Path.GetFileName(ProvisionFile), p.ProvisionName, Validity, EffectiveDate.ToString(), ExpirationDate.ToString(), bDistribution ? "DISTRIBUTION" : "DEVELOPMENT");
 				}
 				if (ProvisionFile != null)
 				{
 					Provisions.Add(new FileReference(ProvisionFile));
 				}
 			}
-			_logger?.LogInformation(" MATCHED-Provision:{0},File:{1},Cert:{2} ", SelectedProvision, SelectedFile, SelectedCert);
+			_logger?.LogInformation(" MATCHED-Provision:{Name},File:{Filename},Cert:{Cert} ", SelectedProvision, SelectedFile, SelectedCert);
 			if (!String.IsNullOrEmpty(SelectedProvision))
 			{
 				MatchedProvision = new FileReference(SelectedProvision);
@@ -319,7 +319,7 @@ namespace UnrealBuildTool
 		/// </summary>
 		public static X509Certificate2? FindCertificate(MobileProvision ProvisionToWorkFrom)
 		{
-			_logger?.LogInformation("  Looking for a certificate that matches the application identifier '{0}'", ProvisionToWorkFrom.ApplicationIdentifier);
+			_logger?.LogInformation("  Looking for a certificate that matches the application identifier '{ApplicationIdentifier}'", ProvisionToWorkFrom.ApplicationIdentifier);
 
 			X509Certificate2? Result = null;
 
@@ -338,7 +338,7 @@ namespace UnrealBuildTool
 				CertTool.BeginOutputReadLine();
 				CertTool.WaitForExit();
 
-				_logger?.LogInformation("  Running {0} {1}\n{2}", CertTool.StartInfo.FileName, CertTool.StartInfo.Arguments, CertToolData);
+				_logger?.LogInformation("  Running {CertToolPath} {CertToolArguments}\n{CertToolData}", CertTool.StartInfo.FileName, CertTool.StartInfo.Arguments, CertToolData);
 
 				if (CertTool.ExitCode == 0)
 				{
@@ -349,7 +349,7 @@ namespace UnrealBuildTool
 						// see if certificate can be found by serial number
 						string CertHash = SourceCert.GetCertHashString();
 
-						_logger?.LogInformation("    {0}", CertHash);
+						_logger?.LogInformation("    {CertHash}", CertHash);
 
 						if (CertToolData.Contains(CertHash))
 						{
@@ -377,7 +377,7 @@ namespace UnrealBuildTool
 				{
 					X509Certificate2Collection FoundCerts = Store.Certificates.Find(X509FindType.FindBySerialNumber, SourceCert.SerialNumber, false);
 
-					_logger?.LogInformation("  .. Provision entry SN '{0}' matched {1} installed certificate(s)", SourceCert.SerialNumber, FoundCerts.Count);
+					_logger?.LogInformation("  .. Provision entry SN '{SerialNumber}' matched {NumFoundCerts} installed certificate(s)", SourceCert.SerialNumber, FoundCerts.Count);
 
 					X509Certificate2? ValidInTimeCert = null;
 					foreach (X509Certificate2 TestCert in FoundCerts)
@@ -388,7 +388,7 @@ namespace UnrealBuildTool
 
 						bool bCertTimeIsValid = (EffectiveDate < Now) && (ExpirationDate > Now);
 
-						_logger?.LogInformation("  .. .. Installed certificate '{0}' is {1} (range '{2}' to '{3}')", GetFriendlyNameFromCert(TestCert), bCertTimeIsValid ? "valid (choosing it)" : "EXPIRED", TestCert.GetEffectiveDateString(), TestCert.GetExpirationDateString());
+						_logger?.LogInformation("  .. .. Installed certificate '{CertFriendlyName}' is {Validity} (range '{StartDate}' to '{EndDate}')", GetFriendlyNameFromCert(TestCert), bCertTimeIsValid ? "valid (choosing it)" : "EXPIRED", TestCert.GetEffectiveDateString(), TestCert.GetExpirationDateString());
 						if (bCertTimeIsValid)
 						{
 							ValidInTimeCert = TestCert;
@@ -539,7 +539,7 @@ namespace UnrealBuildTool
 					// print out the provision and cert name
 					Provision = new FileReference(FileName);
 					CertName = GetFriendlyNameFromCert(Cert!);
-					_logger?.LogInformation("CERTIFICATE-{0},PROVISION-{1}", CertName, Provision.FullName);
+					_logger?.LogInformation("CERTIFICATE-{Name},PROVISION-{ProvisionFilePath}", CertName, Provision.FullName);
 					return true;
 				}
 			}
@@ -795,7 +795,7 @@ namespace UnrealBuildTool
 			string ProvisionDirectory = AppleCodeSign.ProvisionDirectory;
 			if (!Directory.Exists(ProvisionDirectory))
 			{
-				_logger?.LogInformation("Provision Folder {0} doesn't exist, creating..", ProvisionDirectory);
+				_logger?.LogInformation("Provision Folder {ProvisionDirectory} doesn't exist, creating..", ProvisionDirectory);
 				Directory.CreateDirectory(ProvisionDirectory);
 			}
 
@@ -827,7 +827,7 @@ namespace UnrealBuildTool
 			{
 				if (Directory.Exists(ProvisionCopySrcDirectory))
 				{
-					_logger?.LogInformation("Finding provisions in {0}", ProvisionCopySrcDirectory);
+					_logger?.LogInformation("Finding provisions in {ProvisionCopySrcDirectory}", ProvisionCopySrcDirectory);
 					foreach (string Provision in Directory.EnumerateFiles(ProvisionCopySrcDirectory, "*.mobileprovision", SearchOption.AllDirectories))
 					{
 						string TargetFile = ProvisionDirectory + Path.GetFileName(Provision);
@@ -839,18 +839,18 @@ namespace UnrealBuildTool
 								DestFileInfo = new FileInfo(TargetFile);
 								DestFileInfo.Attributes &= ~FileAttributes.ReadOnly;
 							}
-							_logger?.LogInformation("  Copying {0} -> {1}", Provision, TargetFile);
+							_logger?.LogInformation("  Copying {ProvisionFilePath} -> {TargetFilePath}", Provision, TargetFile);
 							File.Copy(Provision, TargetFile, true);
 							DestFileInfo = new FileInfo(TargetFile);
 							DestFileInfo.Attributes &= ~FileAttributes.ReadOnly;
 							if (!File.Exists(TargetFile))
 							{
-								_logger?.LogInformation("ERROR: Failed to copy {0} -> {1}", Provision, TargetFile);
+								_logger?.LogInformation("ERROR: Failed to copy {ProvisionFilePath} -> {TargetFilePath}", Provision, TargetFile);
 							}
 						}
 						else
 						{
-							_logger?.LogInformation("  Not copying {0} as {1} already exists and is not older", Provision, TargetFile);
+							_logger?.LogInformation("  Not copying {ProvisionFilePath} as {TargetFilePath} already exists and is not older", Provision, TargetFile);
 						}
 					}
 				}
@@ -864,11 +864,11 @@ namespace UnrealBuildTool
 		{
 			if (!Directory.Exists(AppleCodeSign.ProvisionDirectory))
 			{
-				_logger?.LogInformation("Provision Folder {0} doesn't exist, nothing to do.", AppleCodeSign.ProvisionDirectory);
+				_logger?.LogInformation("Provision Folder {ProvisionDirectory} doesn't exist, nothing to do.", AppleCodeSign.ProvisionDirectory);
 			}
 			else
 			{
-				_logger?.LogInformation("Cleaning out contents of  Provision Folder {0}", AppleCodeSign.ProvisionDirectory);
+				_logger?.LogInformation("Cleaning out contents of Provision Folder {ProvisionDirectory}", AppleCodeSign.ProvisionDirectory);
 				foreach (string Provision in Directory.GetFiles(AppleCodeSign.ProvisionDirectory))
 				{
 					File.Delete(Provision);
@@ -932,7 +932,12 @@ namespace UnrealBuildTool
 				}
 			}
 
-			_logger?.LogInformation("Searching for mobile provisions that match the game '{0}' (distribution: {3}) with CFBundleIdentifier='{1}' in '{2}'", GameName, CFBundleIdentifier, ProvisionDirectory, CodeSigningConfig.bForDistribution);
+			_logger?.LogInformation(
+				"Searching for mobile provisions that match the game '{GameName}' (distribution: {ForDistribution}) with CFBundleIdentifier='{CFBundleIdentifier}' in '{ProvisionDirectory}'",
+				GameName,
+				CodeSigningConfig.bForDistribution,
+				CFBundleIdentifier,
+				ProvisionDirectory);
 
 			// first sort all profiles so we look at newer ones first.
 			IEnumerable<string> ProfileKeys = ProvisionLibrary.Select(KV => KV.Key)
@@ -959,7 +964,7 @@ namespace UnrealBuildTool
 						continue;
 					}
 
-					_logger?.LogInformation(" Phase {0} considering provision '{1}' named '{2}'", Phase, DebugName, TestProvision.ProvisionName);
+					_logger?.LogInformation(" Phase {Phase} considering provision '{Filename}' named '{Name}'", Phase, DebugName, TestProvision.ProvisionName);
 
 					if (TestProvision.ProvisionName == "iOS Team Provisioning Profile: " + CFBundleIdentifier)
 					{
@@ -970,7 +975,7 @@ namespace UnrealBuildTool
 					// check to see if the platform is the same as what we are looking for
 					if (!String.IsNullOrEmpty(TestProvision.Platform) && TestProvision.Platform != CodeSigningConfig.OSString && !String.IsNullOrEmpty(CodeSigningConfig.OSString))
 					{
-						_logger?.LogInformation("  Failing platform {0} Config: {1}", TestProvision.Platform, CodeSigningConfig.OSString);
+						_logger?.LogInformation("  Failing platform {Platform} Config: {OS}", TestProvision.Platform, CodeSigningConfig.OSString);
 						continue;
 					}
 
@@ -1008,7 +1013,7 @@ namespace UnrealBuildTool
 					}
 					if (!bPassesNameCheck && bCheckIdentifier)
 					{
-						_logger?.LogInformation("  .. Failed phase {0} name check (provision app ID was {1})", Phase, TestProvision.ApplicationIdentifier);
+						_logger?.LogInformation("  .. Failed phase {Phase} name check (provision app ID was {ApplicationIdentifier})", Phase, TestProvision.ApplicationIdentifier);
 						continue;
 					}
 
@@ -1019,7 +1024,11 @@ namespace UnrealBuildTool
 						bool bDistroProv = !TestProvision.bDebug;
 						if (TestProvision != null && !bDistroProv)
 						{
-							_logger?.LogInformation("  .. Failed distribution check (mode={0}, get-task-allow={1}, #devices={2})", CodeSigningConfig.bForDistribution, TestProvision.bDebug, TestProvision.ProvisionedDeviceIDs?.Count);
+							_logger?.LogInformation(
+								"  .. Failed distribution check (mode={ForDistribution}, get-task-allow={Debug}, #devices={NumDevices})",
+								CodeSigningConfig.bForDistribution,
+								TestProvision.bDebug,
+								TestProvision.ProvisionedDeviceIDs?.Count);
 							continue;
 						}
 					}
@@ -1030,7 +1039,11 @@ namespace UnrealBuildTool
 							bool bPassesDebugCheck = TestProvision?.bDebug ?? false;
 							if (!bPassesDebugCheck)
 							{
-								_logger?.LogInformation(" .. Failed debugging check (mode={0}, get-task-allow={1}, #devices={2})", CodeSigningConfig.bForDistribution, TestProvision?.bDebug ?? false, TestProvision?.ProvisionedDeviceIDs?.Count ?? 0);
+								_logger?.LogInformation(
+									" .. Failed debugging check (mode={ForDistribution}, get-task-allow={Debug}, #devices={NumDevices})",
+									CodeSigningConfig.bForDistribution,
+									TestProvision?.bDebug ?? false,
+									TestProvision?.ProvisionedDeviceIDs?.Count ?? 0);
 								continue;
 							}
 						}
@@ -1044,7 +1057,7 @@ namespace UnrealBuildTool
 					bool bPassesDateCheck = (CurrentUTCTime >= TestProvision?.CreationDate) && (CurrentUTCTime < TestProvision?.ExpirationDate);
 					if (!bPassesDateCheck)
 					{
-						_logger?.LogInformation("  .. Failed time period check (valid from {0} to {1}, but UTC time is now {2})",
+						_logger?.LogInformation("  .. Failed time period check (valid from {StartDate} to {EndDate}, but UTC time is now {CurrentDate})",
 						TestProvision?.CreationDate, TestProvision?.ExpirationDate, CurrentUTCTime);
 						continue;
 					}
@@ -1087,7 +1100,12 @@ namespace UnrealBuildTool
 					}
 
 					// Made it past all the tests
-					_logger?.LogInformation("  Picked '{0}' with AppID '{1}' and Name '{2}' as a matching provision for the game '{3}'", DebugName, TestProvision?.ApplicationIdentifier, TestProvision?.ProvisionName, GameName);
+					_logger?.LogInformation(
+						"  Picked '{Filename}' with AppID '{ApplicationIdentifier}' and Name '{Name}' as a matching provision for the game '{GameName}'",
+						DebugName,
+						TestProvision?.ApplicationIdentifier,
+						TestProvision?.ProvisionName,
+						GameName);
 					return Key;
 				}
 			}
diff --git a/Engine/Source/Programs/UnrealBuildTool/Platform/IOS/IOSToolChain.cs b/Engine/Source/Programs/UnrealBuildTool/Platform/IOS/IOSToolChain.cs
--- a/Engine/Source/Programs/UnrealBuildTool/Platform/IOS/IOSToolChain.cs
+++ b/Engine/Source/Programs/UnrealBuildTool/Platform/IOS/IOSToolChain.cs
@@ -1599,7 +1599,7 @@ namespace UnrealBuildTool
 					{
 						GenerateProjectFiles(Target.ProjectFile, new string[] { "-platforms=" + (Target.Platform == UnrealTargetPlatform.IOS ? "IOS" : "TVOS"), "-NoIntellIsense", (Target.Platform == UnrealTargetPlatform.IOS ? "-iosdeployonly" : "-tvosdeployonly"), "-ignorejunk", (Target.bForDistribution ? "-distribution" : "-development"), String.Format("-project={0}", Target.ProjectFile), "-game", "-bundleID=" + BundleID, "-includetemptargets" }, null, Logger, out GeneratedProjectFile);
 					}
-					Logger.LogInformation("Generated {0}!", GeneratedProjectFile?.FullName ?? "NONE");
+					Logger.LogInformation("Generated {ProjectDirectory}!", GeneratedProjectFile?.FullName ?? "NONE");
 					// Make sure it exists
 					if (!DirectoryReference.Exists(XcodeWorkspaceDir!) || (GeneratedProjectFile != null && GeneratedProjectFile != XcodeWorkspaceDir))
 					{
diff --git a/Engine/Source/Programs/UnrealBuildTool/Platform/Linux/UEBuildLinux.cs b/Engine/Source/Programs/UnrealBuildTool/Platform/Linux/UEBuildLinux.cs
--- a/Engine/Source/Programs/UnrealBuildTool/Platform/Linux/UEBuildLinux.cs
+++ b/Engine/Source/Programs/UnrealBuildTool/Platform/Linux/UEBuildLinux.cs
@@ -559,11 +559,11 @@ namespace UnrealBuildTool
 				{
 					if (!Info.Exists)
 					{
-						Logger.LogWarning("Warning: PGO file '{0}' does not exist, disabling optimization", PGOFilePath);
+						Logger.LogWarning("Warning: PGO file '{PGOFilePath}' does not exist, disabling optimization", PGOFilePath);
 					}
 					else
 					{
-						Logger.LogInformation("PGO file '{0}' exists but has 0 length. Assuming that PGO data is temporarily missing, disabling optimization without a warning.", PGOFilePath);
+						Logger.LogInformation("PGO file '{PGOFilePath}' exists but has 0 length. Assuming that PGO data is temporarily missing, disabling optimization without a warning.", PGOFilePath);
 					}
 					CompileEnvironment.bPGOOptimize = false;
 					LinkEnvironment.bPGOOptimize = false;
diff --git a/Engine/Source/Programs/UnrealBuildTool/ProjectFiles/VisualStudioWorkspace/VSWorkspaceProjectFile.cs b/Engine/Source/Programs/UnrealBuildTool/ProjectFiles/VisualStudioWorkspace/VSWorkspaceProjectFile.cs
--- a/Engine/Source/Programs/UnrealBuildTool/ProjectFiles/VisualStudioWorkspace/VSWorkspaceProjectFile.cs
+++ b/Engine/Source/Programs/UnrealBuildTool/ProjectFiles/VisualStudioWorkspace/VSWorkspaceProjectFile.cs
@@ -164,8 +164,12 @@ namespace UnrealBuildTool
 					}
 					catch (Exception Ex)
 					{
-						Logger.LogWarning("Exception while generating include data for Target:{0}, Platform: {1}, Configuration: {2}", TargetDesc.Name, Platform.ToString(), Configuration.ToString());
-						Logger.LogWarning("{0}", Ex.ToString());
+						Logger.LogWarning(
+							Ex,
+							"Failed to generate include data for Target:{TargetName}, Platform: {Platform}, Configuration: {Configuration}",
+							TargetDesc.Name,
+							Platform.ToString(),
+							Configuration.ToString());
 					}
 				}));
 
diff --git a/Engine/Source/Programs/UnrealBuildTool/ProjectFiles/VisualStudioWorkspace/VSWorkspaceProjectFileGenerator.cs b/Engine/Source/Programs/UnrealBuildTool/ProjectFiles/VisualStudioWorkspace/VSWorkspaceProjectFileGenerator.cs
--- a/Engine/Source/Programs/UnrealBuildTool/ProjectFiles/VisualStudioWorkspace/VSWorkspaceProjectFileGenerator.cs
+++ b/Engine/Source/Programs/UnrealBuildTool/ProjectFiles/VisualStudioWorkspace/VSWorkspaceProjectFileGenerator.cs
@@ -330,7 +330,7 @@ namespace UnrealBuildTool
 			}
 			catch (Exception Ex)
 			{
-				Logger.LogWarning("Exception while writing root project file: {0}", Ex.ToString());
+				Logger.LogWarning(Ex, "Failed to write root project file");
 				return false;
 			}
 
@@ -610,9 +610,9 @@ namespace UnrealBuildTool
 					throw new BuildException("Failed to create workspace project file {0}", Output.FullName);
 				}
 			}
-			catch (Exception)
+			catch (Exception Ex)
 			{
-				Logger.LogWarning("Failed to create workspace project files. Disabling Visual Studio workspace MSBuild integration for target {0}.", Target.Name);
+				Logger.LogWarning(Ex, "Failed to create workspace project files. Disabling Visual Studio workspace MSBuild integration for target {TargetName}.", Target.Name);
 				if (File.Exists(Output.FullName))
 				{
 					File.Delete(Output.FullName);
@@ -712,9 +712,9 @@ namespace UnrealBuildTool
 
 				Utils.WriteFileIfChanged(Output, VCCommonTargetFileContent.ToString(), Logger);
 			}
-			catch (Exception)
+			catch (Exception Ex)
 			{
-				Logger.LogWarning("Failed to create workspace project file {0}.", Output.FullName);
+				Logger.LogWarning(Ex, "Failed to create workspace project file {ProjectFilePath}.", Output.FullName);
 			}
 		}
 
