Selectively disable issue CA5351 (Do Not Use Broken Cryptographic Algorithms) in UBT, for non-cryptographic usages of MD5 where replacing the hash algorithm would be disruptive to users with little upside.

diff --git a/Engine/Source/Programs/UnrealBuildTool/Platform/Android/UEDeployAndroid.cs b/Engine/Source/Programs/UnrealBuildTool/Platform/Android/UEDeployAndroid.cs
--- a/Engine/Source/Programs/UnrealBuildTool/Platform/Android/UEDeployAndroid.cs
+++ b/Engine/Source/Programs/UnrealBuildTool/Platform/Android/UEDeployAndroid.cs
@@ -6299,7 +6299,9 @@ popd
 				return "badpath";
 			}
 
+#pragma warning disable CA5351 // Not used for cryptography, would be disruptive to users to change hash algo
 			MD5 md5 = MD5.Create();
+#pragma warning restore CA5351
 			byte[]? TotalHashBytes = null;
 
 			string[] SourceFiles = Directory.GetFiles(SourceDirectory, "*.*", SearchOption.AllDirectories);
diff --git a/Engine/Source/Programs/UnrealBuildTool/ProjectFiles/VisualStudio/VCProjectFileGenerator.cs b/Engine/Source/Programs/UnrealBuildTool/ProjectFiles/VisualStudio/VCProjectFileGenerator.cs
--- a/Engine/Source/Programs/UnrealBuildTool/ProjectFiles/VisualStudio/VCProjectFileGenerator.cs
+++ b/Engine/Source/Programs/UnrealBuildTool/ProjectFiles/VisualStudio/VCProjectFileGenerator.cs
@@ -501,7 +501,9 @@ namespace UnrealBuildTool
 
 		static Guid MakeMd5Guid(byte[] Input)
 		{
+#pragma warning disable CA5351 // Not used for cryptography, would be disruptive to users to change hash algo
 			byte[] Hash = MD5.Create().ComputeHash(Input);
+#pragma warning restore CA5351
 			Hash[6] = (byte)(0x30 | (Hash[6] & 0x0f)); // 0b0011'xxxx Version 3 UUID (MD5)
 			Hash[8] = (byte)(0x80 | (Hash[8] & 0x3f)); // 0b10xx'xxxx RFC 4122 UUID
 			Array.Reverse(Hash, 0, 4);
