Fix issue CA5362 (Potential reference cycle in deserialized object graph) in UBT.

diff --git a/Engine/Source/Programs/UnrealBuildTool/Preprocessor/SourceFile.cs b/Engine/Source/Programs/UnrealBuildTool/Preprocessor/SourceFile.cs
--- a/Engine/Source/Programs/UnrealBuildTool/Preprocessor/SourceFile.cs
+++ b/Engine/Source/Programs/UnrealBuildTool/Preprocessor/SourceFile.cs
@@ -12,7 +12,6 @@ namespace UnrealBuildTool
 	/// <summary>
 	/// Encapsulates metadata for a C++ source file
 	/// </summary>
-	[Serializable]
 	class SourceFile
 	{
 		/// <summary>
diff --git a/Engine/Source/Programs/UnrealBuildTool/Preprocessor/SourceFileFragment.cs b/Engine/Source/Programs/UnrealBuildTool/Preprocessor/SourceFileFragment.cs
--- a/Engine/Source/Programs/UnrealBuildTool/Preprocessor/SourceFileFragment.cs
+++ b/Engine/Source/Programs/UnrealBuildTool/Preprocessor/SourceFileFragment.cs
@@ -8,7 +8,6 @@ namespace UnrealBuildTool
 	/// <summary>
 	/// Represents an fragment of a source file consisting of a sequence of includes followed by some arbitrary directives or source code.
 	/// </summary>
-	[Serializable]
 	class SourceFileFragment
 	{
 		/// <summary>
diff --git a/Engine/Source/Programs/UnrealBuildTool/System/TimingData.cs b/Engine/Source/Programs/UnrealBuildTool/System/TimingData.cs
--- a/Engine/Source/Programs/UnrealBuildTool/System/TimingData.cs
+++ b/Engine/Source/Programs/UnrealBuildTool/System/TimingData.cs
@@ -49,7 +49,6 @@ namespace UnrealBuildTool
 	/// Class used to capture timing data from events like those generated by MSVC when timing data
 	/// is being generated.
 	/// </summary>
-	[Serializable]
 	public class TimingData : IBinarySerializable
 	{
 		/// <summary>
