Fix issue CA5369 (Use XmlReader for Deserialize) in UBT.

diff --git a/Engine/Source/Programs/UnrealBuildTool/Platform/Windows/PVSToolChain.cs b/Engine/Source/Programs/UnrealBuildTool/Platform/Windows/PVSToolChain.cs
--- a/Engine/Source/Programs/UnrealBuildTool/Platform/Windows/PVSToolChain.cs
+++ b/Engine/Source/Programs/UnrealBuildTool/Platform/Windows/PVSToolChain.cs
@@ -13,6 +13,7 @@ using System.Text;
 using System.Text.Json;
 using System.Text.RegularExpressions;
 using System.Threading.Tasks;
+using System.Xml;
 using System.Xml.Serialization;
 using UnrealBuildBase;
 
@@ -176,7 +177,8 @@ namespace UnrealBuildTool
 				{
 					XmlSerializer Serializer = new(typeof(PVSApplicationSettings));
 					using FileStream Stream = new(SettingsPath.FullName, FileMode.Open, FileAccess.Read, FileShare.Read);
-					return (PVSApplicationSettings?)Serializer.Deserialize(Stream);
+					using XmlReader XmlReader = XmlReader.Create(Stream);
+					return (PVSApplicationSettings?)Serializer.Deserialize(XmlReader);
 				}
 				catch (Exception Ex)
 				{
diff --git a/Engine/Source/Programs/UnrealBuildTool/System/Utils.cs b/Engine/Source/Programs/UnrealBuildTool/System/Utils.cs
--- a/Engine/Source/Programs/UnrealBuildTool/System/Utils.cs
+++ b/Engine/Source/Programs/UnrealBuildTool/System/Utils.cs
@@ -9,6 +9,7 @@ using System.Linq;
 using System.Runtime.InteropServices;
 using System.Text;
 using System.Text.RegularExpressions;
+using System.Xml;
 using System.Xml.Serialization;
 using EpicGames.Core;
 using Microsoft.Extensions.Logging;
@@ -842,7 +843,8 @@ namespace UnrealBuildTool
 				Serialiser.UnknownAttribute += new XmlAttributeEventHandler(UnknownAttributeDelegate);
 
 				// Create an object graph from the XML data
-				Instance = (T)Serialiser.Deserialize(XmlStream)!;
+				using XmlReader XmlReader = XmlReader.Create(XmlStream);
+				Instance = (T)Serialiser.Deserialize(XmlReader)!;
 			}
 			catch (Exception E)
 			{
