Fix issues IDE0011 (Add braces) and IDE2001 (Embedded statements must be on their own line) in UBT.

diff --git a/Engine/Source/Programs/UnrealBuildTool/Configuration/ModuleRules.cs b/Engine/Source/Programs/UnrealBuildTool/Configuration/ModuleRules.cs
--- a/Engine/Source/Programs/UnrealBuildTool/Configuration/ModuleRules.cs
+++ b/Engine/Source/Programs/UnrealBuildTool/Configuration/ModuleRules.cs
@@ -2125,7 +2125,7 @@ namespace UnrealBuildTool
 		/// <param name="items"></param>
 		public static void AddRange<T>(this HashSet<T> set, params T[] items)
 		{
-			foreach (T? item in items) set.Add(item);
+			set.UnionWith(items);
 		}
 	}
-}
\ No newline at end of file
+}
