Fix issue CA1041 (Provide ObsoleteAttribute message) in UBT.

diff --git a/Engine/Source/Programs/UnrealBuildTool/Configuration/TestModuleRules.cs b/Engine/Source/Programs/UnrealBuildTool/Configuration/TestModuleRules.cs
--- a/Engine/Source/Programs/UnrealBuildTool/Configuration/TestModuleRules.cs
+++ b/Engine/Source/Programs/UnrealBuildTool/Configuration/TestModuleRules.cs
@@ -504,12 +504,12 @@ namespace UnrealBuildTool
 			public bool Deactivated { get; set; }
 
 			/// <summary>
-			/// Depercated, use GauntletArgs or ExtraArgs instead to help indicate arguments to launch the test under.
+			/// Deprecated, use GauntletArgs or ExtraArgs instead to help indicate arguments to launch the test under.
 			/// </summary>
 			public string InitialExtraArgs
 			{
 				get;
-				[Obsolete]
+				[Obsolete("Deprecated in UE5.3 - use GauntletArgs or ExtraArgs")]
 				set;
 			}
 
diff --git a/Engine/Source/Programs/UnrealBuildTool/Platform/Windows/UEBuildWindows.cs b/Engine/Source/Programs/UnrealBuildTool/Platform/Windows/UEBuildWindows.cs
--- a/Engine/Source/Programs/UnrealBuildTool/Platform/Windows/UEBuildWindows.cs
+++ b/Engine/Source/Programs/UnrealBuildTool/Platform/Windows/UEBuildWindows.cs
@@ -498,7 +498,7 @@ namespace UnrealBuildTool
 		/// Enable building with the Win10 SDK instead of the older Win8.1 SDK 
 		/// </summary>
 		[ConfigFile(ConfigHierarchyType.Engine, "/Script/WindowsTargetPlatform.WindowsTargetSettings", "bUseWindowsSDK10")]
-		[Obsolete]
+		[Obsolete("Set MainVersion in Windows_SDK.json")]
 		public bool bUseWindowsSDK10 = false;
 
 		/// <summary>
diff --git a/Engine/Source/Programs/UnrealBuildTool/System/PlatformExports.cs b/Engine/Source/Programs/UnrealBuildTool/System/PlatformExports.cs
--- a/Engine/Source/Programs/UnrealBuildTool/System/PlatformExports.cs
+++ b/Engine/Source/Programs/UnrealBuildTool/System/PlatformExports.cs
@@ -104,7 +104,7 @@ namespace UnrealBuildTool
 		/// <param name="Platform">Platform to check</param>
 		/// <param name="Logger">Logger for output</param>
 		/// <returns>True if the platform supports the parallel executor in UAT</returns>
-		[Obsolete]
+		[Obsolete("All platforms support the parallel executor")]
 		public static bool CanUseParallelExecutor(UnrealTargetPlatform Platform, ILogger Logger)
 		{
 			return true;
