Fix issue CA1008 (Enums should have zero value) in UBT.

diff --git a/Engine/Source/Programs/UnrealBuildTool/Platform/Windows/PVSToolChain.cs b/Engine/Source/Programs/UnrealBuildTool/Platform/Windows/PVSToolChain.cs
--- a/Engine/Source/Programs/UnrealBuildTool/Platform/Windows/PVSToolChain.cs
+++ b/Engine/Source/Programs/UnrealBuildTool/Platform/Windows/PVSToolChain.cs
@@ -24,6 +24,11 @@ namespace UnrealBuildTool
 	/// </summary>
 	public enum PVSAnalysisModeFlags : uint
 	{
+		/// <summary>
+		/// No flags
+		/// </summary>
+		None = 0,
+
 		/// <summary>
 		/// Check for 64-bit portability issues
 		/// </summary>
@@ -51,10 +56,14 @@ namespace UnrealBuildTool
 	}
 
 	/// <summary>
-	/// Flags for the PVS analyzer timeout
+	/// Supported PVS analyzer timeout values
 	/// </summary>
-	public enum AnalysisTimeoutFlags
+	public enum AnalysisTimeout
 	{
+		/// <summary>
+		/// Analysis timeout when not set (infinite)
+		/// </summary>
+		No_timeout = 0,
 		/// <summary>
 		/// Analysis timeout for file 10 minutes (600 seconds)
 		/// </summary>
@@ -67,10 +76,6 @@ namespace UnrealBuildTool
 		/// Analysis timeout for file 60 minutes (3600 seconds)
 		/// </summary>
 		After_60_minutes = 3600,
-		/// <summary>
-		/// Analysis timeout when not set (a lot of seconds)
-		/// </summary>
-		No_timeout = 999999
 	}
 
 	/// <summary>
@@ -122,7 +127,7 @@ namespace UnrealBuildTool
 		/// <summary>
 		/// File analysis timeout
 		/// </summary>
-		public AnalysisTimeoutFlags AnalysisTimeout;
+		public AnalysisTimeout AnalysisTimeout;
 
 		/// <summary>
 		/// Disable analyzer Level 3 (Low) messages
@@ -235,7 +240,7 @@ namespace UnrealBuildTool
 		/// <summary>
 		/// Override for the analysis timeoutFlag to use
 		/// </summary>
-		public AnalysisTimeoutFlags AnalysisTimeoutFlag
+		public AnalysisTimeout AnalysisTimeoutFlag
 		{
 			get
 			{
@@ -249,7 +254,7 @@ namespace UnrealBuildTool
 				}
 				else
 				{
-					return AnalysisTimeoutFlags.After_30_minutes;
+					return AnalysisTimeout.After_30_minutes;
 				}
 			}
 			set => TimeoutPrivate = value;
@@ -258,7 +263,7 @@ namespace UnrealBuildTool
 		/// <summary>
 		/// Private storage for the analysis timeout
 		/// </summary>
-		AnalysisTimeoutFlags? TimeoutPrivate;
+		AnalysisTimeout? TimeoutPrivate;
 
 		/// <summary>
 		/// Override for the disable Level 3 (Low) analyzer messages
@@ -344,7 +349,7 @@ namespace UnrealBuildTool
 		/// <summary>
 		/// Override for the analysis timeout to use
 		/// </summary>
-		public AnalysisTimeoutFlags AnalysisTimeoutFlag => Inner.AnalysisTimeoutFlag;
+		public AnalysisTimeout AnalysisTimeoutFlag => Inner.AnalysisTimeoutFlag;
 
 		/// <summary>
 		/// Override NoNoise analysis setting to use
@@ -827,7 +832,7 @@ namespace UnrealBuildTool
 				configFileContents.AppendFormat($"errors-off={error}\n");
 			}
 
-			int timeout = Settings.AnalysisTimeoutFlag == AnalysisTimeoutFlags.No_timeout ? 0 : (int)Settings.AnalysisTimeoutFlag;
+			int timeout = (int)Settings.AnalysisTimeoutFlag;
 			configFileContents.AppendFormat("timeout={0}\n", timeout);
 			configFileContents.Append("silent-exit-code-mode=yes\n");
 			configFileContents.Append("new-output-format=yes\n");
diff --git a/Engine/Source/Programs/UnrealBuildTool/System/InstalledPlatformInfo.cs b/Engine/Source/Programs/UnrealBuildTool/System/InstalledPlatformInfo.cs
--- a/Engine/Source/Programs/UnrealBuildTool/System/InstalledPlatformInfo.cs
+++ b/Engine/Source/Programs/UnrealBuildTool/System/InstalledPlatformInfo.cs
@@ -37,7 +37,6 @@ namespace UnrealBuildTool
 	/// <summary>
 	/// The state of a downloaded platform
 	/// </summary>
-	[Flags]
 	public enum InstalledPlatformState
 	{
 		/// <summary>
