Fix issue CA1309 (Use ordinal StringComparison) in UBT.

diff --git a/Engine/Source/Programs/UnrealBuildTool/Configuration/TargetDescriptor.cs b/Engine/Source/Programs/UnrealBuildTool/Configuration/TargetDescriptor.cs
--- a/Engine/Source/Programs/UnrealBuildTool/Configuration/TargetDescriptor.cs
+++ b/Engine/Source/Programs/UnrealBuildTool/Configuration/TargetDescriptor.cs
@@ -572,7 +572,7 @@ namespace UnrealBuildTool
 						{
 							// find one with a matching name
 							TargetProjectFile = NativeProjects.EnumerateProjectFiles(Logger)
-								.Where(x => x.GetFileNameWithoutAnyExtensions().Equals(TargetName, StringComparison.InvariantCultureIgnoreCase))
+								.Where(x => x.GetFileNameWithoutAnyExtensions().Equals(TargetName, StringComparison.OrdinalIgnoreCase))
 								.FirstOrDefault();
 						}
 
diff --git a/Engine/Source/Programs/UnrealBuildTool/Configuration/UEBuildBinary.cs b/Engine/Source/Programs/UnrealBuildTool/Configuration/UEBuildBinary.cs
--- a/Engine/Source/Programs/UnrealBuildTool/Configuration/UEBuildBinary.cs
+++ b/Engine/Source/Programs/UnrealBuildTool/Configuration/UEBuildBinary.cs
@@ -654,7 +654,7 @@ namespace UnrealBuildTool
 					{
 						foreach (RestrictedFolder Folder in RestrictedFolder.GetValues())
 						{
-							if (Folder.ToString().Equals(Alias))
+							if (Folder.ToString().Equals(Alias, StringComparison.Ordinal))
 							{
 								AliasedBinaryFolders.Add(Folder);
 							}
diff --git a/Engine/Source/Programs/UnrealBuildTool/Configuration/UEBuildModule.cs b/Engine/Source/Programs/UnrealBuildTool/Configuration/UEBuildModule.cs
--- a/Engine/Source/Programs/UnrealBuildTool/Configuration/UEBuildModule.cs
+++ b/Engine/Source/Programs/UnrealBuildTool/Configuration/UEBuildModule.cs
@@ -888,7 +888,7 @@ namespace UnrealBuildTool
 		/// <returns>True if the variable name matches</returns>
 		private bool MatchVariableName(string Path, int StartIdx, int EndIdx, string Name)
 		{
-			return Name.Length == EndIdx - StartIdx && String.Compare(Path, StartIdx, Name, 0, EndIdx - StartIdx) == 0;
+			return Name.Length == EndIdx - StartIdx && String.Compare(Path, StartIdx, Name, 0, EndIdx - StartIdx, StringComparison.Ordinal) == 0;
 		}
 
 		/// <summary>
diff --git a/Engine/Source/Programs/UnrealBuildTool/Configuration/UEBuildPlatform.cs b/Engine/Source/Programs/UnrealBuildTool/Configuration/UEBuildPlatform.cs
--- a/Engine/Source/Programs/UnrealBuildTool/Configuration/UEBuildPlatform.cs
+++ b/Engine/Source/Programs/UnrealBuildTool/Configuration/UEBuildPlatform.cs
@@ -552,7 +552,7 @@ namespace UnrealBuildTool
 		public static bool IsBuildProductName(string FileName, int Index, int Count, string[] NamePrefixes, string[] NameSuffixes, string Extension)
 		{
 			// Check if the extension matches, and forward on to the next IsBuildProductName() overload without it if it does.
-			if (Count > Extension.Length && String.Compare(FileName, Index + Count - Extension.Length, Extension, 0, Extension.Length, StringComparison.InvariantCultureIgnoreCase) == 0)
+			if (Count > Extension.Length && String.Compare(FileName, Index + Count - Extension.Length, Extension, 0, Extension.Length, StringComparison.OrdinalIgnoreCase) == 0)
 			{
 				return IsBuildProductName(FileName, Index, Count - Extension.Length, NamePrefixes, NameSuffixes);
 			}
@@ -572,13 +572,13 @@ namespace UnrealBuildTool
 		{
 			foreach (string NamePrefix in NamePrefixes)
 			{
-				if (Count >= NamePrefix.Length && String.Compare(FileName, Index, NamePrefix, 0, NamePrefix.Length, StringComparison.InvariantCultureIgnoreCase) == 0)
+				if (Count >= NamePrefix.Length && String.Compare(FileName, Index, NamePrefix, 0, NamePrefix.Length, StringComparison.OrdinalIgnoreCase) == 0)
 				{
 					int MinIdx = Index + NamePrefix.Length;
 					foreach (string NameSuffix in NameSuffixes)
 					{
 						int MaxIdx = Index + Count - NameSuffix.Length;
-						if (MaxIdx >= MinIdx && String.Compare(FileName, MaxIdx, NameSuffix, 0, NameSuffix.Length, StringComparison.InvariantCultureIgnoreCase) == 0)
+						if (MaxIdx >= MinIdx && String.Compare(FileName, MaxIdx, NameSuffix, 0, NameSuffix.Length, StringComparison.OrdinalIgnoreCase) == 0)
 						{
 							if (MinIdx < MaxIdx && FileName[MinIdx] == '-')
 							{
diff --git a/Engine/Source/Programs/UnrealBuildTool/Configuration/UEBuildTarget.cs b/Engine/Source/Programs/UnrealBuildTool/Configuration/UEBuildTarget.cs
--- a/Engine/Source/Programs/UnrealBuildTool/Configuration/UEBuildTarget.cs
+++ b/Engine/Source/Programs/UnrealBuildTool/Configuration/UEBuildTarget.cs
@@ -1849,7 +1849,7 @@ namespace UnrealBuildTool
 				ProjectDirectory = ProjectFile.Directory;
 				// programs are often compiled with another project passed in for context, so if it's using a _different_ uproject, use that uproject
 				// for the output location
-				if (Rules.Type != TargetType.Program || !String.Equals(ProjectFile.GetFileNameWithoutAnyExtensions(), Rules.Name, StringComparison.CurrentCultureIgnoreCase))
+				if (Rules.Type != TargetType.Program || !String.Equals(ProjectFile.GetFileNameWithoutAnyExtensions(), Rules.Name, StringComparison.OrdinalIgnoreCase))
 				{
 					OutputRootDirectory = ProjectDirectory;
 				}
@@ -3236,7 +3236,7 @@ namespace UnrealBuildTool
 					{
 					WriteMetadataTargetInfo EngineInfo = new WriteMetadataTargetInfo(null, VersionFile, Receipt.Version, null, null, EngineFileToManifest, null, Rules.bMergeModules);
 						FileReference EngineInfoFile = FileReference.Combine(ProjectIntermediateDirectory, "EngineMetadata.dat");
-						string EngineInfoName = TargetName.Equals(VersionFile.GetFileNameWithoutExtension()) ? VersionFile.GetFileName() : $"{VersionFile.GetFileName()} ({TargetName})";
+						string EngineInfoName = TargetName.Equals(VersionFile.GetFileNameWithoutExtension(), StringComparison.Ordinal) ? VersionFile.GetFileName() : $"{VersionFile.GetFileName()} ({TargetName})";
 						CreateWriteMetadataAction(MakefileBuilder, EngineInfoName, EngineInfoFile, EngineInfo, EnginePrereqItems);
 					}));
 					writeMetadataTasks.Add(Task.Run(() =>
@@ -5235,7 +5235,7 @@ namespace UnrealBuildTool
 
 			// Get the configuration that this module will be built in. Engine modules compiled in DebugGame will use Development.
 			UnrealTargetConfiguration ModuleConfiguration = Configuration;
-			if (Configuration == UnrealTargetConfiguration.DebugGame && !ModuleRules.Context.bCanBuildDebugGame && !ModuleRules.Name.Equals(Rules.LaunchModuleName, StringComparison.InvariantCultureIgnoreCase))
+			if (Configuration == UnrealTargetConfiguration.DebugGame && !ModuleRules.Context.bCanBuildDebugGame && !ModuleRules.Name.Equals(Rules.LaunchModuleName, StringComparison.OrdinalIgnoreCase))
 			{
 				ModuleConfiguration = UnrealTargetConfiguration.Development;
 			}
@@ -6619,7 +6619,7 @@ namespace UnrealBuildTool
 			// Reads additional dependencies array for project module from project file and fills PrivateDependencyModuleNames.
 			if (ProjectDescriptor != null && ProjectDescriptor.Modules != null)
 			{
-				ModuleDescriptor? Module = ProjectDescriptor.Modules.FirstOrDefault(x => x.Name.Equals(ModuleName, StringComparison.InvariantCultureIgnoreCase));
+				ModuleDescriptor? Module = ProjectDescriptor.Modules.FirstOrDefault(x => x.Name.Equals(ModuleName, StringComparison.OrdinalIgnoreCase));
 				if (Module != null && Module.AdditionalDependencies != null)
 				{
 					RulesObject.PrivateDependencyModuleNames.AddRange(Module.AdditionalDependencies);
@@ -6829,7 +6829,7 @@ namespace UnrealBuildTool
 			}
 
 			// Warn if the module reference has incorrect text case
-			if (!Module.Name.Equals(ModuleName) && Module.Name.Equals(ModuleName, StringComparison.InvariantCultureIgnoreCase))
+			if (!Module.Name.Equals(ModuleName, StringComparison.Ordinal) && Module.Name.Equals(ModuleName, StringComparison.OrdinalIgnoreCase))
 			{
 				Logger.LogWarning("Module '{ModuleName}' (referenced via {ReferenceChain}) has incorrect text case. Did you mean '{PossibleModuleName}'?", ModuleName, ReferenceChain, Module.Name);
 			}
diff --git a/Engine/Source/Programs/UnrealBuildTool/Executors/SNDBS.cs b/Engine/Source/Programs/UnrealBuildTool/Executors/SNDBS.cs
--- a/Engine/Source/Programs/UnrealBuildTool/Executors/SNDBS.cs
+++ b/Engine/Source/Programs/UnrealBuildTool/Executors/SNDBS.cs
@@ -369,7 +369,7 @@ namespace UnrealBuildTool
 							return;
 						}
 						// Suppress redundant tool output of status we already printed (e.g., msvc cl prints compile unit name always)
-						if (!Text.Equals(CurrentStatus))
+						if (!Text.Equals(CurrentStatus, StringComparison.Ordinal))
 						{
 							WriteToolOutput(Text);
 						}
diff --git a/Engine/Source/Programs/UnrealBuildTool/Modes/AggregateClangTimingInfo.cs b/Engine/Source/Programs/UnrealBuildTool/Modes/AggregateClangTimingInfo.cs
--- a/Engine/Source/Programs/UnrealBuildTool/Modes/AggregateClangTimingInfo.cs
+++ b/Engine/Source/Programs/UnrealBuildTool/Modes/AggregateClangTimingInfo.cs
@@ -71,26 +71,26 @@ namespace UnrealBuildTool
 			{
 				SourceFile = inputFile;
 
-				TotalExecuteCompiler = trace?.traceEvents?.FindLast(x => String.Equals(x.name, "Total ExecuteCompiler"))?.dur ?? 0;
+				TotalExecuteCompiler = trace?.traceEvents?.FindLast(x => String.Equals(x.name, "Total ExecuteCompiler", StringComparison.Ordinal))?.dur ?? 0;
 
 				// Subset of execute compiler
-				TotalFrontend = trace?.traceEvents?.FindLast(x => String.Equals(x.name, "Total Frontend"))?.dur ?? 0;
-				TotalBackend = trace?.traceEvents?.FindLast(x => String.Equals(x.name, "Total Backend"))?.dur ?? 0;
+				TotalFrontend = trace?.traceEvents?.FindLast(x => String.Equals(x.name, "Total Frontend", StringComparison.Ordinal))?.dur ?? 0;
+				TotalBackend = trace?.traceEvents?.FindLast(x => String.Equals(x.name, "Total Backend", StringComparison.Ordinal))?.dur ?? 0;
 
 				// Subset of frontend
-				TotalSource = trace?.traceEvents?.FindLast(x => String.Equals(x.name, "Total Source"))?.dur ?? 0;
-				TotalInstantiateFunction = trace?.traceEvents?.FindLast(x => String.Equals(x.name, "Total InstantiateFunction"))?.dur ?? 0;
-				TotalCodeGenFunction = trace?.traceEvents?.FindLast(x => String.Equals(x.name, "Total CodeGen Function"))?.dur ?? 0;
+				TotalSource = trace?.traceEvents?.FindLast(x => String.Equals(x.name, "Total Source", StringComparison.Ordinal))?.dur ?? 0;
+				TotalInstantiateFunction = trace?.traceEvents?.FindLast(x => String.Equals(x.name, "Total InstantiateFunction", StringComparison.Ordinal))?.dur ?? 0;
+				TotalCodeGenFunction = trace?.traceEvents?.FindLast(x => String.Equals(x.name, "Total CodeGen Function", StringComparison.Ordinal))?.dur ?? 0;
 
 				// Subset of backend
-				TotalModuleToFunctionPassAdaptor = trace?.traceEvents?.FindLast(x => String.Equals(x.name, "Total ModuleToFunctionPassAdaptor"))?.dur ?? 0;
-				TotalModuleInlinerWrapperPass = trace?.traceEvents?.FindLast(x => String.Equals(x.name, "Total ModuleInlinerWrapperPass"))?.dur ?? 0;
-				TotalOptModule = trace?.traceEvents?.FindLast(x => String.Equals(x.name, "Total OptModule"))?.dur ?? 0;
+				TotalModuleToFunctionPassAdaptor = trace?.traceEvents?.FindLast(x => String.Equals(x.name, "Total ModuleToFunctionPassAdaptor", StringComparison.Ordinal))?.dur ?? 0;
+				TotalModuleInlinerWrapperPass = trace?.traceEvents?.FindLast(x => String.Equals(x.name, "Total ModuleInlinerWrapperPass", StringComparison.Ordinal))?.dur ?? 0;
+				TotalOptModule = trace?.traceEvents?.FindLast(x => String.Equals(x.name, "Total OptModule", StringComparison.Ordinal))?.dur ?? 0;
 
 				// Frontend entry counts
-				SourceEntries = trace?.traceEvents?.Where(x => String.Equals(x.name, "Source")).LongCount() ?? 0;
-				InstantiateFunctionEntries = trace?.traceEvents?.Where(x => String.Equals(x.name, "InstantiateFunction")).LongCount() ?? 0;
-				CodeGenFunctionEntries = trace?.traceEvents?.Where(x => String.Equals(x.name, "CodeGen Function")).LongCount() ?? 0;
+				SourceEntries = trace?.traceEvents?.Where(x => String.Equals(x.name, "Source", StringComparison.Ordinal)).LongCount() ?? 0;
+				InstantiateFunctionEntries = trace?.traceEvents?.Where(x => String.Equals(x.name, "InstantiateFunction", StringComparison.Ordinal)).LongCount() ?? 0;
+				CodeGenFunctionEntries = trace?.traceEvents?.Where(x => String.Equals(x.name, "CodeGen Function", StringComparison.Ordinal)).LongCount() ?? 0;
 
 				// Other
 				ObjectBytes = GetObjectSize();
@@ -205,7 +205,7 @@ namespace UnrealBuildTool
 				Dictionary<FileReference, List<long>> Sources = new Dictionary<FileReference, List<long>>();
 				foreach (ClangTrace ClangTrace in ClangTraces.Where(x => x.traceEvents != null))
 				{
-					foreach (ClangTrace.TraceEvent Event in ClangTrace.traceEvents!.Where(x => String.Equals(x.name, "Source") && x.args?.ContainsKey("detail") == true))
+					foreach (ClangTrace.TraceEvent Event in ClangTrace.traceEvents!.Where(x => String.Equals(x.name, "Source", StringComparison.Ordinal) && x.args?.ContainsKey("detail") == true))
 					{
 						FileReference SourceFile = new FileReference(Event.args!["detail"].ToString()!);
 						if (!Sources.ContainsKey(SourceFile))
diff --git a/Engine/Source/Programs/UnrealBuildTool/Modes/FixIncludePathsMode.cs b/Engine/Source/Programs/UnrealBuildTool/Modes/FixIncludePathsMode.cs
--- a/Engine/Source/Programs/UnrealBuildTool/Modes/FixIncludePathsMode.cs
+++ b/Engine/Source/Programs/UnrealBuildTool/Modes/FixIncludePathsMode.cs
@@ -229,8 +229,8 @@ namespace UnrealBuildTool
 												}
 
 												// if the include and the source file live in the same directory then it is OK to be relative
-												if (String.Equals(System.IO.Directory.GetParent(FullPath)?.FullName, System.IO.Directory.GetParent(InputFile.FullName)?.FullName, StringComparison.CurrentCultureIgnoreCase) &&
-													String.Equals(Include, System.IO.Path.GetFileName(FullPath), StringComparison.CurrentCultureIgnoreCase))
+												if (String.Equals(System.IO.Directory.GetParent(FullPath)?.FullName, System.IO.Directory.GetParent(InputFile.FullName)?.FullName, StringComparison.OrdinalIgnoreCase) &&
+													String.Equals(Include, System.IO.Path.GetFileName(FullPath), StringComparison.OrdinalIgnoreCase))
 												{
 													Logger.LogInformation("{FileName}({LineNumber}): Using '{Include}' because it is in the same directory.", InputFile.FullName, LineNumber, Include);
 													PreferredInclude = Include;
@@ -403,7 +403,7 @@ namespace UnrealBuildTool
 					return -1;
 				}
 
-				return String.Compare(x, y);
+				return String.Compare(x, y, StringComparison.Ordinal);
 			}
 		}
 
diff --git a/Engine/Source/Programs/UnrealBuildTool/Modes/PipInstallMode.cs b/Engine/Source/Programs/UnrealBuildTool/Modes/PipInstallMode.cs
--- a/Engine/Source/Programs/UnrealBuildTool/Modes/PipInstallMode.cs
+++ b/Engine/Source/Programs/UnrealBuildTool/Modes/PipInstallMode.cs
@@ -478,15 +478,15 @@ namespace UnrealBuildTool.Modes
 		private static bool CompatiblePlatform(string? PlatformField, UnrealTargetPlatform Platform)
 		{
 			return (PlatformField == null)
-				|| String.Equals(PlatformField, Platform.ToString(), StringComparison.InvariantCultureIgnoreCase)
-				|| String.Equals(PlatformField, "All", StringComparison.InvariantCultureIgnoreCase);
+				|| String.Equals(PlatformField, Platform.ToString(), StringComparison.OrdinalIgnoreCase)
+				|| String.Equals(PlatformField, "All", StringComparison.OrdinalIgnoreCase);
 		}
 
 		private static bool CompatibleArch(string? ArchField, UnrealArch Arch)
 		{
 			return (ArchField == null)
-				|| String.Equals(ArchField, MakeArchString(Arch), StringComparison.InvariantCultureIgnoreCase)
-				|| String.Equals(ArchField, "All", StringComparison.InvariantCultureIgnoreCase);
+				|| String.Equals(ArchField, MakeArchString(Arch), StringComparison.OrdinalIgnoreCase)
+				|| String.Equals(ArchField, "All", StringComparison.OrdinalIgnoreCase);
 		}
 
 		private static string MakeArchString(UnrealArch Arch)
diff --git a/Engine/Source/Programs/UnrealBuildTool/Platform/Android/AndroidPlatformSDK.cs b/Engine/Source/Programs/UnrealBuildTool/Platform/Android/AndroidPlatformSDK.cs
--- a/Engine/Source/Programs/UnrealBuildTool/Platform/Android/AndroidPlatformSDK.cs
+++ b/Engine/Source/Programs/UnrealBuildTool/Platform/Android/AndroidPlatformSDK.cs
@@ -304,7 +304,7 @@ namespace UnrealBuildTool
 			// convert r<num>[letter] to hex
 			if (!String.IsNullOrEmpty(StringValue))
 			{
-				if (String.Compare(Hint, "Software", true) == 0)
+				if (String.Compare(Hint, "Software", StringComparison.OrdinalIgnoreCase) == 0)
 				{
 					// returning a version with a format similar to the ndk version.
 					string VersionString = String.Format("{0}{1:00}{2:00}", StringValue, 0, 0);
diff --git a/Engine/Source/Programs/UnrealBuildTool/Platform/Android/AndroidToolChain.cs b/Engine/Source/Programs/UnrealBuildTool/Platform/Android/AndroidToolChain.cs
--- a/Engine/Source/Programs/UnrealBuildTool/Platform/Android/AndroidToolChain.cs
+++ b/Engine/Source/Programs/UnrealBuildTool/Platform/Android/AndroidToolChain.cs
@@ -705,7 +705,7 @@ namespace UnrealBuildTool
 			string Extension = Path.GetExtension(SourceFile.AbsolutePath).ToUpperInvariant();
 			if (Extension == ".C")
 			{
-				if (SourceFile.AbsolutePath.Equals(GetNativeGluePath()))
+				if (SourceFile.AbsolutePath.Equals(GetNativeGluePath(), StringComparison.Ordinal))
 				{
 					// Remove visibility settings for android native glue. Since it doesn't decorate with visibility attributes.
 					Arguments.RemoveAll(x => x.StartsWith("-fvisibility"));
@@ -1033,7 +1033,7 @@ namespace UnrealBuildTool
 		void GenerateEmptyLinkFunctionsForRemovedModules(List<FileItem> SourceFiles, UnrealArch Arch, string ModuleName, DirectoryReference OutputDirectory, IActionGraphBuilder Graph, ILogger Logger)
 		{
 			// Only add to Launch module
-			if (!ModuleName.Equals("Launch"))
+			if (!ModuleName.Equals("Launch", StringComparison.Ordinal))
 			{
 				return;
 			}
@@ -1113,7 +1113,7 @@ namespace UnrealBuildTool
 			List<FileItem> ModifiedInputFiles = new(InputFiles);
 
 			// Deal with Launch module special if first time seen
-			if (!HasHandledLaunchModule.Contains(CompileEnvironment.Architecture) && (ModuleName.Equals("Launch") || ModuleName.Equals("AndroidLauncher")))
+			if (!HasHandledLaunchModule.Contains(CompileEnvironment.Architecture) && (ModuleName.Equals("Launch", StringComparison.Ordinal) || ModuleName.Equals("AndroidLauncher", StringComparison.Ordinal)))
 			{
 				// Directly added NDK files for NDK extensions
 				ModifiedInputFiles.Add(FileItem.GetItemByPath(GetNativeGluePath()));
@@ -1123,7 +1123,7 @@ namespace UnrealBuildTool
 				HasHandledLaunchModule.Add(CompileEnvironment.Architecture);
 			}
 
-			if (!HasHandledCoreModule.Contains(CompileEnvironment.Architecture) && ModuleName.Equals("Core") && (CompileEnvironment.PrecompiledHeaderAction == PrecompiledHeaderAction.None))
+			if (!HasHandledCoreModule.Contains(CompileEnvironment.Architecture) && ModuleName.Equals("Core", StringComparison.Ordinal) && (CompileEnvironment.PrecompiledHeaderAction == PrecompiledHeaderAction.None))
 			{
 				// This is used by Crypto code in Core
 				ModifiedInputFiles.Add(FileItem.GetItemByPath(GetCpuFeaturesPath()));
diff --git a/Engine/Source/Programs/UnrealBuildTool/Platform/Android/UEBuildAndroid.cs b/Engine/Source/Programs/UnrealBuildTool/Platform/Android/UEBuildAndroid.cs
--- a/Engine/Source/Programs/UnrealBuildTool/Platform/Android/UEBuildAndroid.cs
+++ b/Engine/Source/Programs/UnrealBuildTool/Platform/Android/UEBuildAndroid.cs
@@ -259,7 +259,7 @@ namespace UnrealBuildTool
 				foreach (string CpuSuffix in AndroidToolChain.AllCpuSuffixes.Values)
 				{
 					int CpuIdx = ExtensionEndIdx - CpuSuffix.Length;
-					if (CpuIdx > 0 && String.Compare(Name, CpuIdx, CpuSuffix, 0, CpuSuffix.Length, StringComparison.InvariantCultureIgnoreCase) == 0)
+					if (CpuIdx > 0 && String.Compare(Name, CpuIdx, CpuSuffix, 0, CpuSuffix.Length, StringComparison.OrdinalIgnoreCase) == 0)
 					{
 						return IsBuildProductName(Name, 0, CpuIdx, NamePrefixes, NameSuffixes);
 					}
diff --git a/Engine/Source/Programs/UnrealBuildTool/Platform/Android/UEDeployAndroid.cs b/Engine/Source/Programs/UnrealBuildTool/Platform/Android/UEDeployAndroid.cs
--- a/Engine/Source/Programs/UnrealBuildTool/Platform/Android/UEDeployAndroid.cs
+++ b/Engine/Source/Programs/UnrealBuildTool/Platform/Android/UEDeployAndroid.cs
@@ -623,7 +623,7 @@ namespace UnrealBuildTool
 		public static void DeleteDirectory(string InPath, ILogger Logger, string SubDirectoryToKeep = "")
 		{
 			// skip the dir we want to
-			if (String.Equals(Path.GetFileName(InPath), SubDirectoryToKeep, StringComparison.CurrentCultureIgnoreCase))
+			if (String.Equals(Path.GetFileName(InPath), SubDirectoryToKeep, StringComparison.OrdinalIgnoreCase))
 			{
 				return;
 			}
@@ -2037,18 +2037,18 @@ namespace UnrealBuildTool
 
 					if (Parts.Count > 1)
 					{
-						if (Parts[0].Equals("S"))
+						if (Parts[0].Equals("S", StringComparison.Ordinal))
 						{
 							ActiveStdOutFilter.Add(new FilterOperation(FilterAction.Skip, Parts[1], "", ""));
 						}
-						else if (Parts[0].Equals("R"))
+						else if (Parts[0].Equals("R", StringComparison.Ordinal))
 						{
 							if (Parts.Count == 4)
 							{
 								ActiveStdOutFilter.Add(new FilterOperation(FilterAction.Replace, Parts[1], Parts[2], Parts[3]));
 							}
 						}
-						else if (Parts[0].Equals("E"))
+						else if (Parts[0].Equals("E", StringComparison.Ordinal))
 						{
 							if (Parts.Count == 4)
 							{
@@ -2239,7 +2239,7 @@ namespace UnrealBuildTool
 				foreach (string Key in Section.KeyNames)
 				{
 					// filter out NDK and SDK override since actual resolved versions already written above
-					if (Key.Equals("SDKAPILevelOverride") || Key.Equals("NDKAPILevelOverride"))
+					if (Key.Equals("SDKAPILevelOverride", StringComparison.Ordinal) || Key.Equals("NDKAPILevelOverride", StringComparison.Ordinal))
 					{
 						continue;
 					}
@@ -2260,7 +2260,7 @@ namespace UnrealBuildTool
 				foreach (string Key in Section.KeyNames)
 				{
 					// filter out NDK and SDK levels since actual resolved versions already written above
-					if (Key.Equals("SDKAPILevel") || Key.Equals("NDKAPILevel"))
+					if (Key.Equals("SDKAPILevel", StringComparison.Ordinal) || Key.Equals("NDKAPILevel", StringComparison.Ordinal))
 					{
 						continue;
 					}
@@ -2289,8 +2289,8 @@ namespace UnrealBuildTool
 			{
 				foreach (string Key in Section.KeyNames)
 				{
-					if (!Key.Equals("GameDefaultMap") &&
-						!Key.Equals("GlobalDefaultGameMode"))
+					if (!Key.Equals("GameDefaultMap", StringComparison.Ordinal) &&
+						!Key.Equals("GlobalDefaultGameMode", StringComparison.Ordinal))
 					{
 						continue;
 					}
@@ -3903,7 +3903,7 @@ namespace UnrealBuildTool
 				SourceReader.BaseStream.Read(Extra, 0, ExtraLen);
 
 				string FilenameStr = Encoding.UTF8.GetString(Filename, 0, FilenameLen);
-				if (!FilenameStr.Equals(DestSOFile))
+				if (!FilenameStr.Equals(DestSOFile, StringComparison.Ordinal))
 				{
 					TotalEntries++;
 					uint Location = (uint)DestWriter.BaseStream.Position;
@@ -4603,7 +4603,7 @@ dependencies.implementation('androidx.tracing:tracing:1.2.0')
 				{
 					const string Key = "bEnableBundle=";
 					string? Line = File.ReadAllLines(BuildSettingsCacheFile).FirstOrDefault(Line => Line.StartsWith(Key));
-					if (Line != null && !Line[Key.Length..].Equals(bEnableBundle ? "True" : "False"))
+					if (Line != null && !Line[Key.Length..].Equals(bEnableBundle ? "True" : "False", StringComparison.Ordinal))
 					{
 						bCreateFromScratch = true;
 					}
diff --git a/Engine/Source/Programs/UnrealBuildTool/Platform/IOS/IOSCodeSigning.cs b/Engine/Source/Programs/UnrealBuildTool/Platform/IOS/IOSCodeSigning.cs
--- a/Engine/Source/Programs/UnrealBuildTool/Platform/IOS/IOSCodeSigning.cs
+++ b/Engine/Source/Programs/UnrealBuildTool/Platform/IOS/IOSCodeSigning.cs
@@ -575,7 +575,7 @@ namespace UnrealBuildTool
 			string[] FileList = Directory.GetFiles(BasePath);
 			foreach (string Filename in FileList)
 			{
-				if (Path.GetExtension(Filename).Equals(FileExt, StringComparison.InvariantCultureIgnoreCase))
+				if (Path.GetExtension(Filename).Equals(FileExt, StringComparison.OrdinalIgnoreCase))
 				{
 					return Filename;
 				}
@@ -959,7 +959,7 @@ namespace UnrealBuildTool
 
 					// make sure the file is not managed by Xcode
 					if (TestProvision.FileName != null && TestProvision.UUID != null &&
-					Path.GetFileName(TestProvision.FileName).ToLower().Equals(TestProvision.UUID.ToLower() + ".mobileprovision"))
+					Path.GetFileName(TestProvision.FileName).Equals($"{TestProvision.UUID}.mobileprovision", StringComparison.OrdinalIgnoreCase))
 					{
 						continue;
 					}
@@ -1371,7 +1371,7 @@ namespace UnrealBuildTool
 			{
 				foreach (string TestUDID in ProvisionedDeviceIDs)
 				{
-					if (TestUDID.Equals(UDID, StringComparison.InvariantCultureIgnoreCase))
+					if (TestUDID.Equals(UDID, StringComparison.OrdinalIgnoreCase))
 					{
 						bFound = true;
 						break;
diff --git a/Engine/Source/Programs/UnrealBuildTool/Platform/IOS/IOSToolChain.cs b/Engine/Source/Programs/UnrealBuildTool/Platform/IOS/IOSToolChain.cs
--- a/Engine/Source/Programs/UnrealBuildTool/Platform/IOS/IOSToolChain.cs
+++ b/Engine/Source/Programs/UnrealBuildTool/Platform/IOS/IOSToolChain.cs
@@ -229,7 +229,7 @@ namespace UnrealBuildTool
 
 			string Extension = Path.GetExtension(SourceFile.AbsolutePath).ToUpperInvariant();
 
-			if (!Extension.Equals(".C") && !Extension.Equals(".SWIFT"))
+			if (!Extension.Equals(".C", StringComparison.Ordinal) && !Extension.Equals(".SWIFT", StringComparison.Ordinal))
 			{
 				Arguments.Add(GetObjCExceptionsFlag(CompileEnvironment));
 			}
@@ -794,7 +794,7 @@ namespace UnrealBuildTool
 
 					// Check to see if this 
 					bool bIsExecutable = false;
-					if (Path.GetFileNameWithoutExtension(E.FileName).Equals(GameName, StringComparison.InvariantCultureIgnoreCase))
+					if (Path.GetFileNameWithoutExtension(E.FileName).Equals(GameName, StringComparison.OrdinalIgnoreCase))
 					{
 						bIsExecutable = true;
 					}
diff --git a/Engine/Source/Programs/UnrealBuildTool/Platform/IOS/MobileProvisionContents.cs b/Engine/Source/Programs/UnrealBuildTool/Platform/IOS/MobileProvisionContents.cs
--- a/Engine/Source/Programs/UnrealBuildTool/Platform/IOS/MobileProvisionContents.cs
+++ b/Engine/Source/Programs/UnrealBuildTool/Platform/IOS/MobileProvisionContents.cs
@@ -94,7 +94,7 @@ namespace UnrealBuildTool
 					}
 
 					Console.WriteLine("Found entitlement node:" + KeyElement.InnerText);
-					if (!KeyElement.InnerText.Equals("application-identifier"))
+					if (!KeyElement.InnerText.Equals("application-identifier", StringComparison.Ordinal))
 					{
 						continue;
 					}
@@ -200,7 +200,7 @@ namespace UnrealBuildTool
 			for (int i = 0; i < elemList.Count; i++)
 			{
 				XmlNode? Node = elemList[i];
-				if (Node != null && Node.InnerXml.Equals(InValue))
+				if (Node != null && Node.InnerXml.Equals(InValue, StringComparison.Ordinal))
 				{
 					XmlNode? valueNode = Node.NextSibling;
 
@@ -220,12 +220,12 @@ namespace UnrealBuildTool
 			for (int i = 0; i < elemList.Count; i++)
 			{
 				XmlNode? Node = elemList[i];
-				if (Node != null && Node.InnerXml.Equals(InValue))
+				if (Node != null && Node.InnerXml.Equals(InValue, StringComparison.Ordinal))
 				{
 					XmlNode? valueNode = Node.NextSibling;
 					if (valueNode != null)
 					{
-						if (valueNode.Name.Equals("array"))
+						if (valueNode.Name.Equals("array", StringComparison.Ordinal))
 						{
 							XmlNode? firstChildNode = valueNode.FirstChild;
 							if (firstChildNode != null)
diff --git a/Engine/Source/Programs/UnrealBuildTool/Platform/Mac/MacToolChain.cs b/Engine/Source/Programs/UnrealBuildTool/Platform/Mac/MacToolChain.cs
--- a/Engine/Source/Programs/UnrealBuildTool/Platform/Mac/MacToolChain.cs
+++ b/Engine/Source/Programs/UnrealBuildTool/Platform/Mac/MacToolChain.cs
@@ -89,9 +89,9 @@ namespace UnrealBuildTool
 		{
 			DirectoryReference? ProgramFinder = StartingDir;
 			while (ProgramFinder != null &&
-!String.Equals(ProgramFinder.GetDirectoryName(), "Source", StringComparison.CurrentCultureIgnoreCase) &&
-!String.Equals(ProgramFinder.GetDirectoryName(), "Intermediate", StringComparison.CurrentCultureIgnoreCase) &&
-!String.Equals(ProgramFinder.GetDirectoryName(), "Binaries", StringComparison.CurrentCultureIgnoreCase))
+!String.Equals(ProgramFinder.GetDirectoryName(), "Source", StringComparison.OrdinalIgnoreCase) &&
+!String.Equals(ProgramFinder.GetDirectoryName(), "Intermediate", StringComparison.OrdinalIgnoreCase) &&
+!String.Equals(ProgramFinder.GetDirectoryName(), "Binaries", StringComparison.OrdinalIgnoreCase))
 			{
 				ProgramFinder = ProgramFinder.ParentDirectory;
 			}
@@ -763,7 +763,7 @@ namespace UnrealBuildTool
 					HashSet<FileReference> UsedSourceLibs = new();
 					void ProcessLib(FileReference SourceLib, FileReference DestLib, FileReference ReferencingFile)
 					{
-						if (string.Equals(SourceLib.GetExtension(), ".dylib", StringComparison.InvariantCultureIgnoreCase))
+						if (string.Equals(SourceLib.GetExtension(), ".dylib", StringComparison.OrdinalIgnoreCase))
 						{
 							DirectoryReference DestDir = DestLib.Directory;
 							if (UsedSourceLibs.Contains(SourceLib) || UsedDestDirs.Contains(DestDir))
@@ -1243,7 +1243,7 @@ namespace UnrealBuildTool
 			List<string> OutputFiles = Makefile.OutputItems.Select(Item => Path.ChangeExtension(Item.FullName, ".dSYM")).Distinct().ToList();
 			foreach (FileItem DebugItem in DebugInfoFiles)
 			{
-				if (OutputFiles.Any(Item => String.Equals(Item, DebugItem.FullName, StringComparison.InvariantCultureIgnoreCase)))
+				if (OutputFiles.Any(Item => String.Equals(Item, DebugItem.FullName, StringComparison.OrdinalIgnoreCase)))
 				{
 					Makefile.OutputItems.Add(DebugItem);
 				}
diff --git a/Engine/Source/Programs/UnrealBuildTool/Platform/Mac/UEBuildMac.cs b/Engine/Source/Programs/UnrealBuildTool/Platform/Mac/UEBuildMac.cs
--- a/Engine/Source/Programs/UnrealBuildTool/Platform/Mac/UEBuildMac.cs
+++ b/Engine/Source/Programs/UnrealBuildTool/Platform/Mac/UEBuildMac.cs
@@ -210,7 +210,7 @@ namespace UnrealBuildTool
 			List<UnrealArch> Architectures = new();
 
 			// return all supported if getting supported, compiling for distribution, or we want active, and "all" is selected
-			if (bGetAllSupported || bIsDistributionMode || DefaultArchitecture.Equals("all", StringComparison.InvariantCultureIgnoreCase) ||
+			if (bGetAllSupported || bIsDistributionMode || DefaultArchitecture.Equals("all", StringComparison.OrdinalIgnoreCase) ||
 				(bIsBuildMachine && bBuildAllSupportedOnBuildMachine))
 			{
 				if (bSupportsArm64)
diff --git a/Engine/Source/Programs/UnrealBuildTool/Platform/Windows/MicrosoftPlatformSDK.cs b/Engine/Source/Programs/UnrealBuildTool/Platform/Windows/MicrosoftPlatformSDK.cs
--- a/Engine/Source/Programs/UnrealBuildTool/Platform/Windows/MicrosoftPlatformSDK.cs
+++ b/Engine/Source/Programs/UnrealBuildTool/Platform/Windows/MicrosoftPlatformSDK.cs
@@ -195,7 +195,7 @@ namespace UnrealBuildTool
 			VersionNumber? WindowsSdkVersion = null;
 			if (!String.IsNullOrEmpty(DesiredVersion))
 			{
-				if (String.Equals(DesiredVersion, "Latest", StringComparison.InvariantCultureIgnoreCase) && CachedWindowsSdkDirs!.Count > 0)
+				if (String.Equals(DesiredVersion, "Latest", StringComparison.OrdinalIgnoreCase) && CachedWindowsSdkDirs!.Count > 0)
 				{
 					WindowsSdkVersion = CachedWindowsSdkDirs.OrderBy(x => x.Key).Last().Key;
 				}
diff --git a/Engine/Source/Programs/UnrealBuildTool/Platform/Windows/PVSToolChain.cs b/Engine/Source/Programs/UnrealBuildTool/Platform/Windows/PVSToolChain.cs
--- a/Engine/Source/Programs/UnrealBuildTool/Platform/Windows/PVSToolChain.cs
+++ b/Engine/Source/Programs/UnrealBuildTool/Platform/Windows/PVSToolChain.cs
@@ -723,7 +723,7 @@ namespace UnrealBuildTool
 												Math.Max(compilerCommandLine.LastIndexOf("-TP "), compilerCommandLine.LastIndexOf("-Tp ")));
 
 			bool compileAsCCode = cFlagLastPosition == cppFlagLastPosition
-				? Path.GetExtension(sourceFileName).Equals(".c", StringComparison.InvariantCultureIgnoreCase)
+				? Path.GetExtension(sourceFileName).Equals(".c", StringComparison.OrdinalIgnoreCase)
 				: cFlagLastPosition > cppFlagLastPosition;
 
 			return compileAsCCode;
@@ -801,7 +801,7 @@ namespace UnrealBuildTool
 
 				configFileContents.AppendFormat("std={0}\n", languageStandardForCfg);
 
-				bool disableMsExtensionsFromArgs = preprocessAction.Arguments.Any(arg => arg.Equals("/Za") || arg.Equals("-Za") || arg.Equals("/permissive-"));
+				bool disableMsExtensionsFromArgs = preprocessAction.Arguments.Any(arg => arg.Equals("/Za", StringComparison.Ordinal) || arg.Equals("-Za", StringComparison.Ordinal) || arg.Equals("/permissive-", StringComparison.Ordinal));
 				bool disableMsExtensions = isVisualCppCompiler && (languageStandardForCfg == CPP_20 || disableMsExtensionsFromArgs);
 				configFileContents.AppendFormat("disable-ms-extensions={0}\n", disableMsExtensions ? "yes" : "no");
 			}
diff --git a/Engine/Source/Programs/UnrealBuildTool/ProjectFiles/CMake/CMakefileGenerator.cs b/Engine/Source/Programs/UnrealBuildTool/ProjectFiles/CMake/CMakefileGenerator.cs
--- a/Engine/Source/Programs/UnrealBuildTool/ProjectFiles/CMake/CMakefileGenerator.cs
+++ b/Engine/Source/Programs/UnrealBuildTool/ProjectFiles/CMake/CMakefileGenerator.cs
@@ -283,7 +283,7 @@ namespace UnrealBuildTool
 					string Definition = PreProcessorDefinition.Replace("TEXT(\"", "").Replace("\")", "").Replace("()=", "=");
 					string AlternateDefinition = Definition.Contains("=0") ? Definition.Replace("=0", "=1") : Definition.Replace("=1", "=0");
 
-					if (Definition.Equals("WITH_EDITORONLY_DATA=0"))
+					if (Definition.Equals("WITH_EDITORONLY_DATA=0", StringComparison.Ordinal))
 					{
 						Definition = AlternateDefinition;
 					}
diff --git a/Engine/Source/Programs/UnrealBuildTool/ProjectFiles/CodeLite/CodeLiteGenerator.cs b/Engine/Source/Programs/UnrealBuildTool/ProjectFiles/CodeLite/CodeLiteGenerator.cs
--- a/Engine/Source/Programs/UnrealBuildTool/ProjectFiles/CodeLite/CodeLiteGenerator.cs
+++ b/Engine/Source/Programs/UnrealBuildTool/ProjectFiles/CodeLite/CodeLiteGenerator.cs
@@ -171,10 +171,10 @@ namespace UnrealBuildTool
 							(CurrentTarget.TargetRules.Type == TargetType.Editor) ||
 							(CurrentTarget.TargetRules.Type == TargetType.Game))
 						{
-							if (ProjectName.Equals("UnrealClient") ||
-								ProjectName.Equals("UnrealServer") ||
-								ProjectName.Equals("UnrealGame") ||
-								ProjectName.Equals("UnrealEditor"))
+							if (ProjectName.Equals("UnrealClient", StringComparison.Ordinal) ||
+								ProjectName.Equals("UnrealServer", StringComparison.Ordinal) ||
+								ProjectName.Equals("UnrealGame", StringComparison.Ordinal) ||
+								ProjectName.Equals("UnrealEditor", StringComparison.Ordinal))
 							{
 								if (CodeLiteWorkspaceTargetEngine == null)
 								{
diff --git a/Engine/Source/Programs/UnrealBuildTool/ProjectFiles/CodeLite/CodeLiteProject.cs b/Engine/Source/Programs/UnrealBuildTool/ProjectFiles/CodeLite/CodeLiteProject.cs
--- a/Engine/Source/Programs/UnrealBuildTool/ProjectFiles/CodeLite/CodeLiteProject.cs
+++ b/Engine/Source/Programs/UnrealBuildTool/ProjectFiles/CodeLite/CodeLiteProject.cs
@@ -125,10 +125,10 @@ namespace UnrealBuildTool
 						(ProjectTargetType == TargetType.Game) ||
 						(ProjectTargetType == TargetType.Server))
 					{
-						if (TargetName.Equals("UnrealClient") ||
-							TargetName.Equals("UnrealServer") ||
-							TargetName.Equals("UnrealGame") ||
-							TargetName.Equals("UnrealEditor"))
+						if (TargetName.Equals("UnrealClient", StringComparison.Ordinal) ||
+							TargetName.Equals("UnrealServer", StringComparison.Ordinal) ||
+							TargetName.Equals("UnrealGame", StringComparison.Ordinal) ||
+							TargetName.Equals("UnrealEditor", StringComparison.Ordinal))
 						{
 							int Idx = Unreal.EngineDirectory.FullName.Length;
 							CurrentFilePath = Path.GetDirectoryName(Path.GetFullPath(CurrentFile.Reference.FullName))!.Substring(Idx);
@@ -341,7 +341,7 @@ namespace UnrealBuildTool
 							XAttribute GeneralExecutableToRunArguments = new XAttribute("CommandArguments", commandArguments);
 							CodeLiteConfigurationGeneral.Add(GeneralExecutableToRunArguments);
 						}
-						if (TargetName.Equals("UnrealGame"))
+						if (TargetName.Equals("UnrealGame", StringComparison.Ordinal))
 						{
 							XAttribute GeneralExecutableWorkingDirectory = new XAttribute("WorkingDirectory", UnrealEngineEditorWorkingDirectory);
 							CodeLiteConfigurationGeneral.Add(GeneralExecutableWorkingDirectory);
diff --git a/Engine/Source/Programs/UnrealBuildTool/ProjectFiles/ProjectFileGenerator.cs b/Engine/Source/Programs/UnrealBuildTool/ProjectFiles/ProjectFileGenerator.cs
--- a/Engine/Source/Programs/UnrealBuildTool/ProjectFiles/ProjectFileGenerator.cs
+++ b/Engine/Source/Programs/UnrealBuildTool/ProjectFiles/ProjectFileGenerator.cs
@@ -69,7 +69,7 @@ namespace UnrealBuildTool
 			bool AlreadyExists = false;
 			foreach (PrimaryProjectFolder ExistingFolder in SubFolders)
 			{
-				if (ExistingFolder.FolderName.Equals(FirstFolderName, StringComparison.InvariantCultureIgnoreCase))
+				if (ExistingFolder.FolderName.Equals(FirstFolderName, StringComparison.OrdinalIgnoreCase))
 				{
 					// Already exists!
 					ResultFolder = ExistingFolder;
@@ -1079,7 +1079,7 @@ namespace UnrealBuildTool
 			FileReference PrimaryProjectPathLocation = FileReference.Combine(Unreal.EngineDirectory, "Intermediate", "ProjectFiles", "PrimaryProjectPath.txt");
 			_fileHashCache.WriteFileIfChanged(FileItem.GetItemByFileReference(PrimaryProjectPathLocation), DirectoryReference.Combine(PrimaryProjectPath, PrimaryProjectName).FullName.Replace("\\", "/", StringComparison.InvariantCultureIgnoreCase), Logger, false);
 
-			bool bCleanProjectFiles = Arguments.Any(x => x.Equals("-CleanProjects", StringComparison.InvariantCultureIgnoreCase));
+			bool bCleanProjectFiles = Arguments.Any(x => x.Equals("-CleanProjects", StringComparison.OrdinalIgnoreCase));
 			if (bCleanProjectFiles)
 			{
 				using ITimelineEvent CleanTimer = Timeline.ScopeEvent("CleanProjectFiles");
@@ -2111,7 +2111,7 @@ namespace UnrealBuildTool
 						SubFolder.SubFolders.Count == 0)
 					{
 						// 3)
-						if (SubFolder.FolderName.Equals(SubFolder.ChildProjects[0].ProjectFilePath.GetFileNameWithoutExtension(), StringComparison.InvariantCultureIgnoreCase))
+						if (SubFolder.FolderName.Equals(SubFolder.ChildProjects[0].ProjectFilePath.GetFileNameWithoutExtension(), StringComparison.OrdinalIgnoreCase))
 						{
 							CanCollapseFolder = true;
 						}
@@ -2120,11 +2120,11 @@ namespace UnrealBuildTool
 
 				if (!bKeepSourceSubDirectories)
 				{
-					if (SubFolder.FolderName.Equals("Source", StringComparison.InvariantCultureIgnoreCase))
+					if (SubFolder.FolderName.Equals("Source", StringComparison.OrdinalIgnoreCase))
 					{
 						// Avoid collapsing the Engine's Source directory, since there are so many other solution folders in
 						// the parent directory.
-						if (!Folder.FolderName.Equals("Engine", StringComparison.InvariantCultureIgnoreCase))
+						if (!Folder.FolderName.Equals("Engine", StringComparison.OrdinalIgnoreCase))
 						{
 							CanCollapseFolder = true;
 						}
@@ -2184,7 +2184,7 @@ namespace UnrealBuildTool
 				{
 					if (CurChildProject != OtherChildProject)
 					{
-						if (CurChildProject.ProjectFilePath.GetFileName().Equals(OtherChildProject.ProjectFilePath.GetFileNameWithoutExtension(), StringComparison.InvariantCultureIgnoreCase))
+						if (CurChildProject.ProjectFilePath.GetFileName().Equals(OtherChildProject.ProjectFilePath.GetFileNameWithoutExtension(), StringComparison.OrdinalIgnoreCase))
 						{
 							throw new BuildException("Detected collision between two project files with the same path in the same primary project folder, " + OtherChildProject.ProjectFilePath.FullName + " and " + CurChildProject.ProjectFilePath.FullName + " (primary project folder: " + PrimaryProjectFolderPath + ")");
 						}
@@ -2199,14 +2199,14 @@ namespace UnrealBuildTool
 				// with project file names or file items, as that's not supported in Visual Studio.)
 				foreach (ProjectFile CurChildProject in Folder.ChildProjects)
 				{
-					if (CurChildProject.ProjectFilePath.GetFileNameWithoutExtension().Equals(SubFolder.FolderName, StringComparison.InvariantCultureIgnoreCase))
+					if (CurChildProject.ProjectFilePath.GetFileNameWithoutExtension().Equals(SubFolder.FolderName, StringComparison.OrdinalIgnoreCase))
 					{
 						throw new BuildException("Detected collision between a primary project sub-folder " + SubFolder.FolderName + " and a project within the outer folder " + CurChildProject.ProjectFilePath + " (primary project folder: " + PrimaryProjectFolderPath + ")");
 					}
 				}
 				foreach (string CurFile in Folder.Files)
 				{
-					if (Path.GetFileName(CurFile).Equals(SubFolder.FolderName, StringComparison.InvariantCultureIgnoreCase))
+					if (Path.GetFileName(CurFile).Equals(SubFolder.FolderName, StringComparison.OrdinalIgnoreCase))
 					{
 						throw new BuildException("Detected collision between a primary project sub-folder " + SubFolder.FolderName + " and a file within the outer folder " + CurFile + " (primary project folder: " + PrimaryProjectFolderPath + ")");
 					}
@@ -2215,7 +2215,7 @@ namespace UnrealBuildTool
 				{
 					if (CurFolder != SubFolder)
 					{
-						if (CurFolder.FolderName.Equals(SubFolder.FolderName, StringComparison.InvariantCultureIgnoreCase))
+						if (CurFolder.FolderName.Equals(SubFolder.FolderName, StringComparison.OrdinalIgnoreCase))
 						{
 							throw new BuildException("Detected collision between a primary project sub-folder " + SubFolder.FolderName + " and a sibling folder " + CurFolder.FolderName + " (primary project folder: " + PrimaryProjectFolderPath + ")");
 						}
@@ -3129,7 +3129,7 @@ namespace UnrealBuildTool
 					if (CheckProjectFile == null && TargetFilePath.ContainsName("Programs", 0))
 					{
 						string BaseName = TargetFilePath.GetFileNameWithoutAnyExtensions();
-						CheckProjectFile = AllGames.FirstOrDefault(x => x.GetFileNameWithoutAnyExtensions().Equals(BaseName, StringComparison.InvariantCultureIgnoreCase));
+						CheckProjectFile = AllGames.FirstOrDefault(x => x.GetFileNameWithoutAnyExtensions().Equals(BaseName, StringComparison.OrdinalIgnoreCase));
 					}
 
 					if (CheckProjectFile == null)
@@ -3543,7 +3543,7 @@ namespace UnrealBuildTool
 						{
 							string LanguageSuffix = Path.GetExtension(Path.GetFileNameWithoutExtension(DocumentationFile.FullName));
 							if (!String.IsNullOrEmpty(LanguageSuffix) &&
-								!LanguageSuffix.Equals(".int", StringComparison.InvariantCultureIgnoreCase))
+								!LanguageSuffix.Equals(".int", StringComparison.OrdinalIgnoreCase))
 							{
 								bPassesFilter = false;
 							}
@@ -4119,7 +4119,7 @@ namespace UnrealBuildTool
 			//StrB = StripGUIDs(StrB);
 
 			// Use simple ordinal comparison.
-			return String.Compare(StrA, StrB, StringComparison.InvariantCultureIgnoreCase);
+			return String.Compare(StrA, StrB, StringComparison.OrdinalIgnoreCase);
 		}
 
 		/// <summary>
diff --git a/Engine/Source/Programs/UnrealBuildTool/ProjectFiles/VisualStudio/VCProjectFileGenerator.cs b/Engine/Source/Programs/UnrealBuildTool/ProjectFiles/VisualStudio/VCProjectFileGenerator.cs
--- a/Engine/Source/Programs/UnrealBuildTool/ProjectFiles/VisualStudio/VCProjectFileGenerator.cs
+++ b/Engine/Source/Programs/UnrealBuildTool/ProjectFiles/VisualStudio/VCProjectFileGenerator.cs
@@ -154,7 +154,7 @@ namespace UnrealBuildTool
 
 		public override string[] GetTargetArguments(string[] Arguments)
 		{
-			return Arguments.Where(s => String.Equals(s, BuildToolOverride, StringComparison.InvariantCultureIgnoreCase)).ToArray();
+			return Arguments.Where(s => String.Equals(s, BuildToolOverride, StringComparison.OrdinalIgnoreCase)).ToArray();
 		}
 
 		/// File extension for project files we'll be generating (e.g. ".vcxproj")
@@ -1162,7 +1162,7 @@ namespace UnrealBuildTool
 					(x, y) =>
 					{
 						return String.Compare(x.VCSolutionConfigAndPlatformName, y.VCSolutionConfigAndPlatformName,
-							StringComparison.InvariantCultureIgnoreCase);
+							StringComparison.OrdinalIgnoreCase);
 					}
 				)
 			);
diff --git a/Engine/Source/Programs/UnrealBuildTool/ProjectFiles/VisualStudioCode/VSCodeProjectFileGenerator.cs b/Engine/Source/Programs/UnrealBuildTool/ProjectFiles/VisualStudioCode/VSCodeProjectFileGenerator.cs
--- a/Engine/Source/Programs/UnrealBuildTool/ProjectFiles/VisualStudioCode/VSCodeProjectFileGenerator.cs
+++ b/Engine/Source/Programs/UnrealBuildTool/ProjectFiles/VisualStudioCode/VSCodeProjectFileGenerator.cs
@@ -1000,7 +1000,7 @@ namespace UnrealBuildTool
 						}
 						OutFile.EndObject();
 
-						if (BuildProduct.Platform == UnrealTargetPlatform.Android && BaseCommand.Equals("Build"))
+						if (BuildProduct.Platform == UnrealTargetPlatform.Android && BaseCommand.Equals("Build", StringComparison.Ordinal))
 						{
 							WriteNativeTaskDeployAndroid(InProject, OutFile, Target, BuildProduct);
 						}
diff --git a/Engine/Source/Programs/UnrealBuildTool/ProjectFiles/Xcode/XcodeProject.cs b/Engine/Source/Programs/UnrealBuildTool/ProjectFiles/Xcode/XcodeProject.cs
--- a/Engine/Source/Programs/UnrealBuildTool/ProjectFiles/Xcode/XcodeProject.cs
+++ b/Engine/Source/Programs/UnrealBuildTool/ProjectFiles/Xcode/XcodeProject.cs
@@ -501,7 +501,7 @@ namespace UnrealBuildTool.XcodeProjectXcconfig
 				}
 
 				DirectoryReference? ProgramFinder = DirectoryReference.Combine(ProjectFile.BaseDir);
-				while (ProgramFinder != null && !String.Equals(ProgramFinder.GetDirectoryName(), "Source", StringComparison.CurrentCultureIgnoreCase))
+				while (ProgramFinder != null && !String.Equals(ProgramFinder.GetDirectoryName(), "Source", StringComparison.OrdinalIgnoreCase))
 				{
 					ProgramFinder = ProgramFinder.ParentDirectory;
 				}
@@ -3208,7 +3208,7 @@ namespace UnrealBuildTool.XcodeProjectXcconfig
 		public override bool WriteProjectFile(List<UnrealTargetPlatform> InPlatforms, List<UnrealTargetConfiguration> InConfigurations, PlatformProjectGeneratorCollection PlatformProjectGenerators, ILogger Logger)
 		{
 			// if we don't want this one, just skip
-			if (SingleTargetName != null && !ProjectFilePath.GetFileNameWithoutAnyExtensions().Equals(SingleTargetName, StringComparison.InvariantCultureIgnoreCase))
+			if (SingleTargetName != null && !ProjectFilePath.GetFileNameWithoutAnyExtensions().Equals(SingleTargetName, StringComparison.OrdinalIgnoreCase))
 			{
 				return true;
 			}
diff --git a/Engine/Source/Programs/UnrealBuildTool/ProjectFiles/Xcode/XcodeProjectFileGenerator.cs b/Engine/Source/Programs/UnrealBuildTool/ProjectFiles/Xcode/XcodeProjectFileGenerator.cs
--- a/Engine/Source/Programs/UnrealBuildTool/ProjectFiles/Xcode/XcodeProjectFileGenerator.cs
+++ b/Engine/Source/Programs/UnrealBuildTool/ProjectFiles/Xcode/XcodeProjectFileGenerator.cs
@@ -536,7 +536,7 @@ namespace UnrealBuildTool
 			Parallel.ForEach(Targets, Options, TargetPair =>
 			{
 				// don't bother if we aren't interested in this target
-				if (SingleTargetName != null && !TargetPair.Item2.Name.Equals(SingleTargetName, StringComparison.InvariantCultureIgnoreCase))
+				if (SingleTargetName != null && !TargetPair.Item2.Name.Equals(SingleTargetName, StringComparison.OrdinalIgnoreCase))
 				{
 					return;
 				}
diff --git a/Engine/Source/Programs/UnrealBuildTool/ProjectFiles/Xcode/XcodeProjectLegacy.cs b/Engine/Source/Programs/UnrealBuildTool/ProjectFiles/Xcode/XcodeProjectLegacy.cs
--- a/Engine/Source/Programs/UnrealBuildTool/ProjectFiles/Xcode/XcodeProjectLegacy.cs
+++ b/Engine/Source/Programs/UnrealBuildTool/ProjectFiles/Xcode/XcodeProjectLegacy.cs
@@ -977,8 +977,8 @@ namespace UnrealBuildTool.XcodeProjectLegacy
 			IOSRunTimeVersion = null;
 			TVOSRunTimeVersion = null;
 
-			bool bIsUnrealGame = TargetName.Equals("UnrealGame", StringComparison.InvariantCultureIgnoreCase);
-			bool bIsUnrealClient = TargetName.Equals("UnrealClient", StringComparison.InvariantCultureIgnoreCase);
+			bool bIsUnrealGame = TargetName.Equals("UnrealGame", StringComparison.OrdinalIgnoreCase);
+			bool bIsUnrealClient = TargetName.Equals("UnrealClient", StringComparison.OrdinalIgnoreCase);
 			DirectoryReference? GameDir = ProjectFile?.Directory;
 			string? GamePath = GameDir != null ? ConvertPath(GameDir.FullName) : null;
 
@@ -1129,8 +1129,8 @@ namespace UnrealBuildTool.XcodeProjectLegacy
 			string? IOSRunTimeVersion, TVOSRunTimeVersion;
 			AppendPlatformConfiguration(Content, Config, Config.BuildTarget, ProjectFile, true, !bMacOnly, !bMacOnly, Logger, out IOSRunTimeVersion, out TVOSRunTimeVersion);
 
-			bool bIsUnrealGame = Config.BuildTarget.Equals("UnrealGame", StringComparison.InvariantCultureIgnoreCase);
-			bool bIsUnrealClient = Config.BuildTarget.Equals("UnrealClient", StringComparison.InvariantCultureIgnoreCase);
+			bool bIsUnrealGame = Config.BuildTarget.Equals("UnrealGame", StringComparison.OrdinalIgnoreCase);
+			bool bIsUnrealClient = Config.BuildTarget.Equals("UnrealClient", StringComparison.OrdinalIgnoreCase);
 
 			DirectoryReference? GameDir = ProjectFile?.Directory;
 			string? GamePath = GameDir != null ? ConvertPath(GameDir.FullName) : null;
diff --git a/Engine/Source/Programs/UnrealBuildTool/System/ConfigCache.cs b/Engine/Source/Programs/UnrealBuildTool/System/ConfigCache.cs
--- a/Engine/Source/Programs/UnrealBuildTool/System/ConfigCache.cs
+++ b/Engine/Source/Programs/UnrealBuildTool/System/ConfigCache.cs
@@ -341,15 +341,15 @@ namespace UnrealBuildTool
 			bool bDumpIniLoads = false;
 			foreach (string CmdLineArg in CmdLineArgs)
 			{
-				if (CmdLineArg.StartsWith(IniConfigArgPrefix, StringComparison.InvariantCultureIgnoreCase))
+				if (CmdLineArg.StartsWith(IniConfigArgPrefix, StringComparison.OrdinalIgnoreCase))
 				{
 					OverrideStrings.Add(CmdLineArg.Substring(IniConfigArgPrefix.Length));
 				}
-				if (CmdLineArg.StartsWith(CustomConfigPrefix, StringComparison.InvariantCultureIgnoreCase))
+				if (CmdLineArg.StartsWith(CustomConfigPrefix, StringComparison.OrdinalIgnoreCase))
 				{
 					CustomConfig = CmdLineArg.Substring(CustomConfigPrefix.Length);
 				}
-				if (CmdLineArg.Equals("-dumpiniloads", StringComparison.InvariantCultureIgnoreCase))
+				if (CmdLineArg.Equals("-dumpiniloads", StringComparison.OrdinalIgnoreCase))
 				{
 					bDumpIniLoads = true;
 				}
@@ -753,7 +753,7 @@ namespace UnrealBuildTool
 			int SectionIndex = -1;
 			for (int Index = 0; Index < Lines.Count; Index++)
 			{
-				if (Lines[Index].Trim().Equals(SectionString, StringComparison.InvariantCultureIgnoreCase))
+				if (Lines[Index].Trim().Equals(SectionString, StringComparison.OrdinalIgnoreCase))
 				{
 					SectionIndex = Index;
 					break;
diff --git a/Engine/Source/Programs/UnrealBuildTool/System/ConfigFile.cs b/Engine/Source/Programs/UnrealBuildTool/System/ConfigFile.cs
--- a/Engine/Source/Programs/UnrealBuildTool/System/ConfigFile.cs
+++ b/Engine/Source/Programs/UnrealBuildTool/System/ConfigFile.cs
@@ -131,7 +131,7 @@ namespace UnrealBuildTool
 		{
 			foreach (ConfigLine Line in Lines)
 			{
-				if (Line.Key.Equals(Name))
+				if (Line.Key.Equals(Name, StringComparison.Ordinal))
 				{
 					OutLine = Line;
 					return true;
@@ -200,7 +200,7 @@ namespace UnrealBuildTool
 			foreach (KeyValuePair<string, ConfigFileSection> Pair in Sections)
 			{
 				// remember a remap for section names
-				if (Pair.Key.Equals("SectionNameRemap", StringComparison.InvariantCultureIgnoreCase))
+				if (Pair.Key.Equals("SectionNameRemap", StringComparison.OrdinalIgnoreCase))
 				{
 					foreach (ConfigLine Line in Pair.Value.Lines)
 					{
@@ -210,7 +210,7 @@ namespace UnrealBuildTool
 				else
 				{
 					// any other section is rmembered by the section name here, and each key/value pair is a remap for the given section
-					Dictionary<string, string> KeyRemap = new(StringComparer.InvariantCultureIgnoreCase);
+					Dictionary<string, string> KeyRemap = new(StringComparer.OrdinalIgnoreCase);
 					SectionKeyRemap.Add(Pair.Key, KeyRemap);
 					foreach (ConfigLine Line in Pair.Value.Lines)
 					{
@@ -452,7 +452,7 @@ namespace UnrealBuildTool
 			string NewKey = RemapSectionOrKey(KeyRemap, Key, $"which is a config key in section [{Section.Name}], in '{Filename}'");
 
 			// look for a section:name remap
-			if (!NewKey.Equals(Key) && NewKey.Contains(":", StringComparison.CurrentCulture))
+			if (!NewKey.Equals(Key, StringComparison.Ordinal) && NewKey.Contains(":", StringComparison.Ordinal))
 			{
 				string SectionName = NewKey.Substring(0, NewKey.IndexOf(':'));
 				ConfigFileSection? CurrentSection;
diff --git a/Engine/Source/Programs/UnrealBuildTool/System/ConfigHierarchy.cs b/Engine/Source/Programs/UnrealBuildTool/System/ConfigHierarchy.cs
--- a/Engine/Source/Programs/UnrealBuildTool/System/ConfigHierarchy.cs
+++ b/Engine/Source/Programs/UnrealBuildTool/System/ConfigHierarchy.cs
@@ -106,7 +106,7 @@ namespace UnrealBuildTool
 						}
 						else if (Line.Action == ConfigLineAction.RemoveKeyValue)
 						{
-							Values.RemoveAll(x => x.Equals(Line.Value, StringComparison.InvariantCultureIgnoreCase));
+							Values.RemoveAll(x => x.Equals(Line.Value, StringComparison.OrdinalIgnoreCase));
 						}
 					}
 					else
@@ -486,12 +486,12 @@ namespace UnrealBuildTool
 		public static bool TryParse(string Text, out bool Value)
 		{
 			// C# Boolean type expects "False" or "True" but since we're not case sensitive, we need to suppor that manually
-			if (Text == "1" || Text.Equals("true", StringComparison.InvariantCultureIgnoreCase))
+			if (Text == "1" || Text.Equals("true", StringComparison.OrdinalIgnoreCase))
 			{
 				Value = true;
 				return true;
 			}
-			else if (Text == "0" || Text.Equals("false", StringComparison.InvariantCultureIgnoreCase))
+			else if (Text == "0" || Text.Equals("false", StringComparison.OrdinalIgnoreCase))
 			{
 				Value = false;
 				return true;
diff --git a/Engine/Source/Programs/UnrealBuildTool/System/GraphVisualization.cs b/Engine/Source/Programs/UnrealBuildTool/System/GraphVisualization.cs
--- a/Engine/Source/Programs/UnrealBuildTool/System/GraphVisualization.cs
+++ b/Engine/Source/Programs/UnrealBuildTool/System/GraphVisualization.cs
@@ -107,7 +107,7 @@ namespace UnrealBuildTool
 						{
 							AllAttributes[AttributeName] = new GraphAttribute(AllAttributes.Count, AttributeName, AttributeType);
 						}
-						else if (!Attribute.TypeName.Equals(AttributeType))
+						else if (!Attribute.TypeName.Equals(AttributeType, StringComparison.Ordinal))
 						{
 							throw new InvalidOperationException("Multiple graph nodes with the same attribute name but different types encountered!");
 						}
diff --git a/Engine/Source/Programs/UnrealBuildTool/System/Plugins.cs b/Engine/Source/Programs/UnrealBuildTool/System/Plugins.cs
--- a/Engine/Source/Programs/UnrealBuildTool/System/Plugins.cs
+++ b/Engine/Source/Programs/UnrealBuildTool/System/Plugins.cs
@@ -488,7 +488,7 @@ namespace UnrealBuildTool
 				{
 					foreach (ModuleDescriptor ChildModule in Child.Descriptor.Modules)
 					{
-						ModuleDescriptor? ParentModule = Parent.Descriptor.Modules.FirstOrDefault(x => x.Name.Equals(ChildModule.Name) && x.Type == ChildModule.Type);
+						ModuleDescriptor? ParentModule = Parent.Descriptor.Modules.FirstOrDefault(x => x.Name.Equals(ChildModule.Name, StringComparison.Ordinal) && x.Type == ChildModule.Type);
 						if (ParentModule != null)
 						{
 							// merge allow/deny lists (if the parent had a list, and child didn't specify a list, just add the child platform to the parent list - for allow/deny lists!)
@@ -581,7 +581,7 @@ namespace UnrealBuildTool
 				{
 					foreach (PluginReferenceDescriptor ChildPluginReference in Child.Descriptor.Plugins)
 					{
-						PluginReferenceDescriptor? ParentPluginReference = Parent.Descriptor.Plugins.FirstOrDefault(x => x.Name.Equals(ChildPluginReference.Name));
+						PluginReferenceDescriptor? ParentPluginReference = Parent.Descriptor.Plugins.FirstOrDefault(x => x.Name.Equals(ChildPluginReference.Name, StringComparison.Ordinal));
 						if (ParentPluginReference != null)
 						{
 							// we only need to explicitly list the platform in an allow list if the parent also had an allow list (otherwise, we could mistakenly remove all other platforms)
diff --git a/Engine/Source/Programs/UnrealBuildTool/System/TargetReceipt.cs b/Engine/Source/Programs/UnrealBuildTool/System/TargetReceipt.cs
--- a/Engine/Source/Programs/UnrealBuildTool/System/TargetReceipt.cs
+++ b/Engine/Source/Programs/UnrealBuildTool/System/TargetReceipt.cs
@@ -515,16 +515,16 @@ namespace UnrealBuildTool
 		/// Checks the Additional properties for one with the given name that matches the given value
 		/// </summary>
 		/// <param name="PropertyName">Property name to search for</param>
-		/// <param name="Value">Value to compare against (with StringComparison.InvariantCultureIgnoreCase)</param>
+		/// <param name="Value">Value to compare against (with StringComparison.OrdinalIgnoreCase)</param>
 		/// <returns>True if any property with PropertyName has a value matching Value</returns>
 		public bool HasValueForAdditionalProperty(string PropertyName, string Value)
 		{
 			// get all properties with the given name?
-			IEnumerable<ReceiptProperty> Results = AdditionalProperties.Where(x => x.Name == PropertyName);
+			IEnumerable<ReceiptProperty> Results = AdditionalProperties.Where(x => PropertyName.Equals(x.Name, StringComparison.Ordinal));
 			foreach (ReceiptProperty Property in Results)
 			{
 				// does the property value match?
-				if (Property.Value.Equals(Value, StringComparison.InvariantCultureIgnoreCase))
+				if (Property.Value.Equals(Value, StringComparison.OrdinalIgnoreCase))
 				{
 					return true;
 				}
diff --git a/Engine/Source/Programs/UnrealBuildTool/System/UHTExecution.cs b/Engine/Source/Programs/UnrealBuildTool/System/UHTExecution.cs
--- a/Engine/Source/Programs/UnrealBuildTool/System/UHTExecution.cs
+++ b/Engine/Source/Programs/UnrealBuildTool/System/UHTExecution.cs
@@ -828,7 +828,7 @@ namespace UnrealBuildTool
 				// Find the CoreUObject module
 				foreach (UHTModuleInfo Module in UObjectModules)
 				{
-					if (Module.ModuleName.Equals("CoreUObject", StringComparison.InvariantCultureIgnoreCase))
+					if (Module.ModuleName.Equals("CoreUObject", StringComparison.OrdinalIgnoreCase))
 					{
 						CoreGeneratedTimestampUtc = GetCoreGeneratedTimestampUtc(Module.ModuleName, Path.GetDirectoryName(Module.GeneratedCPPFilenameBase)!);
 						break;
diff --git a/Engine/Source/Programs/UnrealBuildTool/System/UnrealPluginLanguage.cs b/Engine/Source/Programs/UnrealBuildTool/System/UnrealPluginLanguage.cs
--- a/Engine/Source/Programs/UnrealBuildTool/System/UnrealPluginLanguage.cs
+++ b/Engine/Source/Programs/UnrealBuildTool/System/UnrealPluginLanguage.cs
@@ -603,7 +603,7 @@ namespace UnrealBuildTool
 			}
 			TargetPlatform = InTargetPlatform;
 
-			string PathPrefix = Path.GetFileName(Directory.GetCurrentDirectory()).Equals("Source") ? ".." : "Engine";
+			string PathPrefix = Path.GetFileName(Directory.GetCurrentDirectory()).Equals("Source", StringComparison.Ordinal) ? ".." : "Engine";
 
 			XDoc = XDocument.Parse("<root " + XMLRootDefinition + "></root>");
 			foreach (string Basename in InXMLFiles)
@@ -896,8 +896,12 @@ namespace UnrealBuildTool
 			{
 				return false;
 			}
-			Input = Input.ToLower();
-			return !(Input.Equals("0") || Input.Equals("false") || Input.Equals("off") || Input.Equals("no"));
+			return !(
+				Input.Equals("0", StringComparison.OrdinalIgnoreCase) ||
+				Input.Equals("false", StringComparison.OrdinalIgnoreCase) ||
+				Input.Equals("off", StringComparison.OrdinalIgnoreCase) ||
+				Input.Equals("no", StringComparison.OrdinalIgnoreCase)
+			);
 		}
 
 		private int StringToInt(UPLContext Context, XElement Node, string? Input)
@@ -1208,7 +1212,7 @@ namespace UnrealBuildTool
 					case "isArch":
 						{
 							string? arch = GetAttribute(CurrentContext, Node, "arch");
-							if (arch != null && arch.Equals(Architecture))
+							if (arch != null && arch.Equals(Architecture, StringComparison.Ordinal))
 							{
 								foreach (XElement instruction in Node.Elements().Reverse())
 								{
@@ -1287,7 +1291,7 @@ namespace UnrealBuildTool
 							while (ExecutionStack.Count > 0)
 							{
 								Node = ExecutionStack.Pop();
-								if (Node.Name.ToString().Equals("while"))
+								if (Node.Name.ToString().Equals("while", StringComparison.Ordinal))
 								{
 									break;
 								}
@@ -1301,7 +1305,7 @@ namespace UnrealBuildTool
 							while (ExecutionStack.Count > 0)
 							{
 								Node = ExecutionStack.Pop();
-								if (Node.Name.ToString().Equals("while"))
+								if (Node.Name.ToString().Equals("while", StringComparison.Ordinal))
 								{
 									ExecutionStack.Push(Node);
 									break;
@@ -1363,7 +1367,7 @@ namespace UnrealBuildTool
 								}
 								else
 								{
-									if (CurrentElement.Name.ToString().Equals(Tag))
+									if (CurrentElement.Name.ToString().Equals(Tag, StringComparison.Ordinal))
 									{
 										AddAttribute(CurrentElement, Name, Value);
 									}
@@ -1400,7 +1404,7 @@ namespace UnrealBuildTool
 								}
 								else
 								{
-									if (CurrentElement.Name.ToString().Equals(Tag))
+									if (CurrentElement.Name.ToString().Equals(Tag, StringComparison.Ordinal))
 									{
 										RemoveAttribute(CurrentElement, Name);
 									}
@@ -1649,7 +1653,7 @@ namespace UnrealBuildTool
 								}
 								else
 								{
-									if (CurrentElement.Name.ToString().Equals(Tag))
+									if (CurrentElement.Name.ToString().Equals(Tag, StringComparison.Ordinal))
 									{
 										CurrentElement.Add(new XElement(Element));
 									}
@@ -1697,7 +1701,7 @@ namespace UnrealBuildTool
 								}
 								else
 								{
-									if (CurrentElement.Name.ToString().Equals(Tag))
+									if (CurrentElement.Name.ToString().Equals(Tag, StringComparison.Ordinal))
 									{
 										AddElements(CurrentElement, Node);
 									}
@@ -1954,7 +1958,7 @@ namespace UnrealBuildTool
 										Value = false;
 										foreach (string Entry in StringList)
 										{
-											if (Entry.Equals(Contains))
+											if (Entry.Equals(Contains, StringComparison.Ordinal))
 											{
 												Value = true;
 												break;
@@ -2045,7 +2049,7 @@ namespace UnrealBuildTool
 							string Arg2 = GetAttribute(CurrentContext, Node, "arg2", true, false, "");
 							if (Result != null)
 							{
-								CurrentContext.BoolVariables[Result] = Arg1.Equals(Arg2);
+								CurrentContext.BoolVariables[Result] = Arg1.Equals(Arg2, StringComparison.Ordinal);
 							}
 						}
 						break;
@@ -3041,11 +3045,11 @@ namespace UnrealBuildTool
 			if (Result)
 			{
 				// C# Boolean type expects "False" or "True" but since we're not case sensitive, we need to suppor that manually
-				if (String.Equals(TextValue, "true", StringComparison.CurrentCultureIgnoreCase) || String.Equals(TextValue, "1"))
+				if (String.Equals(TextValue, "true", StringComparison.OrdinalIgnoreCase) || String.Equals(TextValue, "1", StringComparison.Ordinal))
 				{
 					Value = true;
 				}
-				else if (String.Equals(TextValue, "false", StringComparison.CurrentCultureIgnoreCase) || String.Equals(TextValue, "0"))
+				else if (String.Equals(TextValue, "false", StringComparison.OrdinalIgnoreCase) || String.Equals(TextValue, "0", StringComparison.Ordinal))
 				{
 					Value = false;
 				}
@@ -3419,7 +3423,7 @@ namespace UnrealBuildTool
 						IniValues? Value;
 						if (CurrentSection.TryGetValue(Key, out Value))
 						{
-							int ExistingIndex = Value.FindIndex(X => (String.Equals(SingleValue, X, StringComparison.CurrentCultureIgnoreCase)));
+							int ExistingIndex = Value.FindIndex(X => (String.Equals(SingleValue, X, StringComparison.OrdinalIgnoreCase)));
 							if (ExistingIndex >= 0)
 							{
 								Value.RemoveAt(ExistingIndex);
diff --git a/Engine/Source/Programs/UnrealBuildTool/System/Utils.cs b/Engine/Source/Programs/UnrealBuildTool/System/Utils.cs
--- a/Engine/Source/Programs/UnrealBuildTool/System/Utils.cs
+++ b/Engine/Source/Programs/UnrealBuildTool/System/Utils.cs
@@ -581,7 +581,7 @@ namespace UnrealBuildTool
 					string CurPlatformString = CurPlatform.ToString();
 					foreach (UnrealPlatformGroup Group in UnrealPlatformGroup.GetValidGroups())
 					{
-						if (Group.ToString().Equals(CurPlatformString))
+						if (Group.ToString().Equals(CurPlatformString, StringComparison.Ordinal))
 						{
 							ShouldConsider = false;
 							break;
@@ -1622,7 +1622,7 @@ namespace UnrealBuildTool
 					// check native project paths (uprojectdirs)
 					IEnumerable<FileReference> NativeProjectFiles = NativeProjects.EnumerateProjectFiles(Logger);
 
-					CandidateProjectFile = NativeProjectFiles.Where(F => F.GetFileName().Equals(ProjectName, StringComparison.InvariantCultureIgnoreCase)).FirstOrDefault();
+					CandidateProjectFile = NativeProjectFiles.Where(F => F.GetFileName().Equals(ProjectName, StringComparison.OrdinalIgnoreCase)).FirstOrDefault();
 				}
 
 				if (CandidateProjectFile == null || !FileReference.Exists(CandidateProjectFile))
diff --git a/Engine/Source/Programs/UnrealBuildTool/ToolChain/RemoteMac.cs b/Engine/Source/Programs/UnrealBuildTool/ToolChain/RemoteMac.cs
--- a/Engine/Source/Programs/UnrealBuildTool/ToolChain/RemoteMac.cs
+++ b/Engine/Source/Programs/UnrealBuildTool/ToolChain/RemoteMac.cs
@@ -489,7 +489,7 @@ namespace UnrealBuildTool
 			bool bLogIsMapped = false;
 			foreach (RemoteMapping Mapping in Mappings)
 			{
-				if (RemoteLogFile.Directory.FullName.Equals(Mapping.LocalDirectory.FullName, StringComparison.InvariantCultureIgnoreCase))
+				if (RemoteLogFile.Directory.FullName.Equals(Mapping.LocalDirectory.FullName, StringComparison.OrdinalIgnoreCase))
 				{
 					bLogIsMapped = true;
 					break;
@@ -697,12 +697,12 @@ namespace UnrealBuildTool
 				string Key = LocalArgument[..EqualsIdx];
 				string Value = LocalArgument[(EqualsIdx + 1)..];
 
-				if (Key.Equals("-Log", StringComparison.InvariantCultureIgnoreCase))
+				if (Key.Equals("-Log", StringComparison.OrdinalIgnoreCase))
 				{
 					// We are already writing to the local log file. The remote will produce a different log (RemoteLogFile)
 					continue;
 				}
-				if (Key.Equals("-Manifest", StringComparison.InvariantCultureIgnoreCase) && LocalManifestFiles != null)
+				if (Key.Equals("-Manifest", StringComparison.OrdinalIgnoreCase) && LocalManifestFiles != null)
 				{
 					LocalManifestFiles.Add(new FileReference(Value));
 					continue;
@@ -711,7 +711,7 @@ namespace UnrealBuildTool
 				string RemoteArgument = LocalArgument;
 				foreach (RemoteMapping Mapping in Mappings)
 				{
-					if (Value.StartsWith(Mapping.LocalDirectory.FullName, StringComparison.InvariantCultureIgnoreCase))
+					if (Value.StartsWith(Mapping.LocalDirectory.FullName, StringComparison.OrdinalIgnoreCase))
 					{
 						RemoteArgument = String.Format("{0}={1}", Key, GetRemotePath(Value));
 						break;
diff --git a/Engine/Source/Programs/UnrealBuildTool/UnrealBuildTool.cs b/Engine/Source/Programs/UnrealBuildTool/UnrealBuildTool.cs
--- a/Engine/Source/Programs/UnrealBuildTool/UnrealBuildTool.cs
+++ b/Engine/Source/Programs/UnrealBuildTool/UnrealBuildTool.cs
@@ -341,7 +341,7 @@ namespace UnrealBuildTool
 					}
 
 					// special case for Mode
-					if (String.Equals(Att.Prefix, "-Mode="))
+					if (String.Equals(Att.Prefix, "-Mode=", StringComparison.Ordinal))
 					{
 						Console.WriteLine($"  {Att.Prefix!.PadRight(LongestPrefix)} :  Select tool mode. One of the following (default tool mode is \"Build\"):");
 						string Indent = "".PadRight(LongestPrefix + 8);
