[UBT]
* Changed header units to use ifcmap and ifcDebugRecords to reduce size of files/.

diff --git a/Engine/Source/Programs/UnrealBuildTool/Platform/Windows/VCToolChain.cs b/Engine/Source/Programs/UnrealBuildTool/Platform/Windows/VCToolChain.cs
--- a/Engine/Source/Programs/UnrealBuildTool/Platform/Windows/VCToolChain.cs
+++ b/Engine/Source/Programs/UnrealBuildTool/Platform/Windows/VCToolChain.cs
@@ -318,6 +318,7 @@ namespace UnrealBuildTool
 				Arguments.Add($"/dxifcInlineFunctions");
 				Arguments.Add($"/ifcOutput {CreatePchFilePath}");
 				Arguments.Add($"/dxifcSuppressDllImportTransform");
+				Arguments.Add($"-d1experimental:ifcDebugRecords");
 			}
 			else
 			{
@@ -336,6 +337,7 @@ namespace UnrealBuildTool
 				Arguments.Add($"/translateInclude");
 				Arguments.Add($"/headerUnit:quote {PchThroughHeaderFilePath}={UsingPchFilePath}");
 				Arguments.Add($"/dxifcSuppressDllImportTransform");
+				Arguments.Add($"-d1experimental:ifcDebugRecords");
 			}
 			else
 			{
@@ -2114,9 +2116,11 @@ namespace UnrealBuildTool
 
 		protected override CPPOutput CompileCPPFiles(CppCompileEnvironment CompileEnvironment, IEnumerable<FileItem> InputFiles, DirectoryReference OutputDir, string ModuleName, IActionGraphBuilder Graph)
 		{
+			CPPOutput Result = new CPPOutput();
+
 			if (Target.StaticAnalyzer != StaticAnalyzer.None && CompileEnvironment.bDisableStaticAnalysis)
 			{
-				return new CPPOutput();
+				return Result;
 			}
 
 			VCCompileAction BaseCompileAction = CreateBaseCompileAction(CompileEnvironment);
@@ -2160,6 +2164,19 @@ namespace UnrealBuildTool
 						string FakeUniquePCHSymbolName = CompileEnvironment.PrecompiledHeaderIncludeFilename.GetFileNameWithoutExtension();
 						CompileAction.Arguments.Add($"/Yl{FakeUniquePCHSymbolName}");
 					}
+
+					if (CompileEnvironment.bUseHeaderUnitsForPch)
+					{
+						string ifcPath = NormalizeCommandLinePath(CompileEnvironment.PrecompiledHeaderIncludeFilename!, CompileEnvironment.RootPaths);
+						string tomlString =
+							$"[[header-unit]]\r\n" +
+							$"name = [\"quote\", '{ifcPath.Replace('/', '\\')}']\r\n" +
+							$"ifc = \"{CompileEnvironment.PrecompiledHeaderIncludeFilename!.FullName.Replace('\\', '/') + ".ifc"}\"\r\n";
+						FileReference tomlFileRef = FileReference.Combine(OutputDir, SourceFile.Location.GetFileName() + ".toml");
+						FileItem tomlFile = Graph.CreateIntermediateTextFile(tomlFileRef, tomlString);
+						CompileAction.AdditionalPrerequisiteItems.Add(tomlFile);
+						Result.ObjectFiles.Add(tomlFile);
+					}
 				}
 
 				string FileName = SourceFile.Name;
@@ -2416,7 +2433,6 @@ namespace UnrealBuildTool
 				Actions.Add(CompileAction);
 			}
 
-			CPPOutput Result = new CPPOutput();
 			Result.ObjectFiles.AddRange(Actions.Where(x => x.ObjectFile != null || x.PreprocessedFile != null).Select(x => x.ObjectFile != null ? x.ObjectFile! : x.PreprocessedFile!));
 			// Clang static analysis doesn't create object files, so treat the dependency list file as the output
 			if (Target.WindowsPlatform.Compiler.IsClang() && Target.StaticAnalyzer == StaticAnalyzer.Default)
@@ -3260,8 +3276,21 @@ namespace UnrealBuildTool
 			List<string> InputFileNames = new List<string>();
 			foreach (FileItem InputFile in LinkEnvironment.InputFiles)
 			{
-				InputFileNames.Add($"\"{NormalizeCommandLinePath(InputFile, RootPaths)}\"");
-				PrerequisiteItems.Add(InputFile);
+				if (InputFile.HasExtension(".toml"))
+				{
+					// We don't want the toml file for import lib creation
+					if (!LinkEnvironment.bIsBuildingDLL || !bIsBuildingLibraryOrImportLibrary)
+					{
+						Arguments.Add("/ifcDebugRecords");
+						Arguments.Add($"/ifcMap:\"{NormalizeCommandLinePath(InputFile, RootPaths)}\"");
+						PrerequisiteItems.Add(InputFile);
+					}
+				}
+				else
+				{
+					InputFileNames.Add($"\"{NormalizeCommandLinePath(InputFile, RootPaths)}\"");
+					PrerequisiteItems.Add(InputFile);
+				}
 			}
 
 			if (!bIsBuildingLibraryOrImportLibrary)
