Fix / selectively disable issue CA1308 (Normalize strings to uppercase) in UBT.

diff --git a/Engine/Source/Programs/UnrealBuildTool/Executors/UnrealBuildAccelerator/UBAAgentCoordinatorHorde.cs b/Engine/Source/Programs/UnrealBuildTool/Executors/UnrealBuildAccelerator/UBAAgentCoordinatorHorde.cs
--- a/Engine/Source/Programs/UnrealBuildTool/Executors/UnrealBuildAccelerator/UBAAgentCoordinatorHorde.cs
+++ b/Engine/Source/Programs/UnrealBuildTool/Executors/UnrealBuildAccelerator/UBAAgentCoordinatorHorde.cs
@@ -803,8 +803,10 @@ namespace UnrealBuildTool
 			}
 
 			// Normalize URLs
+#pragma warning disable CA1308 // Normalize strings to uppercase - URLs are conventionally lowercased
 			_hordeConfig.HordeServer = _hordeConfig.HordeServer?.TrimEnd('/').ToLowerInvariant();
 			_hordeConfig.UBASentryUrl = _hordeConfig.UBASentryUrl?.TrimEnd('/').ToLowerInvariant();
+#pragma warning restore CA1308
 		}
 
 		public DirectoryReference? GetUBARootDir()
diff --git a/Engine/Source/Programs/UnrealBuildTool/Executors/UnrealBuildAccelerator/UBAExecutor.cs b/Engine/Source/Programs/UnrealBuildTool/Executors/UnrealBuildAccelerator/UBAExecutor.cs
--- a/Engine/Source/Programs/UnrealBuildTool/Executors/UnrealBuildAccelerator/UBAExecutor.cs
+++ b/Engine/Source/Programs/UnrealBuildTool/Executors/UnrealBuildAccelerator/UBAExecutor.cs
@@ -522,8 +522,6 @@ namespace UnrealBuildTool
 					_ = Task.Run(LaunchVisualizer);
 				}
 
-				string arch = RuntimeInformation.ProcessArchitecture.ToString().ToLowerInvariant();
-
 				OnSessionCreated?.Invoke(_session!);
 
 				bool success = ExecuteActionsInternal(logger, inputActions, _server!, _storageServer!, _session!, _cacheClients.Values);
diff --git a/Engine/Source/Programs/UnrealBuildTool/Platform/Windows/VCEnvironment.cs b/Engine/Source/Programs/UnrealBuildTool/Platform/Windows/VCEnvironment.cs
--- a/Engine/Source/Programs/UnrealBuildTool/Platform/Windows/VCEnvironment.cs
+++ b/Engine/Source/Programs/UnrealBuildTool/Platform/Windows/VCEnvironment.cs
@@ -276,7 +276,9 @@ namespace UnrealBuildTool
 			if (bAllowClangLinker && Compiler.IsClang())
 			{
 				// Since obj files could be LLVM IR Stream file format (when building with ltcg) we can't use link.exe. UbaObjTool support all formats and produce identical lib files as link.exe when running non-ltcg
+#pragma warning disable CA1308 // Normalize strings to uppercase - the directory names are lowercased in VCS
 				DirectoryReference Dir = DirectoryReference.Combine(Unreal.EngineDirectory, "Binaries", "Win64", "UnrealBuildAccelerator", RuntimeInformation.ProcessArchitecture.ToString().ToLowerInvariant());
+#pragma warning restore CA1308
 				return FileReference.Combine(Dir, "UbaObjTool.exe");
 			}
 			return FileReference.Combine(GetVCToolPath(ToochainDir, Architecture), "link.exe"); // We add /LIB to cmd line so we can use link.exe directly instead of going via lib.exe
diff --git a/Engine/Source/Programs/UnrealBuildTool/Platform/Windows/VCToolChain.cs b/Engine/Source/Programs/UnrealBuildTool/Platform/Windows/VCToolChain.cs
--- a/Engine/Source/Programs/UnrealBuildTool/Platform/Windows/VCToolChain.cs
+++ b/Engine/Source/Programs/UnrealBuildTool/Platform/Windows/VCToolChain.cs
@@ -499,7 +499,9 @@ namespace UnrealBuildTool
 					Arguments.Add("--analyze");
 
 					// Write out a pretty web page with navigation to understand how the analysis was derived if HTML is enabled.
+#pragma warning disable CA1308 // Normalize strings to uppercase - Clang arg may be case-sensitive
 					Arguments.Add($"-Xclang -analyzer-output={Target.StaticAnalyzerOutputType.ToString().ToLowerInvariant()}");
+#pragma warning restore CA1308
 
 					// Needed for some of the C++ checkers.
 					Arguments.Add("-Xclang -analyzer-config -Xclang aggressive-binary-operation-simplification=true");
diff --git a/Engine/Source/Programs/UnrealBuildTool/Preprocessor/SourceFileMarkup.cs b/Engine/Source/Programs/UnrealBuildTool/Preprocessor/SourceFileMarkup.cs
--- a/Engine/Source/Programs/UnrealBuildTool/Preprocessor/SourceFileMarkup.cs
+++ b/Engine/Source/Programs/UnrealBuildTool/Preprocessor/SourceFileMarkup.cs
@@ -194,7 +194,7 @@ namespace UnrealBuildTool
 				result.Append('#');
 				if (Type != SourceFileMarkupType.OtherDirective)
 				{
-					result.Append(Type.ToString().ToLowerInvariant());
+					result.Append(Type.ToString().ToUpperInvariant());
 				}
 				if (Tokens != null && Tokens.Count > 0)
 				{
diff --git a/Engine/Source/Programs/UnrealBuildTool/ProjectFiles/ProjectFileGenerator.cs b/Engine/Source/Programs/UnrealBuildTool/ProjectFiles/ProjectFileGenerator.cs
--- a/Engine/Source/Programs/UnrealBuildTool/ProjectFiles/ProjectFileGenerator.cs
+++ b/Engine/Source/Programs/UnrealBuildTool/ProjectFiles/ProjectFileGenerator.cs
@@ -906,43 +906,42 @@ namespace UnrealBuildTool
 			string PreferredAccessor;
 			if (Ini.GetString("/Script/SourceCodeAccess.SourceCodeAccessSettings", "PreferredAccessor", out PreferredAccessor))
 			{
-				PreferredAccessor = PreferredAccessor.ToLowerInvariant();
-				if (PreferredAccessor == "clionsourcecodeaccessor")
+				if (PreferredAccessor.Equals("clionsourcecodeaccessor", StringComparison.OrdinalIgnoreCase))
 				{
 					Format = ProjectFileFormat.CLion;
 					return true;
 				}
-				else if (PreferredAccessor == "codelitesourcecodeaccessor")
+				else if (PreferredAccessor.Equals("codelitesourcecodeaccessor", StringComparison.OrdinalIgnoreCase))
 				{
 					Format = ProjectFileFormat.CodeLite;
 					return true;
 				}
-				else if (PreferredAccessor == "xcodesourcecodeaccessor")
+				else if (PreferredAccessor.Equals("xcodesourcecodeaccessor", StringComparison.OrdinalIgnoreCase))
 				{
 					Format = ProjectFileFormat.XCode;
 					return true;
 				}
-				else if (PreferredAccessor == "visualstudiocode")
+				else if (PreferredAccessor.Equals("visualstudiocode", StringComparison.OrdinalIgnoreCase))
 				{
 					Format = ProjectFileFormat.VisualStudioCode;
 					return true;
 				}
-				else if (PreferredAccessor == "kdevelopsourcecodeaccessor")
+				else if (PreferredAccessor.Equals("kdevelopsourcecodeaccessor", StringComparison.OrdinalIgnoreCase))
 				{
 					Format = ProjectFileFormat.KDevelop;
 					return true;
 				}
-				else if (PreferredAccessor == "visualstudiosourcecodeaccessor")
+				else if (PreferredAccessor.Equals("visualstudiosourcecodeaccessor", StringComparison.OrdinalIgnoreCase))
 				{
 					Format = ProjectFileFormat.VisualStudio;
 					return true;
 				}
-				else if (PreferredAccessor == "visualstudio2022")
+				else if (PreferredAccessor.Equals("visualstudio2022", StringComparison.OrdinalIgnoreCase))
 				{
 					Format = ProjectFileFormat.VisualStudio2022;
 					return true;
 				}
-				else if (PreferredAccessor == "visualstudio2026")
+				else if (PreferredAccessor.Equals("visualstudio2026", StringComparison.OrdinalIgnoreCase))
 				{
 					Format = ProjectFileFormat.VisualStudio2026;
 					return true;
diff --git a/Engine/Source/Programs/UnrealBuildTool/ProjectFiles/VisualStudioCode/VSCodeProjectFileGenerator.cs b/Engine/Source/Programs/UnrealBuildTool/ProjectFiles/VisualStudioCode/VSCodeProjectFileGenerator.cs
--- a/Engine/Source/Programs/UnrealBuildTool/ProjectFiles/VisualStudioCode/VSCodeProjectFileGenerator.cs
+++ b/Engine/Source/Programs/UnrealBuildTool/ProjectFiles/VisualStudioCode/VSCodeProjectFileGenerator.cs
@@ -1279,7 +1279,9 @@ namespace UnrealBuildTool
 				OutFile.BeginObject();
 				{
 					OutFile.AddField("name", Target.Name + " Oculus (" + BuildProduct.Config.ToString() + ") " + ConfigType);
+#pragma warning disable CA1308 // Normalize strings to uppercase - should maintain compatibility with the output file schema
 					OutFile.AddField("request", ConfigType.ToLowerInvariant());
+#pragma warning restore CA1308
 					if (ConfigType == "Launch")
 					{
 						OutFile.AddField("preLaunchTask", LaunchTaskName);
diff --git a/Engine/Source/Programs/UnrealBuildTool/System/SourceFileSearch.cs b/Engine/Source/Programs/UnrealBuildTool/System/SourceFileSearch.cs
--- a/Engine/Source/Programs/UnrealBuildTool/System/SourceFileSearch.cs
+++ b/Engine/Source/Programs/UnrealBuildTool/System/SourceFileSearch.cs
@@ -66,7 +66,7 @@ namespace UnrealBuildTool
 			}
 			else
 			{
-				ExcludedDirectorySuffixes = SubdirectoryNamesToExclude.Select(x => Path.DirectorySeparatorChar + x.ToLowerInvariant()).Union(DefaultExcludedDirectorySuffixes).ToArray();
+				ExcludedDirectorySuffixes = SubdirectoryNamesToExclude.Select(x => $"{Path.DirectorySeparatorChar}{x}").Union(DefaultExcludedDirectorySuffixes).ToArray();
 			}
 
 			// Find all the files
@@ -101,7 +101,7 @@ namespace UnrealBuildTool
 			}
 			else
 			{
-				ExcludedDirectorySuffixes = SubdirectoryNamesToExclude.Select(x => Path.DirectorySeparatorChar + x.ToLowerInvariant()).Union(DefaultExcludedDirectorySuffixes).ToArray();
+				ExcludedDirectorySuffixes = SubdirectoryNamesToExclude.Select(x => $"{Path.DirectorySeparatorChar}{x}").Union(DefaultExcludedDirectorySuffixes).ToArray();
 			}
 
 			// Find all the files
diff --git a/Engine/Source/Programs/UnrealBuildTool/System/Unity.cs b/Engine/Source/Programs/UnrealBuildTool/System/Unity.cs
--- a/Engine/Source/Programs/UnrealBuildTool/System/Unity.cs
+++ b/Engine/Source/Programs/UnrealBuildTool/System/Unity.cs
@@ -278,7 +278,9 @@ namespace UnrealBuildTool
 				++CurrentUnityFileCount;
 
 				FileItem FirstFile = UnityFile.Files.FirstOrDefault() ?? UnityFile.VirtualFiles.First();
+#pragma warning disable CA1308 // Normalize strings to uppercase - File extensions are safe to lowercase under normal usage
 				string UnityExt = FirstFile.Location.GetExtension().ToLowerInvariant();
+#pragma warning restore CA1308
 				StringWriter OutputUnitySourceWriter = new();
 
 				OutputUnitySourceWriter.WriteLine($"// This file is automatically generated at compile-time to include some subset of the user-created {UnityExt.Trim('.')} files.");
diff --git a/Engine/Source/Programs/UnrealBuildTool/ToolChain/ClangToolChain.cs b/Engine/Source/Programs/UnrealBuildTool/ToolChain/ClangToolChain.cs
--- a/Engine/Source/Programs/UnrealBuildTool/ToolChain/ClangToolChain.cs
+++ b/Engine/Source/Programs/UnrealBuildTool/ToolChain/ClangToolChain.cs
@@ -1003,7 +1003,9 @@ namespace UnrealBuildTool
 			}
 
 			// Write out a pretty web page with navigation to understand how the analysis was derived if HTML is enabled.
+#pragma warning disable CA1308 // Normalize strings to uppercase - Clang arg may be case-sensitive
 			Arguments.Add($"-Xclang -analyzer-output={StaticAnalyzerOutputType.ToString().ToLowerInvariant()}");
+#pragma warning restore CA1308
 
 			// Needed for some of the C++ checkers.
 			Arguments.Add("-Xclang -analyzer-config -Xclang aggressive-binary-operation-simplification=true");
