Selectively disable issue CA1069 (Enums should not have duplicate values) in UBT.

diff --git a/Engine/Source/Programs/UnrealBuildTool/ProjectFiles/VisualStudio/VCSolutionOptions.cs b/Engine/Source/Programs/UnrealBuildTool/ProjectFiles/VisualStudio/VCSolutionOptions.cs
--- a/Engine/Source/Programs/UnrealBuildTool/ProjectFiles/VisualStudio/VCSolutionOptions.cs
+++ b/Engine/Source/Programs/UnrealBuildTool/ProjectFiles/VisualStudio/VCSolutionOptions.cs
@@ -139,6 +139,7 @@ namespace UnrealBuildTool
 			Consolidate = 8,
 		}
 
+#pragma warning disable CA1069 // Enums should not have duplicate values - must match https://learn.microsoft.com/en-us/windows/win32/stg/stgm-constants
 		[Flags]
 		public enum STGM : int
 		{
@@ -161,6 +162,7 @@ namespace UnrealBuildTool
 			NoSnapshot = 0x00200000,
 			DirectSWMR = 0x00400000,
 		}
+#pragma warning restore CA1069
 
 		[ComImport, Guid("0000000c-0000-0000-C000-000000000046"), InterfaceType(ComInterfaceType.InterfaceIsIUnknown)]
 		interface IOleStream
diff --git a/Engine/Source/Programs/UnrealBuildTool/System/PluginDescriptor.cs b/Engine/Source/Programs/UnrealBuildTool/System/PluginDescriptor.cs
--- a/Engine/Source/Programs/UnrealBuildTool/System/PluginDescriptor.cs
+++ b/Engine/Source/Programs/UnrealBuildTool/System/PluginDescriptor.cs
@@ -40,10 +40,12 @@ namespace UnrealBuildTool
 		/// </summary>
 		LatestPlusOne,
 
+#pragma warning disable CA1069 // Enums should not have duplicate values - having access to Latest is useful
 		/// <summary>
 		/// The latest plugin descriptor version
 		/// </summary>
 		Latest = LatestPlusOne - 1
+#pragma warning restore CA1069
 	}
 
 	/// <summary>
@@ -704,8 +706,8 @@ namespace UnrealBuildTool
 					{
 						if (ChildModule.bHasExplicitPlatforms && ChildModule.PlatformAllowList != null && ChildModule.PlatformAllowList.Count == 0)
 						{
-							// The order that child plugins are merged into the parent is undefined - there is no heirarchy. Only the PlatformAllowList and PlatformDenyList are currently merged into the parent
-							// Having an explicity-empty list here suggests someone is trying to create a heirarchy and may get caught out if their module declaration also includes TargetAllowList, ProgramAllowList etc. as these properties will be ignored
+							// The order that child plugins are merged into the parent is undefined - there is no hierarchy. Only the PlatformAllowList and PlatformDenyList are currently merged into the parent
+							// Having an explicity-empty list here suggests someone is trying to create a hierarchy and may get caught out if their module declaration also includes TargetAllowList, ProgramAllowList etc. as these properties will be ignored
 							Log.TraceWarningOnce(FileName, $"Plugin extensions should not declare HasExplicitPlatforms with an empty PlatformAllowList. (module {ChildModule.Name})");
 						}
 					}
