Fix issue CA1068 (CancellationToken parameters must come last) in UBT.

diff --git a/Engine/Source/Programs/UnrealBuildTool/Executors/UnrealBuildAccelerator/UBAAgentCoordinatorHorde.cs b/Engine/Source/Programs/UnrealBuildTool/Executors/UnrealBuildAccelerator/UBAAgentCoordinatorHorde.cs
--- a/Engine/Source/Programs/UnrealBuildTool/Executors/UnrealBuildAccelerator/UBAAgentCoordinatorHorde.cs
+++ b/Engine/Source/Programs/UnrealBuildTool/Executors/UnrealBuildAccelerator/UBAAgentCoordinatorHorde.cs
@@ -182,7 +182,7 @@ namespace UnrealBuildTool
 			_logger.LogInformation("Created tool bundle with locator {UbaAgentLocator}", _ubaAgentLocator.ToString());
 		}
 
-		public UnrealBuildAcceleratorCacheConfig? RequestCacheServer(CancellationToken cancellationToken, int desiredConnectionCount)
+		public UnrealBuildAcceleratorCacheConfig? RequestCacheServer(int desiredConnectionCount, CancellationToken cancellationToken)
 		{
 			// REALLY UGLY. Remove this loop. We know InitAsync has been called before this one so right now it is a hack to just loop.
 			while (_client == null)
@@ -279,7 +279,7 @@ namespace UnrealBuildTool
 				"aws-instance-type",
 			};
 
-		public async Task<bool> AddWorkerAsync(Requirements requirements, UnrealBuildAcceleratorHordeConfig hordeConfig, CancellationToken cancellationToken, int activeCores)
+		public async Task<bool> AddWorkerAsync(Requirements requirements, UnrealBuildAcceleratorHordeConfig hordeConfig, int activeCores, CancellationToken cancellationToken)
 		{
 			if (_client == null)
 			{
@@ -390,7 +390,7 @@ namespace UnrealBuildTool
 					Port = ubaPort,
 					ProxyPort = ubaProxyPort,
 				};
-				worker.BackgroundTask = RunWorkerAsync(worker, lease, locator, exeName, workerLogger, hordeConfig, _cancellationTokenSource.Token, agentName, desc);
+				worker.BackgroundTask = RunWorkerAsync(worker, lease, locator, exeName, workerLogger, hordeConfig, agentName, desc, _cancellationTokenSource.Token);
 				lock (_workers)
 				{
 					_workers.Add(worker);
@@ -530,7 +530,7 @@ namespace UnrealBuildTool
 			return result.AccessToken;
 		}
 
-		async Task RunWorkerAsync(Worker self, IComputeLease lease, BlobLocator tool, string executable, ILogger logger, UnrealBuildAcceleratorHordeConfig hordeConfig, CancellationToken cancellationToken, string agentName, string desc)
+		async Task RunWorkerAsync(Worker self, IComputeLease lease, BlobLocator tool, string executable, ILogger logger, UnrealBuildAcceleratorHordeConfig hordeConfig, string agentName, string desc, CancellationToken cancellationToken)
 		{
 			logger.LogDebug("Running worker task..");
 			try
@@ -865,7 +865,7 @@ namespace UnrealBuildTool
 				return null;
 			}
 
-			return hordeSession.RequestCacheServer(cancellationToken, _hordeConfig.CacheDesiredConnectionCount);
+			return hordeSession.RequestCacheServer(_hordeConfig.CacheDesiredConnectionCount, cancellationToken);
 		}
 
 		public void Start(IUbaAgentCoordinatorScheduler scheduler, Func<LinkedAction, bool> canRunRemotely)
@@ -958,7 +958,7 @@ namespace UnrealBuildTool
 						}
 
 						Requirements requirements = GetRequirements(_hordeConfig);
-						if (!await hordeSession.AddWorkerAsync(requirements, _hordeConfig, _cancellationToken, activeCores))
+						if (!await hordeSession.AddWorkerAsync(requirements, _hordeConfig, activeCores, _cancellationToken))
 						{
 							await Task.Delay(500); // Sleep a little bit to make sure previous horde status is visible
 							string status = (queuedAgents + activeAgents) == 0 ? " - No agents available" : " - No additional agents available";
