Fix issue CA1066 (Implement IEquatable when overriding Equals) in UBT.

diff --git a/Engine/Source/Programs/UnrealBuildTool/Configuration/UEBuildTarget.cs b/Engine/Source/Programs/UnrealBuildTool/Configuration/UEBuildTarget.cs
--- a/Engine/Source/Programs/UnrealBuildTool/Configuration/UEBuildTarget.cs
+++ b/Engine/Source/Programs/UnrealBuildTool/Configuration/UEBuildTarget.cs
@@ -26,7 +26,7 @@ namespace UnrealBuildTool
 	/// The platform we're building for
 	/// </summary>
 	[Serializable, TypeConverter(typeof(UnrealTargetPlatformTypeConverter))]
-	public partial struct UnrealTargetPlatform : ISerializable
+	public partial struct UnrealTargetPlatform : ISerializable, IEquatable<UnrealTargetPlatform>
 	{
 		#region Private/boilerplate
 
@@ -130,6 +130,10 @@ namespace UnrealBuildTool
 
 			return Id == ((UnrealTargetPlatform)B).Id;
 		}
+
+		/// <inheritdoc/>
+		public readonly bool Equals(UnrealTargetPlatform other) => this == other;
+
 		/// <summary>
 		/// 
 		/// </summary>
@@ -356,7 +360,7 @@ namespace UnrealBuildTool
 	/// <summary>
 	/// Platform groups
 	/// </summary>
-	public partial struct UnrealPlatformGroup
+	public partial struct UnrealPlatformGroup : IEquatable<UnrealPlatformGroup>
 	{
 		#region Private/boilerplate
 		// internal concrete name of the group
@@ -446,6 +450,10 @@ namespace UnrealBuildTool
 
 			return Id == ((UnrealPlatformGroup)B).Id;
 		}
+
+		/// <inheritdoc/>
+		public readonly bool Equals(UnrealPlatformGroup other) => this == other;
+
 		/// <summary>
 		/// 
 		/// </summary>
@@ -573,7 +581,7 @@ namespace UnrealBuildTool
 	/// The architecture we're building for
 	/// </summary>
 	[Serializable, TypeConverter(typeof(UnrealArchPlatformTypeConverter))]
-	public partial struct UnrealArch : ISerializable
+	public partial struct UnrealArch : ISerializable, IEquatable<UnrealArch>
 	{
 		#region Private/boilerplate
 
@@ -674,6 +682,10 @@ namespace UnrealBuildTool
 
 			return Id == ((UnrealArch)B).Id;
 		}
+
+		/// <inheritdoc/>
+		public readonly bool Equals(UnrealArch other) => this == other;
+
 		/// <summary>
 		/// 
 		/// </summary>
diff --git a/Engine/Source/Programs/UnrealBuildTool/System/RestrictedFolders.cs b/Engine/Source/Programs/UnrealBuildTool/System/RestrictedFolders.cs
--- a/Engine/Source/Programs/UnrealBuildTool/System/RestrictedFolders.cs
+++ b/Engine/Source/Programs/UnrealBuildTool/System/RestrictedFolders.cs
@@ -9,7 +9,7 @@ namespace UnrealBuildTool
 	/// <summary>
 	/// Names of restricted folders. Note: The name of each entry is used to search for/create folders
 	/// </summary>
-	public partial struct RestrictedFolder
+	public partial struct RestrictedFolder : IEquatable<RestrictedFolder>
 	{
 		/// <summary>
 		/// Unique Id for this folder
@@ -87,6 +87,9 @@ namespace UnrealBuildTool
 			return Other is RestrictedFolder folder && Id == folder.Id;
 		}
 
+		/// <inheritdoc/>
+		public readonly bool Equals(RestrictedFolder other) => this == other;
+
 		/// <summary>
 		/// Gets a hash code for this object
 		/// </summary>
