Selectively disable issue CA1056 (URI properties should not be strings) in UBT.

diff --git a/Engine/Source/Programs/UnrealBuildTool/System/BuildVersion.cs b/Engine/Source/Programs/UnrealBuildTool/System/BuildVersion.cs
--- a/Engine/Source/Programs/UnrealBuildTool/System/BuildVersion.cs
+++ b/Engine/Source/Programs/UnrealBuildTool/System/BuildVersion.cs
@@ -310,7 +310,9 @@ namespace UnrealBuildTool
 		public string? BranchName => Inner.BranchName;
 
 		public string? BuildVersionString => Inner.BuildVersionString;
+#pragma warning disable CA1056 // URI properties should not be strings - can't break user API
 		public string? BuildURL => Inner.BuildURL;
+#pragma warning restore CA1056
 
 #pragma warning restore C1591
 		#endregion
