Selectively disable issue CA1054 (URI parameters should not be strings) in UBT.

diff --git a/Engine/Source/Programs/UnrealBuildTool/System/PluginReferenceDescriptor.cs b/Engine/Source/Programs/UnrealBuildTool/System/PluginReferenceDescriptor.cs
--- a/Engine/Source/Programs/UnrealBuildTool/System/PluginReferenceDescriptor.cs
+++ b/Engine/Source/Programs/UnrealBuildTool/System/PluginReferenceDescriptor.cs
@@ -95,7 +95,9 @@ namespace UnrealBuildTool
 		/// <param name="InName">Name of the plugin</param>
 		/// <param name="InMarketplaceURL">The marketplace URL for plugins which are not installed</param>
 		/// <param name="bInEnabled">Whether the plugin is enabled</param>
+#pragma warning disable CA1054 // URI parameters should not be strings - URL is not processed, so does not need to be modelled or validated
 		public PluginReferenceDescriptor(string InName, string? InMarketplaceURL, bool bInEnabled)
+#pragma warning restore CA1054
 		{
 			Name = InName;
 			MarketplaceURL = InMarketplaceURL;
