Fix issue CA1050 (Declare types in namespaces) in UBT.

diff --git a/Engine/Source/Programs/UnrealBuildTool/System/HashSetPool.cs b/Engine/Source/Programs/UnrealBuildTool/System/HashSetPool.cs
--- a/Engine/Source/Programs/UnrealBuildTool/System/HashSetPool.cs
+++ b/Engine/Source/Programs/UnrealBuildTool/System/HashSetPool.cs
@@ -4,6 +4,8 @@ using Microsoft.Extensions.ObjectPool;
 using System;
 using System.Collections.Generic;
 
+namespace UnrealBuildTool;
+
 // Usage:
 //
 // using var rented = HashSetPool<string>.Rent(estimate);
