Fix issue CA1045 (Do not pass types by reference) in UBT.

diff --git a/Engine/Source/Programs/AutomationTool/Scripts/FinalizeInstalledBuild.cs b/Engine/Source/Programs/AutomationTool/Scripts/FinalizeInstalledBuild.cs
--- a/Engine/Source/Programs/AutomationTool/Scripts/FinalizeInstalledBuild.cs
+++ b/Engine/Source/Programs/AutomationTool/Scripts/FinalizeInstalledBuild.cs
@@ -141,7 +141,7 @@ namespace AutomationTool
 				}
 			}
 
-			UnrealBuildTool.InstalledPlatformInfo.WriteConfigFileEntries(InstalledConfigs, ref IniLines);
+			UnrealBuildTool.InstalledPlatformInfo.WriteConfigFileEntries(InstalledConfigs, IniLines);
 
 			if (!String.IsNullOrEmpty(AnalyticsTypeOverride))
 			{
diff --git a/Engine/Source/Programs/UnrealBuildTool/System/InstalledPlatformInfo.cs b/Engine/Source/Programs/UnrealBuildTool/System/InstalledPlatformInfo.cs
--- a/Engine/Source/Programs/UnrealBuildTool/System/InstalledPlatformInfo.cs
+++ b/Engine/Source/Programs/UnrealBuildTool/System/InstalledPlatformInfo.cs
@@ -392,7 +392,7 @@ namespace UnrealBuildTool
 		/// </summary>
 		/// <param name="Configs"></param>
 		/// <param name="OutEntries"></param>
-		public static void WriteConfigFileEntries(List<InstalledPlatformConfiguration> Configs, ref List<string> OutEntries)
+		public static void WriteConfigFileEntries(List<InstalledPlatformConfiguration> Configs, List<string> OutEntries)
 		{
 			// Write config section header
 			OutEntries.Add("[InstalledPlatforms]");
@@ -400,11 +400,11 @@ namespace UnrealBuildTool
 
 			foreach (InstalledPlatformConfiguration Config in Configs)
 			{
-				WriteConfigFileEntry(Config, ref OutEntries);
+				WriteConfigFileEntry(Config, OutEntries);
 			}
 		}
 
-		private static void WriteConfigFileEntry(InstalledPlatformConfiguration Config, ref List<string> OutEntries)
+		private static void WriteConfigFileEntry(InstalledPlatformConfiguration Config, List<string> OutEntries)
 		{
 			string ConfigDescription = "+InstalledPlatformConfigurations=(";
 			ConfigDescription += String.Format("PlatformName=\"{0}\", ", Config.Platform.ToString());
diff --git a/Engine/Source/Programs/UnrealBuildTool/System/Utils.cs b/Engine/Source/Programs/UnrealBuildTool/System/Utils.cs
--- a/Engine/Source/Programs/UnrealBuildTool/System/Utils.cs
+++ b/Engine/Source/Programs/UnrealBuildTool/System/Utils.cs
@@ -788,27 +788,6 @@ namespace UnrealBuildTool
 			return RelativePath;
 		}
 
-		/// <summary>
-		/// Backspaces the specified number of characters, then displays a progress percentage value to the console
-		/// </summary>
-		/// <param name="Numerator">Progress numerator</param>
-		/// <param name="Denominator">Progress denominator</param>
-		/// <param name="NumCharsToBackspaceOver">Number of characters to backspace before writing the text.  This value will be updated with the length of the new progress string.  The first time progress is displayed, you should pass 0 for this value.</param>
-		public static void DisplayProgress(int Numerator, int Denominator, ref int NumCharsToBackspaceOver)
-		{
-			// Backspace over previous progress value
-			while (NumCharsToBackspaceOver-- > 0)
-			{
-				Console.Write("\b");
-			}
-
-			// Display updated progress string and keep track of how long it was
-			float ProgressValue = Denominator > 0 ? ((float)Numerator / (float)Denominator) : 1.0f;
-			string ProgressString = String.Format("{0}%", Math.Round(ProgressValue * 100.0f));
-			NumCharsToBackspaceOver = ProgressString.Length;
-			Console.Write(ProgressString);
-		}
-
 		/*
 		 * Read and write classes with xml specifiers
 		 */
