Selectively disable issue CA1044 (Properties should not be write only) on a deprecated property that forms part of the licensee-facing API.

diff --git a/Engine/Source/Programs/UnrealBuildTool/Configuration/ModuleRules.Obsolete.cs b/Engine/Source/Programs/UnrealBuildTool/Configuration/ModuleRules.Obsolete.cs
--- a/Engine/Source/Programs/UnrealBuildTool/Configuration/ModuleRules.Obsolete.cs
+++ b/Engine/Source/Programs/UnrealBuildTool/Configuration/ModuleRules.Obsolete.cs
@@ -22,6 +22,7 @@ namespace UnrealBuildTool
 		/// are used, and checks that source files include their matching header first.
 		/// </summary>
 		[Obsolete("Deprecated in UE5.2 - Use IWYUSupport instead.")]
+#pragma warning disable CA1044 // Properties should not be write only - ignore for deprecated property on public API
 		public bool bEnforceIWYU
 		{
 			set
@@ -32,6 +33,7 @@ namespace UnrealBuildTool
 				}
 			}
 		}
+#pragma warning restore CA1044
 
 		/// <summary>
 		/// How to treat deterministic warnings (experimental).
