Fix issue CA1001 (Types that own disposable fields should be disposable) in UBT.

diff --git a/Engine/Source/Programs/UnrealBuildTool/Executors/ActionLogger.cs b/Engine/Source/Programs/UnrealBuildTool/Executors/ActionLogger.cs
--- a/Engine/Source/Programs/UnrealBuildTool/Executors/ActionLogger.cs
+++ b/Engine/Source/Programs/UnrealBuildTool/Executors/ActionLogger.cs
@@ -22,7 +22,7 @@ namespace UnrealBuildTool
 	/// <param name="AdditionalDescription">Additional description of action</param>
 	record ExecuteResults(List<string> LogLines, int ExitCode, TimeSpan ExecutionTime, TimeSpan ProcessorTime, long PeakMemoryUsed, LinkedAction Action, string? AdditionalDescription = null);
 
-	class ActionLogger
+	class ActionLogger : IDisposable
 	{
 		/// <summary>
 		/// Add target names for each action executed
@@ -114,6 +114,8 @@ namespace UnrealBuildTool
 		/// </summary>
 		protected readonly ExecuteResults?[] Results;
 
+		private bool _disposedValue;
+
 		public ActionLogger(int actionCount, string progressWriterText, Action<string> writeToolOutput, System.Action flushToolOutput, ILogger logger)
 		{
 			Results = new ExecuteResults[actionCount];
@@ -136,7 +138,6 @@ namespace UnrealBuildTool
 			}
 		}
 
-
 		private static int s_previousLineLength = -1;
 
 		/// <summary>
@@ -384,6 +385,26 @@ namespace UnrealBuildTool
 			}
 		}
 
+		public void Dispose()
+		{
+			Dispose(disposing: true);
+			GC.SuppressFinalize(this);
+		}
+
+		protected virtual void Dispose(bool disposing)
+		{
+			if (!_disposedValue)
+			{
+				if (disposing)
+				{
+					ProgressWriter.Dispose();
+					ProcessGroup.Dispose();
+				}
+
+				_disposedValue = true;
+			}
+		}
+
 		private void LogActions2()
 		{
 			for (; ; )
diff --git a/Engine/Source/Programs/UnrealBuildTool/Executors/ImmediateActionQueue.cs b/Engine/Source/Programs/UnrealBuildTool/Executors/ImmediateActionQueue.cs
--- a/Engine/Source/Programs/UnrealBuildTool/Executors/ImmediateActionQueue.cs
+++ b/Engine/Source/Programs/UnrealBuildTool/Executors/ImmediateActionQueue.cs
@@ -284,6 +284,8 @@ namespace UnrealBuildTool
 		/// </summary>
 		private readonly IActionArtifactCache? _actionArtifactCache;
 
+		private bool _disposedValue;
+
 		/// <summary>
 		/// Construct a new instance of the action queue
 		/// </summary>
@@ -771,17 +773,22 @@ namespace UnrealBuildTool
 			}
 		}
 
-		/// <summary>
-		/// Dispose the object
-		/// </summary>
-		public void Dispose()
+		protected override void Dispose(bool disposing)
 		{
-			Console.CancelKeyPress -= CancelKeyPress;
-			_cpuUtilizationTimer?.Dispose();
-			_cpuUtilizationTimer = null;
-			CancellationTokenSource.Dispose();
-			ProcessGroup.Dispose();
-			ProgressWriter.Dispose();
+			base.Dispose(disposing);
+
+			if (!_disposedValue)
+			{
+				if (disposing)
+				{
+					Console.CancelKeyPress -= CancelKeyPress;
+					_cpuUtilizationTimer?.Dispose();
+					_cpuUtilizationTimer = null;
+					CancellationTokenSource.Dispose();
+				}
+
+				_disposedValue = true;
+			}
 		}
 
 		/// <summary>
diff --git a/Engine/Source/Programs/UnrealBuildTool/Executors/UnrealBuildAccelerator/UBAExecutor.cs b/Engine/Source/Programs/UnrealBuildTool/Executors/UnrealBuildAccelerator/UBAExecutor.cs
--- a/Engine/Source/Programs/UnrealBuildTool/Executors/UnrealBuildAccelerator/UBAExecutor.cs
+++ b/Engine/Source/Programs/UnrealBuildTool/Executors/UnrealBuildAccelerator/UBAExecutor.cs
@@ -667,7 +667,7 @@ namespace UnrealBuildTool
 		{
 			_ubaStartTimeUtc = DateTime.UtcNow;
 
-			ActionLogger actionLogger = new(inputActions.Count(), "Compiling C++ source code...", x => WriteToolOutput(x), () => FlushToolOutput(), logger)
+			using ActionLogger actionLogger = new(inputActions.Count(), "Compiling C++ source code...", x => WriteToolOutput(x), () => FlushToolOutput(), logger)
 			{
 				ShowCompilationTimes = bShowCompilationTimes,
 				ShowCPUUtilization = bShowCPUUtilization,
diff --git a/Engine/Source/Programs/UnrealBuildTool/System/Telemetry.cs b/Engine/Source/Programs/UnrealBuildTool/System/Telemetry.cs
--- a/Engine/Source/Programs/UnrealBuildTool/System/Telemetry.cs
+++ b/Engine/Source/Programs/UnrealBuildTool/System/Telemetry.cs
@@ -194,13 +194,14 @@ namespace UnrealBuildTool
 		}
 	}
 
-	internal class TelemetryEndpoint
+	internal sealed class TelemetryEndpoint : IDisposable
 	{
 		readonly DataRouterTelemetryService _service;
 		readonly HttpClient _httpClient;
 		readonly Uri _baseAddress;
 		readonly string _appId;
 		readonly Lazy<TelemetryMetadata> _metadata;
+		bool _disposedValue;
 
 		public TelemetryEndpoint(HttpClient httpClient, Uri baseAddress, string appId, Lazy<TelemetryMetadata> metadata)
 		{
@@ -214,6 +215,17 @@ namespace UnrealBuildTool
 		public async Task FlushEventsAsync() => await _service.FlushEventsAsync();
 
 		public void RecordEvent(TelemetryEvent eventData) => _service.RecordEvent(eventData.CloneWithMetadata(_appId, _metadata.Value));
+
+		public void Dispose()
+		{
+			if (!_disposedValue)
+			{
+				_service.Dispose();
+				_disposedValue = true;
+			}
+
+			GC.SuppressFinalize(this);
+		}
 	}
 
 	internal readonly struct TelemetryMetadata
@@ -436,15 +448,19 @@ namespace UnrealBuildTool
 		{
 			if (!_disposed)
 			{
+				FlushEvents();
+
 				if (disposing)
 				{
-					// dispose managed state (managed objects)
+					foreach (TelemetryEndpoint endpoint in _endpoints.Values)
+					{
+						endpoint.Dispose();
+					}
 				}
 
 				// free unmanaged resources (unmanaged objects) and override finalizer
 				// set large fields to null
 
-				FlushEvents();
 				_disposed = true;
 			}
 		}
