Fix issue CA1027 (Mark enums with FlagsAttribute) and selectively disable issue CA1028 (Enum storage should be Int32) in UBT, to correctly represent the underlying code analyzer settings.

diff --git a/Engine/Source/Programs/UnrealBuildTool/Platform/Windows/PVSToolChain.cs b/Engine/Source/Programs/UnrealBuildTool/Platform/Windows/PVSToolChain.cs
--- a/Engine/Source/Programs/UnrealBuildTool/Platform/Windows/PVSToolChain.cs
+++ b/Engine/Source/Programs/UnrealBuildTool/Platform/Windows/PVSToolChain.cs
@@ -22,7 +22,10 @@ namespace UnrealBuildTool
 	/// <summary>
 	/// Flags for the PVS analyzer mode
 	/// </summary>
+	[Flags]
+#pragma warning disable CA1028 // uint data type more closely matches the underlying PVS-Studio value
 	public enum PVSAnalysisModeFlags : uint
+#pragma warning restore CA1028
 	{
 		/// <summary>
 		/// No flags
