Selectively disable issue CA1019 (Define accessors for attribute arguments) in UBT, as it doesn't apply well to our use-case.

diff --git a/Engine/Source/Programs/UnrealBuildTool/Configuration/CompileWarnings/WarningLevelDefaultAttribute.cs b/Engine/Source/Programs/UnrealBuildTool/Configuration/CompileWarnings/WarningLevelDefaultAttribute.cs
--- a/Engine/Source/Programs/UnrealBuildTool/Configuration/CompileWarnings/WarningLevelDefaultAttribute.cs
+++ b/Engine/Source/Programs/UnrealBuildTool/Configuration/CompileWarnings/WarningLevelDefaultAttribute.cs
@@ -142,7 +142,9 @@ namespace UnrealBuildTool.Configuration.CompileWarnings
 		/// <param name="compiler">The compiler which this attribute should be applied.</param>
 		/// <param name="defaultLevel">The default warning level to apply for this bound.</param>
 		/// <param name="context">The context in which to apply this default.</param>
+#pragma warning disable CA1019 // Define accessors for attribute arguments
 		public CompilerWarningLevelDefaultAttribute(WindowsCompiler compiler, WarningLevel defaultLevel = WarningLevel.Default, InitializationContext context = InitializationContext.Any) : base(defaultLevel, context)
+#pragma warning restore CA1019 // Define accessors for attribute arguments
 		{
 			Compiler = compiler;
 		}
@@ -194,7 +196,9 @@ namespace UnrealBuildTool.Configuration.CompileWarnings
 		/// </summary>
 		/// <param name="defaultLevel">The default warning level to apply for this bound.</param>
 		/// <param name="context">The context in which to apply this default.</param>
+#pragma warning disable CA1019 // Define accessors for attribute arguments
 		public IntelCompilerWarningLevelDefaultAttribute(WarningLevel defaultLevel = WarningLevel.Default, InitializationContext context = InitializationContext.Any) : base(WindowsCompiler.Intel, defaultLevel, context) { }
+#pragma warning restore CA1019 // Define accessors for attribute arguments
 
 		public override WarningLevel GetDefaultLevel(BuildSystemContext? buildSystemContext = null)
 		{
@@ -225,7 +229,9 @@ namespace UnrealBuildTool.Configuration.CompileWarnings
 		/// </summary>
 		/// <param name="defaultLevel">The default warning level to apply for this bound.</param>
 		/// <param name="context">The context in which to apply this default.</param>
+#pragma warning disable CA1019 // Define accessors for attribute arguments
 		public MSVCCompilerWarningLevelDefaultAttribute(WarningLevel defaultLevel = WarningLevel.Default, InitializationContext context = InitializationContext.Any) : base(WindowsCompiler.VisualStudio2022, defaultLevel, context) { }
+#pragma warning restore CA1019 // Define accessors for attribute arguments
 
 		public override WarningLevel GetDefaultLevel(BuildSystemContext? buildSystemContext = null)
 		{
@@ -257,7 +263,9 @@ namespace UnrealBuildTool.Configuration.CompileWarnings
 		/// </summary>
 		/// <param name="defaultLevel">The default warning level to apply for this bound.</param>
 		/// <param name="context">The context in which to apply this default.</param>
+#pragma warning disable CA1019 // Define accessors for attribute arguments
 		public ClangNonIntelSpecializedCompilerWarningLevelDefaultAttribute(WarningLevel defaultLevel = WarningLevel.Default, InitializationContext context = InitializationContext.Any) : base(WindowsCompiler.Clang, defaultLevel, context) { }
+#pragma warning restore CA1019 // Define accessors for attribute arguments
 
 		public override WarningLevel GetDefaultLevel(BuildSystemContext? buildSystemContext = null)
 		{
