UnrealBuildTool: VFS fixes for apple toolchains

#rnx

diff --git a/Engine/Source/Programs/UnrealBuildTool/Platform/IOS/IOSToolChain.cs b/Engine/Source/Programs/UnrealBuildTool/Platform/IOS/IOSToolChain.cs
--- a/Engine/Source/Programs/UnrealBuildTool/Platform/IOS/IOSToolChain.cs
+++ b/Engine/Source/Programs/UnrealBuildTool/Platform/IOS/IOSToolChain.cs
@@ -307,11 +307,12 @@ namespace UnrealBuildTool
 			Arguments.Add(FormatArchitectureArg(LinkEnvironment.Architectures));
 
 			DirectoryReference SDKPath = Settings.GetSDKPath(LinkEnvironment.Architecture);
-			Arguments.Add($" -isysroot \"{SDKPath}\"");
+			Arguments.Add($"-isysroot \"{NormalizeCommandLinePath(SDKPath, LinkEnvironment.RootPaths)}\"");
 
 			// needed to find swift directories for some versions of the SDK
 			string PlatformName = LinkEnvironment.Architecture == UnrealArch.IOSSimulator ? Settings.SimulatorPlatformDirName : Settings.PlatformDirName;
-			Arguments.Add($"-L{ApplePlatformSDK.GetToolchainDirectory()}/../lib/swift/{PlatformName}");
+			DirectoryReference libSwiftDir = DirectoryReference.Combine(ApplePlatformSDK.GetToolchainDirectory().ParentDirectory!, "lib", "swift", PlatformName);
+			Arguments.Add($"-L{NormalizeCommandLinePath(libSwiftDir, LinkEnvironment.RootPaths)}");
 
 			if (!LinkEnvironment.bIsBuildingLibrary && !LinkEnvironment.bIsBuildingDLL)
 			{
diff --git a/Engine/Source/Programs/UnrealBuildTool/Platform/Mac/MacToolChain.cs b/Engine/Source/Programs/UnrealBuildTool/Platform/Mac/MacToolChain.cs
--- a/Engine/Source/Programs/UnrealBuildTool/Platform/Mac/MacToolChain.cs
+++ b/Engine/Source/Programs/UnrealBuildTool/Platform/Mac/MacToolChain.cs
@@ -220,14 +220,15 @@ namespace UnrealBuildTool
 			}
 		}
 
-		void AddFrameworkToLinkCommand(List<string> LinkArguments, string FrameworkName, string Arg = "-framework")
+		void AddFrameworkToLinkCommand(LinkEnvironment linkEnvironment, List<string> LinkArguments, string FrameworkName, string Arg = "-framework")
 		{
 			if (FrameworkName.EndsWith(".framework"))
 			{
-				LinkArguments.Add("-F \"" + Path.GetDirectoryName(Path.GetFullPath(FrameworkName)) + "\"");
+				FileReference frameworkPath = new FileReference(FrameworkName);
+				LinkArguments.Add($"-F \"{NormalizeCommandLinePath(frameworkPath.Directory, linkEnvironment.RootPaths)}\"");
 				FrameworkName = Path.GetFileNameWithoutExtension(FrameworkName);
 			}
-			LinkArguments.Add(Arg + " \"" + FrameworkName + "\"");
+			LinkArguments.Add($"{Arg} \"{FrameworkName}\"");
 		}
 
 		protected override void GetLinkArguments_Sanitizers(LinkEnvironment linkEnvironment, List<string> arguments)
@@ -248,7 +249,7 @@ namespace UnrealBuildTool
 			base.GetLinkArguments_Global(LinkEnvironment, Arguments);
 			// Pass through architecture and OS info		
 			Arguments.Add(FormatArchitectureArg(LinkEnvironment.Architectures));
-			Arguments.Add(String.Format("-isysroot \"{0}\"", NormalizeCommandLinePath(Settings.GetSDKPath(LinkEnvironment.Architecture), LinkEnvironment.RootPaths)));
+			Arguments.Add($"-isysroot \"{NormalizeCommandLinePath(Settings.GetSDKPath(LinkEnvironment.Architecture), LinkEnvironment.RootPaths)}\"");
 			Arguments.Add($"-target {Settings.GetTargetTuple(LinkEnvironment.Architecture, Target!.Platform, Target!.Type)}");
 			Arguments.Add("-dead_strip");
 
@@ -830,7 +831,7 @@ namespace UnrealBuildTool
 					}
 					else if (AdditionalLibrary.Contains(".framework/"))
 					{
-						LinkArguments.Add($"\"{NormalizeCommandLinePath(FileReference.FromString(AdditionalLibrary), LinkEnvironment.RootPaths)}\"");
+						LinkArguments.Add($"\"{NormalizeCommandLinePath(FileItem.GetItemByPath(AdditionalLibrary), LinkEnvironment.RootPaths)}\"");
 					}
 					else if (String.IsNullOrEmpty(Path.GetDirectoryName(AdditionalLibrary)) && String.IsNullOrEmpty(Path.GetExtension(AdditionalLibrary)))
 					{
@@ -849,7 +850,7 @@ namespace UnrealBuildTool
 
 				foreach (string AdditionalLibrary in LinkEnvironment.DelayLoadDLLs)
 				{
-					LinkArguments.Add(String.Format("-weak_library \"{0}\"", Path.GetFullPath(AdditionalLibrary)));
+					LinkArguments.Add($"-weak_library \"{NormalizeCommandLinePath(FileItem.GetItemByPath(AdditionalLibrary), LinkEnvironment.RootPaths)}\"");
 
 					if (!bCanUseSimpleRPaths)
 					{
@@ -889,7 +890,7 @@ namespace UnrealBuildTool
 
 				foreach (KeyValuePair<string, bool> Framework in AllFrameworks)
 				{
-					AddFrameworkToLinkCommand(LinkArguments, Framework.Key, Framework.Value ? "-weak_framework" : "-framework");
+					AddFrameworkToLinkCommand(LinkEnvironment, LinkArguments, Framework.Key, Framework.Value ? "-weak_framework" : "-framework");
 					AddLibraryPathToRPaths(Framework.Key, AbsolutePath, ref RPaths, LinkArguments, bIsBuildingAppBundle, Logger);
 				}
 
diff --git a/Engine/Source/Programs/UnrealBuildTool/ToolChain/AppleToolChain.cs b/Engine/Source/Programs/UnrealBuildTool/ToolChain/AppleToolChain.cs
--- a/Engine/Source/Programs/UnrealBuildTool/ToolChain/AppleToolChain.cs
+++ b/Engine/Source/Programs/UnrealBuildTool/ToolChain/AppleToolChain.cs
@@ -113,9 +113,6 @@ namespace UnrealBuildTool
 			SDKVersion = ApplePlatformSDK.GetPlatformSDKVersion(OSPrefix) ?? "";
 			SDKVersionFloat = ApplePlatformSDK.GetPlatformSDKVersionFloat(OSPrefix);
 
-			// Xcode sdk files are not allowed to be transferred over the network, ensure UBA will not do so
-			EpicGames.UBA.Utils.RegisterDisallowedPaths(ApplePlatformSDK.DeveloperDir.FullName);
-
 			TestXcode(bVerbose, Logger);
 		}
 
@@ -266,6 +263,28 @@ namespace UnrealBuildTool
 			AppleExports.GetSwiftIntegrationSettings(ProjectFile, Target == null ? TargetType.Game : Target.Type, Target == null ? UnrealTargetPlatform.Mac : Target.Platform, out bUseSwiftUIMain, out bCreateSwiftBridgingHeader);
 		}
 
+		/// <inheritdoc/>
+		public override void SetUpGlobalEnvironment(ReadOnlyTargetRules Target, CppCompileEnvironment GlobalCompileEnvironment, LinkEnvironment GlobalLinkEnvironment)
+		{
+			base.SetUpGlobalEnvironment(Target, GlobalCompileEnvironment, GlobalLinkEnvironment);
+
+			DirectoryReference XcodeDeveloperDir = ApplePlatformSDK.DeveloperDir;
+			DirectoryReference? XcodeAppDir = XcodeDeveloperDir.ParentDirectory?.ParentDirectory;
+
+			if (XcodeAppDir != null)
+			{
+				GlobalCompileEnvironment.RootPaths.AddFolderName(CppRootPathFolder.PlatformSDK, "Xcode.app");
+				GlobalLinkEnvironment.RootPaths.AddFolderName(CppRootPathFolder.PlatformSDK, "Xcode.app");
+				GlobalCompileEnvironment.RootPaths[CppRootPathFolder.PlatformSDK] = XcodeAppDir;
+				GlobalLinkEnvironment.RootPaths[CppRootPathFolder.PlatformSDK] = XcodeAppDir;
+
+				EpicGames.UBA.Utils.RegisterPathHash(XcodeAppDir.FullName, ApplePlatformSDK.InstalledXcodeVersion.Value ?? "0");
+
+				// Xcode sdk files are not allowed to be transferred over the network, ensure UBA will not do so
+				EpicGames.UBA.Utils.RegisterDisallowedPaths(XcodeAppDir.FullName);
+			}
+		}
+
 		/// <summary>
 		/// Takes an architecture string as provided by UBT for the target and formats it for Clang. Supports
 		/// multiple architectures joined with '+'
@@ -875,15 +894,13 @@ namespace UnrealBuildTool
 			// Temp solution for UE-191350
 			if (AppleClangVersion >= new VersionNumber(15) && AppleClangVersion < new VersionNumber(16))
 			{
-				Arguments.Add(" -ld_classic");
+				Arguments.Add("-ld_classic");
 			}
 
 			if (!LinkEnvironment.bIsBuildingDLL)
 			{
-				// Todo: This is currently only used when linking binary and swift files have not been tested with dylibs.
-				// The right fix is to move the vfs root one step up to be able to handle the ../lib/ path
-				DirectoryReference dir = DirectoryReference.Combine(ApplePlatformSDK.GetToolchainDirectory(), "../lib/swift", Settings.PlatformDirName);
-				Arguments.Add($"-L{dir}");
+				DirectoryReference libSwiftDir = DirectoryReference.Combine(ApplePlatformSDK.GetToolchainDirectory().ParentDirectory!, "lib", "swift", Settings.PlatformDirName);
+				Arguments.Add($"-L{NormalizeCommandLinePath(libSwiftDir, LinkEnvironment.RootPaths)}");
 			}
 
 			Arguments.Add("-L/usr/lib/swift");
