[UBT]
* Changed so uba logger is instantiated earlier in UBT

diff --git a/Engine/Source/Programs/Shared/EpicGames.UBA/ILogger.cs b/Engine/Source/Programs/Shared/EpicGames.UBA/ILogger.cs
--- a/Engine/Source/Programs/Shared/EpicGames.UBA/ILogger.cs
+++ b/Engine/Source/Programs/Shared/EpicGames.UBA/ILogger.cs
@@ -61,15 +61,21 @@ namespace EpicGames.UBA
 		/// Create a ILogger object
 		/// </summary>
 		/// <param name="logger">The Microsoft.Extensions.Logging.ILogger to wrap</param>
+		/// <param name="makeGlobal">Set this log as the one to be used by all sessions</param>
 		/// <returns>The ILogger</returns>
-		public static ILogger CreateLogger(Microsoft.Extensions.Logging.ILogger logger)
+		public static ILogger CreateLogger(Microsoft.Extensions.Logging.ILogger logger, bool makeGlobal)
 		{
-			return new LoggerImpl(logger);
+			return new LoggerImpl(logger, makeGlobal);
 		}
 
 		/// <summary>
 		/// Creates null logger
 		/// </summary>
 		public static ILogger CreateNullLogger() => new NullLogger();
+
+		/// <summary>
+		/// Get the global trace if there is one
+		/// </summary>
+		public static ILogger? GlobalLogger => LoggerImpl.GlobalLoggerImpl;
 	}
 }
\ No newline at end of file
diff --git a/Engine/Source/Programs/Shared/EpicGames.UBA/Impl/LoggerImpl.cs b/Engine/Source/Programs/Shared/EpicGames.UBA/Impl/LoggerImpl.cs
--- a/Engine/Source/Programs/Shared/EpicGames.UBA/Impl/LoggerImpl.cs
+++ b/Engine/Source/Programs/Shared/EpicGames.UBA/Impl/LoggerImpl.cs
@@ -28,15 +28,27 @@ namespace EpicGames.UBA.Impl
 		static extern void DestroyCallbackLogWriter(nint logger);
 		#endregion
 
-		public LoggerImpl(Microsoft.Extensions.Logging.ILogger logger)
+		public LoggerImpl(Microsoft.Extensions.Logging.ILogger logger, bool makeGlobal)
 		{
 			_logger = logger;
 			_beginScopeCallbackDelegate = BeginScope;
 			_endScopeCallbackDelegate = EndScope;
 			_logCallbackDelegate = Log;
 			_handle = CreateCallbackLogWriter(_beginScopeCallbackDelegate, _endScopeCallbackDelegate, _logCallbackDelegate);
+			if (makeGlobal)
+			{
+				try
+				{
+					GlobalLoggerImpl = this;
+				}
+				catch (EntryPointNotFoundException)
+				{
+				}
+			}
 		}
 
+		internal static LoggerImpl? GlobalLoggerImpl { set; get; }
+
 		#region IDisposable
 		~LoggerImpl() => Dispose(false);
 
diff --git a/Engine/Source/Programs/UnrealBuildTool/Executors/UnrealBuildAccelerator/UBAExecutor.cs b/Engine/Source/Programs/UnrealBuildTool/Executors/UnrealBuildAccelerator/UBAExecutor.cs
--- a/Engine/Source/Programs/UnrealBuildTool/Executors/UnrealBuildAccelerator/UBAExecutor.cs
+++ b/Engine/Source/Programs/UnrealBuildTool/Executors/UnrealBuildAccelerator/UBAExecutor.cs
@@ -30,7 +30,6 @@ namespace UnrealBuildTool
 		readonly Task? _initTask;
 		readonly CancellationTokenSource _cancellationSource;
 
-		EpicGames.UBA.ILogger? _ubaLogger;
 		IServer? _server;
 		IStorageServer? _storageServer;
 		ISessionServer? _session;
@@ -76,7 +75,6 @@ namespace UnrealBuildTool
 				}
 
 				_server?.Dispose();
-				_ubaLogger?.Dispose();
 			}
 			base.Dispose(disposing);
 		}
@@ -268,10 +266,10 @@ namespace UnrealBuildTool
 				config.AddValue("Storage", "CompressionLevel", UBAConfig.CompressionLevel);
 			}
 
-			_ubaLogger = EpicGames.UBA.ILogger.CreateLogger(logger);
-			_server = IServer.CreateServer(UBAConfig.MaxWorkers, UBAConfig.SendSize, _ubaLogger, UBAConfig.bUseQuic);
-			_storageServer = IStorageServer.CreateStorageServer(_server, _ubaLogger, new StorageServerCreateInfo(_rootDirRef.FullName, ((ulong)UBAConfig.StoreCapacityGb) * 1000 * 1000 * 1000, !UBAConfig.bStoreRaw, UBAConfig.Zone));
-			using ISessionServerCreateInfo serverCreateInfo = ISessionServerCreateInfo.CreateSessionServerCreateInfo(_storageServer, _server, _ubaLogger, new SessionServerCreateInfo(_rootDirRef.FullName, ubaTraceFile.FullName.Replace('\\', '/'), UBAConfig.bDisableCustomAlloc, false, UBAConfig.bResetCas, UBAConfig.bWriteToDisk, UBAConfig.bDetailedTrace, !UBAConfig.bDisableWaitOnMem, UBAConfig.bAllowKillOnMem, UBAConfig.bStoreObjFilesCompressed));
+			EpicGames.UBA.ILogger ubaLogger = EpicGames.UBA.ILogger.GlobalLogger!;
+			_server = IServer.CreateServer(UBAConfig.MaxWorkers, UBAConfig.SendSize, ubaLogger, UBAConfig.bUseQuic);
+			_storageServer = IStorageServer.CreateStorageServer(_server, ubaLogger, new StorageServerCreateInfo(_rootDirRef.FullName, ((ulong)UBAConfig.StoreCapacityGb) * 1000 * 1000 * 1000, !UBAConfig.bStoreRaw, UBAConfig.Zone));
+			using ISessionServerCreateInfo serverCreateInfo = ISessionServerCreateInfo.CreateSessionServerCreateInfo(_storageServer, _server, ubaLogger, new SessionServerCreateInfo(_rootDirRef.FullName, ubaTraceFile.FullName.Replace('\\', '/'), UBAConfig.bDisableCustomAlloc, false, UBAConfig.bResetCas, UBAConfig.bWriteToDisk, UBAConfig.bDetailedTrace, !UBAConfig.bDisableWaitOnMem, UBAConfig.bAllowKillOnMem, UBAConfig.bStoreObjFilesCompressed));
 			_session = ISessionServer.CreateSessionServer(serverCreateInfo);
 			_trace = _session.GetTrace();
 
diff --git a/Engine/Source/Programs/UnrealBuildTool/UnrealBuildTool.cs b/Engine/Source/Programs/UnrealBuildTool/UnrealBuildTool.cs
--- a/Engine/Source/Programs/UnrealBuildTool/UnrealBuildTool.cs
+++ b/Engine/Source/Programs/UnrealBuildTool/UnrealBuildTool.cs
@@ -446,6 +446,7 @@ namespace UnrealBuildTool
 			}
 
 			// Attach timeline to uba trace
+			using EpicGames.UBA.ILogger ubaLogger = EpicGames.UBA.ILogger.CreateLogger(Logger, true);
 			using EpicGames.UBA.ITrace ubaTrace = EpicGames.UBA.ITrace.Create(traceName, traceFile, true, traceChannel);
 			using TimelineEventScope ubaTimelineEventScope = new((name) => ubaTrace.TaskBegin(name, ""), (id) => ubaTrace.TaskEnd(id));
 
