Selectively disable issue CA1069 (Enums should not have duplicate values) in UBT.

diff --git a/Engine/Source/Programs/UnrealBuildTool/System/ProjectDescriptor.cs b/Engine/Source/Programs/UnrealBuildTool/System/ProjectDescriptor.cs
--- a/Engine/Source/Programs/UnrealBuildTool/System/ProjectDescriptor.cs
+++ b/Engine/Source/Programs/UnrealBuildTool/System/ProjectDescriptor.cs
@@ -39,10 +39,12 @@ namespace UnrealBuildTool
 		/// </summary>
 		LatestPlusOne,
 
+#pragma warning disable CA1069 // Enums should not have duplicate values - having access to Latest is useful
 		/// <summary>
 		/// The latest plugin descriptor version
 		/// </summary>
 		Latest = LatestPlusOne - 1
+#pragma warning restore CA1069
 	}
 
 	/// <summary>
