Fix issue CA2251 (Use String.Equals over String.Compare) in UBT.

diff --git a/Engine/Source/Programs/UnrealBuildTool/Platform/Android/AndroidPlatformSDK.cs b/Engine/Source/Programs/UnrealBuildTool/Platform/Android/AndroidPlatformSDK.cs
--- a/Engine/Source/Programs/UnrealBuildTool/Platform/Android/AndroidPlatformSDK.cs
+++ b/Engine/Source/Programs/UnrealBuildTool/Platform/Android/AndroidPlatformSDK.cs
@@ -304,7 +304,7 @@ namespace UnrealBuildTool
 			// convert r<num>[letter] to hex
 			if (!String.IsNullOrEmpty(StringValue))
 			{
-				if (String.Compare(Hint, "Software", StringComparison.OrdinalIgnoreCase) == 0)
+				if (String.Equals(Hint, "Software", StringComparison.OrdinalIgnoreCase))
 				{
 					// returning a version with a format similar to the ndk version.
 					string VersionString = String.Format("{0}{1:00}{2:00}", StringValue, 0, 0);
