Fix issue CA2249 (Consider using String.Contains instead of String.IndexOf) in UBT.

diff --git a/Engine/Source/Programs/UnrealBuildTool/System/Utils.cs b/Engine/Source/Programs/UnrealBuildTool/System/Utils.cs
--- a/Engine/Source/Programs/UnrealBuildTool/System/Utils.cs
+++ b/Engine/Source/Programs/UnrealBuildTool/System/Utils.cs
@@ -1641,8 +1641,7 @@ namespace UnrealBuildTool
 					CommandLine.Append(' ');
 				}
 
-				int SpaceIdx = Argument.IndexOf(' ');
-				if (SpaceIdx == -1)
+				if (!Argument.Contains(' '))
 				{
 					CommandLine.Append(Argument);
 				}
