Fix or selectively disable issue CA2227 (Collection properties should be read only) in UBT.

diff --git a/Engine/Source/Programs/UnrealBuildTool/Configuration/TestModuleRules.cs b/Engine/Source/Programs/UnrealBuildTool/Configuration/TestModuleRules.cs
--- a/Engine/Source/Programs/UnrealBuildTool/Configuration/TestModuleRules.cs
+++ b/Engine/Source/Programs/UnrealBuildTool/Configuration/TestModuleRules.cs
@@ -539,65 +539,67 @@ namespace UnrealBuildTool
 				set => UsesCatch2Private = value;
 			}
 
+#pragma warning disable CA2227 // Collection properties should be read only - can't break user API
 			/// <summary>
 			/// Set of supported platforms.
 			/// </summary>
-			public HashSet<UnrealTargetPlatform> SupportedPlatforms { get; set; } = new HashSet<UnrealTargetPlatform>() { UnrealTargetPlatform.Win64 };
+			public HashSet<UnrealTargetPlatform> SupportedPlatforms
+			{
+				get;
+				[Obsolete("Modify the existing collection")]
+				set;
+			} = [UnrealTargetPlatform.Win64];
 
-			private Dictionary<UnrealTargetPlatform, string> PlatformTagsPrivate = new Dictionary<UnrealTargetPlatform, string>();
 			/// <summary>
 			/// Per-platform tags.
 			/// </summary>
 			public Dictionary<UnrealTargetPlatform, string> PlatformTags
 			{
-				get => PlatformTagsPrivate;
-				set => PlatformTagsPrivate = value;
-			}
+				get;
+				[Obsolete("Modify the existing collection")]
+				set;
+			} = [];
 
-			private Dictionary<UnrealTargetPlatform, string> PlatformGauntletArgsPrivate = new Dictionary<UnrealTargetPlatform, string>();
 			/// <summary>
 			/// Per-platform gauntlet args.
 			/// </summary>
 			public Dictionary<UnrealTargetPlatform, string> PlatformGauntletArgs
 			{
-				get => PlatformGauntletArgsPrivate;
-				set => PlatformGauntletArgsPrivate = value;
-			}
+				get;
+				[Obsolete("Modify the existing collection")]
+				set;
+			} = [];
 
-			private Dictionary<UnrealTargetPlatform, string> PlatformCompilationExtraArgsPrivate = new Dictionary<UnrealTargetPlatform, string>();
 			/// <summary>
 			/// Per-platform extra compilation arguments.
 			/// </summary>
 			public Dictionary<UnrealTargetPlatform, string> PlatformCompilationExtraArgs
 			{
-				get => PlatformCompilationExtraArgsPrivate;
-				set => PlatformCompilationExtraArgsPrivate = value;
-			}
-
-			private List<UnrealTargetPlatform> PlatformsRunUnsupportedPrivate = new List<UnrealTargetPlatform>() {
-				UnrealTargetPlatform.Android,
-				UnrealTargetPlatform.IOS,
-				UnrealTargetPlatform.TVOS,
-				UnrealTargetPlatform.VisionOS };
+				get;
+				[Obsolete("Modify the existing collection")]
+				set;
+			} = [];
 
 			/// <summary>
 			/// List of platforms that cannot run tests.
 			/// </summary>
 			public List<UnrealTargetPlatform> PlatformsRunUnsupported
 			{
-				get => PlatformsRunUnsupportedPrivate;
-				set => PlatformsRunUnsupportedPrivate = value;
-			}
+				get;
+				[Obsolete("Modify the existing collection")]
+				set;
+			} = [UnrealTargetPlatform.Android, UnrealTargetPlatform.IOS, UnrealTargetPlatform.TVOS, UnrealTargetPlatform.VisionOS];
 
-			private Dictionary<UnrealTargetPlatform, bool> PlatformRunContainerizedPrivate = new Dictionary<UnrealTargetPlatform, bool>();
 			/// <summary>
 			/// Whether or not the test is run inside a Docker container for a given platform.
 			/// </summary>
 			public Dictionary<UnrealTargetPlatform, bool> PlatformRunContainerized
 			{
-				get => PlatformRunContainerizedPrivate;
-				set => PlatformRunContainerizedPrivate = value;
-			}
+				get;
+				[Obsolete("Modify the existing collection")]
+				set;
+			} = [];
+#pragma warning restore CA2227
 		}
 #pragma warning restore 8618
 	}
diff --git a/Engine/Source/Programs/UnrealBuildTool/System/TimingData.cs b/Engine/Source/Programs/UnrealBuildTool/System/TimingData.cs
--- a/Engine/Source/Programs/UnrealBuildTool/System/TimingData.cs
+++ b/Engine/Source/Programs/UnrealBuildTool/System/TimingData.cs
@@ -111,7 +111,7 @@ namespace UnrealBuildTool
 		/// <summary>
 		/// Any sub-events of this event.
 		/// </summary>
-		public Dictionary<string, TimingData> Children { get; set; } = new Dictionary<string, TimingData>();
+		public Dictionary<string, TimingData> Children { get; } = [];
 
 		/// <summary>
 		/// Adds a sub-event (child) to this event if it doesn't exist, or increase the duration by the
