Fix or selectively disable issue CA2214 (Do not call overridable methods in constructors) in UBT.

diff --git a/Engine/Source/Programs/UnrealBuildTool/Platform/Windows/VCEnvironment.cs b/Engine/Source/Programs/UnrealBuildTool/Platform/Windows/VCEnvironment.cs
--- a/Engine/Source/Programs/UnrealBuildTool/Platform/Windows/VCEnvironment.cs
+++ b/Engine/Source/Programs/UnrealBuildTool/Platform/Windows/VCEnvironment.cs
@@ -287,7 +287,7 @@ namespace UnrealBuildTool
 		/// <summary>
 		/// Gets the path to the resource to object converter.
 		/// </summary>
-		protected virtual FileReference GetResourceToObjectConverterToolPath(WindowsCompiler Compiler, UnrealArch Architecture, DirectoryReference ToochainDir)
+		private FileReference GetResourceToObjectConverterToolPath(WindowsCompiler Compiler, UnrealArch Architecture, DirectoryReference ToochainDir)
 		{
 			return FileReference.Combine(GetVCToolPath(ToochainDir, Architecture), "cvtres.exe");
 		}
@@ -295,7 +295,7 @@ namespace UnrealBuildTool
 		/// <summary>
 		/// Gets the path to the resource compiler.
 		/// </summary>
-		protected virtual FileReference GetResourceCompilerToolPath(DirectoryReference WindowsSdkDir, VersionNumber WindowsSdkVersion, ILogger Logger)
+		private FileReference GetResourceCompilerToolPath(DirectoryReference WindowsSdkDir, VersionNumber WindowsSdkVersion, ILogger Logger)
 		{
 			string BuildHostArch = RuntimeInformation.ProcessArchitecture == System.Runtime.InteropServices.Architecture.Arm64 ? "arm64" : "x64";
 			FileReference ResourceCompilerPath = FileReference.Combine(WindowsSdkDir, "bin", WindowsSdkVersion.ToString(), BuildHostArch, "rc.exe");
@@ -333,16 +333,18 @@ namespace UnrealBuildTool
 		/// <summary>
 		/// Return the standard Clang C++ library path.
 		/// </summary>
-		protected virtual DirectoryReference GetClangLibsDir()
+		private DirectoryReference GetClangLibsDir()
 		{
 			VersionNumber ClangVersion = Compiler.IsIntel() ? MicrosoftPlatformSDK.GetClangVersionForIntelCompiler(CompilerPath) : CompilerVersion;
 			return DirectoryReference.Combine(CompilerDir, "lib", "clang", ClangVersion.GetComponent(0).ToString(), "lib", "windows");
 		}
 
+#pragma warning disable CA2214 // Do not call overridable methods in constructors - happens to be safe as of 2025-10-24, impractical to restructure code to avoid
 		/// <summary>
 		/// Return the standard Visual C++ library path.
 		/// </summary>
 		protected virtual DirectoryReference GetToolChainLibsDir() => DirectoryReference.Combine(ToolChainDir, "lib", Architecture.WindowsSystemLibDir);
+#pragma warning restore CA2214
 
 		/// <summary>
 		/// Sets up the standard compile environment for the toolchain
