Fix issue CA2241 (Provide correct arguments to formatting methods) in UBT.

diff --git a/Engine/Source/Programs/UnrealBuildTool/ProjectFiles/VisualStudio/VCProject.cs b/Engine/Source/Programs/UnrealBuildTool/ProjectFiles/VisualStudio/VCProject.cs
--- a/Engine/Source/Programs/UnrealBuildTool/ProjectFiles/VisualStudio/VCProject.cs
+++ b/Engine/Source/Programs/UnrealBuildTool/ProjectFiles/VisualStudio/VCProject.cs
@@ -516,7 +516,7 @@ namespace UnrealBuildTool
 
 			public override string ToString()
 			{
-				return String.Format("{0} {1} {2}", ProjectTarget, Platform, Configuration, Architecture != null ? " " + Architecture : String.Empty);
+				return $"{ProjectTarget} {Platform} {Configuration} {Architecture}".Trim();
 			}
 		}
 
