Fix issue CA2208 (Instantiate argument exceptions correctly) in UBT.

diff --git a/Engine/Source/Programs/UnrealBuildTool/Platform/IOS/IOSCodeSigning.cs b/Engine/Source/Programs/UnrealBuildTool/Platform/IOS/IOSCodeSigning.cs
--- a/Engine/Source/Programs/UnrealBuildTool/Platform/IOS/IOSCodeSigning.cs
+++ b/Engine/Source/Programs/UnrealBuildTool/Platform/IOS/IOSCodeSigning.cs
@@ -1022,7 +1022,7 @@ namespace UnrealBuildTool
 						// Check to see if this is a distribution provision. get-task-allow must be false for distro profiles.
 						// TestProvision.ProvisionedDeviceIDs.Count==0 is not a valid check as ad-hoc distro profiles do list devices.
 						bool bDistroProv = !TestProvision.bDebug;
-						if (TestProvision != null && !bDistroProv)
+						if (!bDistroProv)
 						{
 							_logger?.LogInformation(
 								"  .. Failed distribution check (mode={ForDistribution}, get-task-allow={Debug}, #devices={NumDevices})",
@@ -1036,29 +1036,29 @@ namespace UnrealBuildTool
 					{
 						if (bCheckDistro)
 						{
-							bool bPassesDebugCheck = TestProvision?.bDebug ?? false;
+							bool bPassesDebugCheck = TestProvision.bDebug;
 							if (!bPassesDebugCheck)
 							{
 								_logger?.LogInformation(
 									" .. Failed debugging check (mode={ForDistribution}, get-task-allow={Debug}, #devices={NumDevices})",
 									CodeSigningConfig.bForDistribution,
-									TestProvision?.bDebug ?? false,
-									TestProvision?.ProvisionedDeviceIDs?.Count ?? 0);
+									TestProvision.bDebug,
+									TestProvision.ProvisionedDeviceIDs?.Count ?? 0);
 								continue;
 							}
 						}
-						else if (!TestProvision?.bDebug ?? true)
+						else if (!TestProvision.bDebug)
 						{
 							CodeSigningConfig.bForceStripSymbols = true;
 						}
 					}
 
 					DateTime CurrentUTCTime = DateTime.UtcNow;
-					bool bPassesDateCheck = (CurrentUTCTime >= TestProvision?.CreationDate) && (CurrentUTCTime < TestProvision?.ExpirationDate);
+					bool bPassesDateCheck = (CurrentUTCTime >= TestProvision.CreationDate) && (CurrentUTCTime < TestProvision.ExpirationDate);
 					if (!bPassesDateCheck)
 					{
 						_logger?.LogInformation("  .. Failed time period check (valid from {StartDate} to {EndDate}, but UTC time is now {CurrentDate})",
-						TestProvision?.CreationDate, TestProvision?.ExpirationDate, CurrentUTCTime);
+						TestProvision.CreationDate, TestProvision.ExpirationDate, CurrentUTCTime);
 						continue;
 					}
 
@@ -1070,10 +1070,6 @@ namespace UnrealBuildTool
 						X509Certificate2? Cert;
 						try
 						{
-							if (TestProvision == null)
-							{
-								throw new ArgumentNullException(nameof(TestProvision));
-							}
 							Cert = AppleCodeSign.FindCertificate(TestProvision);
 						}
 						catch (Exception ex)
@@ -1103,8 +1099,8 @@ namespace UnrealBuildTool
 					_logger?.LogInformation(
 						"  Picked '{Filename}' with AppID '{ApplicationIdentifier}' and Name '{Name}' as a matching provision for the game '{GameName}'",
 						DebugName,
-						TestProvision?.ApplicationIdentifier,
-						TestProvision?.ProvisionName,
+						TestProvision.ApplicationIdentifier,
+						TestProvision.ProvisionName,
 						GameName);
 					return Key;
 				}
