Fix issue CA2008 (Do not create tasks without passing a TaskScheduler) in UBT.

diff --git a/Engine/Source/Programs/UnrealBuildTool/System/UHTExecution.cs b/Engine/Source/Programs/UnrealBuildTool/System/UHTExecution.cs
--- a/Engine/Source/Programs/UnrealBuildTool/System/UHTExecution.cs
+++ b/Engine/Source/Programs/UnrealBuildTool/System/UHTExecution.cs
@@ -1485,7 +1485,7 @@ namespace UnrealBuildTool
 						timestampsFile.Save(Logger);
 					}
 				}
-			});
+			}, TaskScheduler.Current);
 
 			runTask.Start();
 
