Fix issue CA2016 (Forward the CancellationToken parameter to methods that take one) in UBT.

diff --git a/Engine/Source/Programs/UnrealBuildTool/Executors/UnrealBuildAccelerator/UBAAgentCoordinatorHorde.cs b/Engine/Source/Programs/UnrealBuildTool/Executors/UnrealBuildAccelerator/UBAAgentCoordinatorHorde.cs
--- a/Engine/Source/Programs/UnrealBuildTool/Executors/UnrealBuildAccelerator/UBAAgentCoordinatorHorde.cs
+++ b/Engine/Source/Programs/UnrealBuildTool/Executors/UnrealBuildAccelerator/UBAAgentCoordinatorHorde.cs
@@ -212,7 +212,7 @@ namespace UnrealBuildTool
 			// TODO: should drive this off API version reported by server
 			BlobSerializerOptions serializerOptions = BlobSerializerOptions.Create(HordeApiVersion.Initial);
 
-			await using IBlobWriter writer = _storage.CreateBlobWriter(serializerOptions: serializerOptions);
+			await using IBlobWriter writer = _storage.CreateBlobWriter(serializerOptions: serializerOptions, cancellationToken: cancellationToken);
 			DirectoryNode sandbox = new();
 			await sandbox.AddFilesAsync(baseDir, files, writer, cancellationToken: cancellationToken);
 			IHashedBlobRef<DirectoryNode> handle = await writer.WriteBlobAsync(sandbox, cancellationToken);
