Fix issue CA2002 (Do not lock on objects with weak identity) in UBT.

diff --git a/Engine/Source/Programs/UnrealBuildTool/Executors/UnrealBuildAccelerator/UBAExecutor.cs b/Engine/Source/Programs/UnrealBuildTool/Executors/UnrealBuildAccelerator/UBAExecutor.cs
--- a/Engine/Source/Programs/UnrealBuildTool/Executors/UnrealBuildAccelerator/UBAExecutor.cs
+++ b/Engine/Source/Programs/UnrealBuildTool/Executors/UnrealBuildAccelerator/UBAExecutor.cs
@@ -641,12 +641,39 @@ namespace UnrealBuildTool
 			internal AgentCoordinatorScheduler(IScheduler scheduler, IServer server) { _scheduler = scheduler; _server = server; }
 			internal IScheduler? _scheduler;
 			internal IServer? _server;
-			public void Reset() { lock (this) { _scheduler = null; _server = null; } }
-			public bool IsEmpty { get { lock (this) return _scheduler != null ? _scheduler.IsEmpty : true; } }
-			public double GetProcessWeightThatCanRunRemotelyNow() { lock (this) return _scheduler != null ? _scheduler.GetProcessWeightThatCanRunRemotelyNow() : 0; }
+			private readonly object _lockObject = new();
+
+			public void Reset()
+			{
+				lock (_lockObject)
+				{
+					_scheduler = null;
+					_server = null;
+				}
+			}
+
+			public bool IsEmpty
+			{
+				get
+				{
+					lock (_lockObject)
+					{
+						return _scheduler == null || _scheduler.IsEmpty;
+					}
+				}
+			}
+
+			public double GetProcessWeightThatCanRunRemotelyNow()
+			{
+				lock (_lockObject)
+				{
+					return _scheduler != null ? _scheduler.GetProcessWeightThatCanRunRemotelyNow() : 0;
+				}
+			}
+
 			public bool AddClient(string ip, int port, string crypto = "")
 			{
-				lock (this)
+				lock (_lockObject)
 				{
 					if (_server != null)
 					{
diff --git a/Engine/Source/Programs/UnrealBuildTool/System/ProgressWriter.cs b/Engine/Source/Programs/UnrealBuildTool/System/ProgressWriter.cs
--- a/Engine/Source/Programs/UnrealBuildTool/System/ProgressWriter.cs
+++ b/Engine/Source/Programs/UnrealBuildTool/System/ProgressWriter.cs
@@ -36,6 +36,8 @@ namespace UnrealBuildTool
 		/// </summary>
 		int CurrentProgressValue = -1;
 
+		private readonly object LockObject = new();
+
 		/// <summary>
 		/// Constructor
 		/// </summary>
@@ -77,7 +79,7 @@ namespace UnrealBuildTool
 
 			if (ProgressValue != CurrentProgressValue)
 			{
-				lock (this)
+				lock (LockObject)
 				{
 					if (ProgressValue <= CurrentProgressValue)
 					{
diff --git a/Engine/Source/Programs/UnrealBuildTool/System/Utils.cs b/Engine/Source/Programs/UnrealBuildTool/System/Utils.cs
--- a/Engine/Source/Programs/UnrealBuildTool/System/Utils.cs
+++ b/Engine/Source/Programs/UnrealBuildTool/System/Utils.cs
@@ -888,31 +888,6 @@ namespace UnrealBuildTool
 			return (bSuccess);
 		}
 
-		/// <summary>
-		/// Returns true if the specified Process has been created, started and remains valid (i.e. running).
-		/// </summary>
-		/// <param name="p">Process object to test</param>
-		/// <returns>True if valid, false otherwise.</returns>
-		public static bool IsValidProcess(Process p)
-		{
-			// null objects are always invalid
-			if (p == null)
-			{
-				return false;
-			}
-			// due to multithreading on Windows, lock the object
-			lock (p)
-			{
-				// note that this can fail and have a race condition in threads, but the framework throws an exception when this occurs.
-				try
-				{
-					return p.Id != 0;
-				}
-				catch { } // all exceptions can be safely caught and ignored, meaning the process is not started or has stopped.
-			}
-			return false;
-		}
-
 		/// <summary>
 		/// Removes multi-dot extensions from a filename (i.e. *.automation.csproj)
 		/// </summary>
