Fix issue CA1847 (Use String.Contains(char) instead of String.Contains(string) with single characters) in UBT.

diff --git a/Engine/Source/Programs/UnrealBuildTool/System/ConfigFile.cs b/Engine/Source/Programs/UnrealBuildTool/System/ConfigFile.cs
--- a/Engine/Source/Programs/UnrealBuildTool/System/ConfigFile.cs
+++ b/Engine/Source/Programs/UnrealBuildTool/System/ConfigFile.cs
@@ -452,7 +452,7 @@ namespace UnrealBuildTool
 			string NewKey = RemapSectionOrKey(KeyRemap, Key, $"which is a config key in section [{Section.Name}], in '{Filename}'");
 
 			// look for a section:name remap
-			if (!NewKey.Equals(Key, StringComparison.Ordinal) && NewKey.Contains(":", StringComparison.Ordinal))
+			if (!NewKey.Equals(Key, StringComparison.Ordinal) && NewKey.Contains(':', StringComparison.Ordinal))
 			{
 				string SectionName = NewKey.Substring(0, NewKey.IndexOf(':'));
 				ConfigFileSection? CurrentSection;
diff --git a/Engine/Source/Programs/UnrealBuildTool/System/Utils.cs b/Engine/Source/Programs/UnrealBuildTool/System/Utils.cs
--- a/Engine/Source/Programs/UnrealBuildTool/System/Utils.cs
+++ b/Engine/Source/Programs/UnrealBuildTool/System/Utils.cs
@@ -660,7 +660,7 @@ namespace UnrealBuildTool
 			string LocalString = InPath;
 			bool bHadBackSlashes = false;
 			// look to see what kind of slashes we had
-			if (LocalString.Contains("\\", StringComparison.CurrentCulture))
+			if (LocalString.Contains('\\', StringComparison.Ordinal))
 			{
 				LocalString = LocalString.Replace("\\", "/");
 				bHadBackSlashes = true;
