Fix issue CA1834 (Use StringBuilder.Append(char) for single character strings) in UBT.

diff --git a/Engine/Source/Programs/UnrealBuildTool/ProjectFiles/VisualStudio/VCProject.cs b/Engine/Source/Programs/UnrealBuildTool/ProjectFiles/VisualStudio/VCProject.cs
--- a/Engine/Source/Programs/UnrealBuildTool/ProjectFiles/VisualStudio/VCProject.cs
+++ b/Engine/Source/Programs/UnrealBuildTool/ProjectFiles/VisualStudio/VCProject.cs
@@ -1499,7 +1499,7 @@ namespace UnrealBuildTool
 								if (UpdatedValues.TryAdd(Value, NextPropertyName))
 								{
 									string PropertyName = NextPropertyName;
-									PropertyNameBuilder.Append("_").Append(++PropertyIndex);
+									PropertyNameBuilder.Append('_').Append(++PropertyIndex);
 									NextPropertyName = PropertyNameBuilder.ToString();
 									PropertyNameBuilder.Length = PropertyNameBuilderLength;
 									DirectoryToStringDict[DirectoryKVP.Key] = PropertyName;
@@ -1598,7 +1598,7 @@ namespace UnrealBuildTool
 							if (UpdatedValues.TryAdd(Value, NextPropertyName))
 							{
 								string PropertyName = NextPropertyName;
-								PropertyNameBuilder.Append("_").Append(++ValueIndex);
+								PropertyNameBuilder.Append('_').Append(++ValueIndex);
 								NextPropertyName = PropertyNameBuilder.ToString();
 								PropertyNameBuilder.Length = PropertyNameBuilderLength;
 
diff --git a/Engine/Source/Programs/UnrealBuildTool/ProjectFiles/VisualStudioWorkspace/VSWorkspaceProjectFile.cs b/Engine/Source/Programs/UnrealBuildTool/ProjectFiles/VisualStudioWorkspace/VSWorkspaceProjectFile.cs
--- a/Engine/Source/Programs/UnrealBuildTool/ProjectFiles/VisualStudioWorkspace/VSWorkspaceProjectFile.cs
+++ b/Engine/Source/Programs/UnrealBuildTool/ProjectFiles/VisualStudioWorkspace/VSWorkspaceProjectFile.cs
@@ -132,14 +132,14 @@ namespace UnrealBuildTool
 
 							StringBuilder sb = new();
 							sb.Append(ProjectTarget.TargetFilePath.GetFileNameWithoutAnyExtensions());
-							sb.Append("_");
+							sb.Append('_');
 							sb.Append(TargetInfo.Configuration);
-							sb.Append("_");
+							sb.Append('_');
 							sb.Append(TargetInfo.Platform);
 
 							if (!string.IsNullOrEmpty(TargetInfo.Architecture))
 							{
-								sb.Append("_");
+								sb.Append('_');
 								sb.Append(TargetInfo.Architecture);
 							}
 
