Fix issue CA1832 (Use AsSpan or AsMemory instead of Range-based indexers for getting ReadOnlySpan or ReadOnlyMemory portion of an array) in UBT.

diff --git a/Engine/Source/Programs/UnrealBuildTool/System/UHTExecution.cs b/Engine/Source/Programs/UnrealBuildTool/System/UHTExecution.cs
--- a/Engine/Source/Programs/UnrealBuildTool/System/UHTExecution.cs
+++ b/Engine/Source/Programs/UnrealBuildTool/System/UHTExecution.cs
@@ -1176,8 +1176,8 @@ namespace UnrealBuildTool
 				xx.Append(MemoryMarshal.AsBytes(file.FullName.AsSpan()));
 			}
 			byte[] hash16 = xx.GetHashAndReset();
-			ulong lo = BinaryPrimitives.ReadUInt64LittleEndian(hash16[..8]);
-			ulong hi = BinaryPrimitives.ReadUInt64LittleEndian(hash16[8..]);
+			ulong lo = BinaryPrimitives.ReadUInt64LittleEndian(hash16.AsSpan()[..8]);
+			ulong hi = BinaryPrimitives.ReadUInt64LittleEndian(hash16.AsSpan()[8..]);
 			return ((UInt128)hi << 64) | lo;
 		}
 
