Fix issue CA1830 (Prefer strongly-typed Append and Insert method overloads on StringBuilder) in UBT.

diff --git a/Engine/Source/Programs/UnrealBuildTool/ProjectFiles/VisualStudio/VCProject.cs b/Engine/Source/Programs/UnrealBuildTool/ProjectFiles/VisualStudio/VCProject.cs
--- a/Engine/Source/Programs/UnrealBuildTool/ProjectFiles/VisualStudio/VCProject.cs
+++ b/Engine/Source/Programs/UnrealBuildTool/ProjectFiles/VisualStudio/VCProject.cs
@@ -1680,7 +1680,7 @@ namespace UnrealBuildTool
 							else
 							{
 								VCProjectFileContent.AppendLine("    <{0} Include=\"{1}\">", VCFileType, EscapeFileName(AliasedFile.FileSystemPath));
-								VCProjectFileContent.Append(ClCompileInfo.ToString());
+								VCProjectFileContent.Append(ClCompileInfo);
 								VCProjectFileContent.AppendLine("    </{0}>", VCFileType);
 							}
 						}
@@ -2489,7 +2489,7 @@ namespace UnrealBuildTool
 					{
 						StringBuilder PathsStringBuilder = new StringBuilder();
 						ProjGenerator.GetVisualStudioPathsEntries(VSSettings, TargetRulesObject.Type, TargetFilePath, ProjectFilePath, NMakePath, PathsStringBuilder);
-						VCProjectFileContent.Append(PathsStringBuilder.ToString());
+						VCProjectFileContent.Append(PathsStringBuilder);
 					}
 
 					// This is the standard UE based project NMake build line:
