Fix issue CA1829 (Use Length/Count property instead of Enumerable.Count method) in UBT.

diff --git a/Engine/Source/Programs/UnrealBuildTool/Configuration/UEBuildTarget.cs b/Engine/Source/Programs/UnrealBuildTool/Configuration/UEBuildTarget.cs
--- a/Engine/Source/Programs/UnrealBuildTool/Configuration/UEBuildTarget.cs
+++ b/Engine/Source/Programs/UnrealBuildTool/Configuration/UEBuildTarget.cs
@@ -4113,7 +4113,7 @@ namespace UnrealBuildTool
 						pluginGroups.Add(name, self);
 						pluginGroupDependencies.Add(name, dependencies);
 
-						Logger.LogInformation("{Name}: setup for merging with {NumDependencies} total dependencies", name, dependencies.Count());
+						Logger.LogInformation("{Name}: setup for merging with {NumDependencies} total dependencies", name, dependencies.Count);
 					}
 				}
 
@@ -4294,7 +4294,7 @@ namespace UnrealBuildTool
 									{
 										roots.Add(plugin.Name, newPath);
 									}
-									else if (path.Count() < roots[plugin.Name].Count())
+									else if (path.Count < roots[plugin.Name].Count)
 									{
 										roots[plugin.Name] = newPath;
 									}
diff --git a/Engine/Source/Programs/UnrealBuildTool/ToolChain/AppleToolChain.cs b/Engine/Source/Programs/UnrealBuildTool/ToolChain/AppleToolChain.cs
--- a/Engine/Source/Programs/UnrealBuildTool/ToolChain/AppleToolChain.cs
+++ b/Engine/Source/Programs/UnrealBuildTool/ToolChain/AppleToolChain.cs
@@ -653,7 +653,7 @@ namespace UnrealBuildTool
 
 			Arguments.Add($"-module-name {ModuleName}Swift");
 
-			if (CompileEnvironment.ForceIncludeFiles.Count() > 0)
+			if (CompileEnvironment.ForceIncludeFiles.Count > 0)
 			{
 				Arguments.Add($"-import-objc-header \"{NormalizeCommandLinePath(CompileEnvironment.ForceIncludeFiles.First(), rootPaths)}\"");
 			}
@@ -729,7 +729,7 @@ namespace UnrealBuildTool
 			CompileAction.bCanExecuteInUBA = false;
 
 
-			if (CompileEnvironment.ForceIncludeFiles.Count() > 0)
+			if (CompileEnvironment.ForceIncludeFiles.Count > 0)
 			{
 				FileItem OutputInteropHeader = GetBridgingHeader(ModuleName, OutputDir);
 
