Fix issue CA1827 (Do not use Count()/LongCount() when Any() can be used) in UBT.

diff --git a/Engine/Source/Programs/UnrealBuildTool/Configuration/TestModuleRules.cs b/Engine/Source/Programs/UnrealBuildTool/Configuration/TestModuleRules.cs
--- a/Engine/Source/Programs/UnrealBuildTool/Configuration/TestModuleRules.cs
+++ b/Engine/Source/Programs/UnrealBuildTool/Configuration/TestModuleRules.cs
@@ -380,7 +380,7 @@ namespace UnrealBuildTool
 
 		private static bool IsRestrictedPath(string ModuleDirectory)
 		{
-			return ModuleDirectory.Split(new char[] { '/', '\\' }).Intersect(RestrictedFoldersNonPlatform).Count() > 0;
+			return ModuleDirectory.Split(new char[] { '/', '\\' }).Intersect(RestrictedFoldersNonPlatform).Any();
 		}
 
 		private static string TryGetBinariesPath(string ModuleDirectory)
diff --git a/Engine/Source/Programs/UnrealBuildTool/Configuration/UEBuildModuleCPP.cs b/Engine/Source/Programs/UnrealBuildTool/Configuration/UEBuildModuleCPP.cs
--- a/Engine/Source/Programs/UnrealBuildTool/Configuration/UEBuildModuleCPP.cs
+++ b/Engine/Source/Programs/UnrealBuildTool/Configuration/UEBuildModuleCPP.cs
@@ -1020,7 +1020,7 @@ namespace UnrealBuildTool
 				LinkInputFiles.AddRange(Output.ObjectFiles);
 
 				// and add the .h's as compile dependencies to make sure they are generated before we compile
-				IEnumerable<FileItem> HeaderFiles = Output.GeneratedHeaderFiles;
+				IReadOnlyList<FileItem> HeaderFiles = Output.GeneratedHeaderFiles;
 
 				// every cpp file in the module will depend on this header, because any file could need it (this means heavy swift
 				// changes will require rebuilding the whole module a lot)
@@ -1028,11 +1028,11 @@ namespace UnrealBuildTool
 				// both used when compiling files that may need the headers
 				ModuleCompileEnvironment.AdditionalPrerequisites.AddRange(HeaderFiles);
 				CppCompileEnvironment.AdditionalPrerequisites.AddRange(HeaderFiles);
-				if (HeaderFiles.Count() > 0)
+				if (HeaderFiles.Count > 0)
 				{
 					// @todo: this could have different per-arch headers- should be rare but possible
-					ModuleCompileEnvironment.UserIncludePaths.Add(HeaderFiles.First().Location.Directory);
-					CppCompileEnvironment.UserIncludePaths.Add(HeaderFiles.First().Location.Directory);
+					ModuleCompileEnvironment.UserIncludePaths.Add(HeaderFiles[0].Location.Directory);
+					CppCompileEnvironment.UserIncludePaths.Add(HeaderFiles[0].Location.Directory);
 				}
 			}
 
diff --git a/Engine/Source/Programs/UnrealBuildTool/System/ProjectDescriptor.cs b/Engine/Source/Programs/UnrealBuildTool/System/ProjectDescriptor.cs
--- a/Engine/Source/Programs/UnrealBuildTool/System/ProjectDescriptor.cs
+++ b/Engine/Source/Programs/UnrealBuildTool/System/ProjectDescriptor.cs
@@ -252,16 +252,16 @@ namespace UnrealBuildTool
 			}
 
 			// find any uproject file, in case the name doesn't match
-			IEnumerable<FileReference> FoundProjects = DirectoryReference.EnumerateFiles(DirectoryName, "*.uproject", SearchOption.TopDirectoryOnly);
-			if (FoundProjects.Count() == 0)
+			List<FileReference> FoundProjects = [.. DirectoryReference.EnumerateFiles(DirectoryName, "*.uproject", SearchOption.TopDirectoryOnly)];
+			if (FoundProjects.Count == 0)
 			{
 				throw new FileNotFoundException($"Unable to find a .uproject file in {DirectoryName}");
 			}
-			if (FoundProjects.Count() > 1)
+			if (FoundProjects.Count > 1)
 			{
-				Log.TraceWarningOnce("Found multiple uproject files in {0}, choosing {1}", DirectoryName, FoundProjects.First().GetFileName());
+				Log.TraceWarningOnce("Found multiple uproject files in {0}, choosing {1}", DirectoryName, FoundProjects[0].GetFileName());
 			}
-			return ProjectDescriptor.FromFile(FoundProjects.First());
+			return ProjectDescriptor.FromFile(FoundProjects[0]);
 		}
 
 		/// <summary>
