Fix or selectively disable issue CA1823 (Avoid unused private fields) in UBT.

diff --git a/Engine/Source/Programs/UnrealBuildTool/Executors/ParallelExecutor.cs b/Engine/Source/Programs/UnrealBuildTool/Executors/ParallelExecutor.cs
--- a/Engine/Source/Programs/UnrealBuildTool/Executors/ParallelExecutor.cs
+++ b/Engine/Source/Programs/UnrealBuildTool/Executors/ParallelExecutor.cs
@@ -26,8 +26,10 @@ namespace UnrealBuildTool
 		[XmlConfigFile]
 		[Obsolete("ParallelExecutor.MaxProcessorCount is deprecated. Please update xml to use BuildConfiguration.MaxParallelActions")]
 #pragma warning disable 0169
+#pragma warning disable CA1823 // Avoid unused private fields - part of the public API
 		private static int MaxProcessorCount;
 #pragma warning restore 0169
+#pragma warning restore CA1823
 
 		/// <summary>
 		/// Processor count multiplier for local execution. Can be below 1 to reserve CPU for other tasks.
diff --git a/Engine/Source/Programs/UnrealBuildTool/Modes/BuildMode.cs b/Engine/Source/Programs/UnrealBuildTool/Modes/BuildMode.cs
--- a/Engine/Source/Programs/UnrealBuildTool/Modes/BuildMode.cs
+++ b/Engine/Source/Programs/UnrealBuildTool/Modes/BuildMode.cs
@@ -420,8 +420,6 @@ namespace UnrealBuildTool
 			}
 		}
 
-		static Lazy<MemoryMappedFileCache> s_memoryMappedFileCache = new Lazy<MemoryMappedFileCache>(() => new MemoryMappedFileCache());
-
 		/// <summary>
 		/// Build a list of targets with a given set of makefiles.
 		/// </summary>
diff --git a/Engine/Source/Programs/UnrealBuildTool/Platform/Android/AndroidToolChain.cs b/Engine/Source/Programs/UnrealBuildTool/Platform/Android/AndroidToolChain.cs
--- a/Engine/Source/Programs/UnrealBuildTool/Platform/Android/AndroidToolChain.cs
+++ b/Engine/Source/Programs/UnrealBuildTool/Platform/Android/AndroidToolChain.cs
@@ -60,12 +60,6 @@ namespace UnrealBuildTool
 		// this would be 6317467)
 		protected static string? AndroidClangBuild;
 
-		// the "-android" suffix paths here are vcpkg triplets for the android platform
-		private static readonly Dictionary<UnrealArch, string[]> AllArchNames = new() {
-			{ UnrealArch.Arm64, new string[] { "arm64", "arm64-v8a", "arm64-android" } },
-			{ UnrealArch.X64,   new string[] { "x64", "x86_64", "x64-android" } }
-		};
-
 		// architecture paths to use for filtering include and lib paths
 		private static readonly Dictionary<UnrealArch, string[]> AllFilterArchNames = new() {
 			{ UnrealArch.Arm64, new string[] { "arm64", "arm64-v8a", "arm64-android" } },
diff --git a/Engine/Source/Programs/UnrealBuildTool/ProjectFiles/Xcode/XcodeProject.cs b/Engine/Source/Programs/UnrealBuildTool/ProjectFiles/Xcode/XcodeProject.cs
--- a/Engine/Source/Programs/UnrealBuildTool/ProjectFiles/Xcode/XcodeProject.cs
+++ b/Engine/Source/Programs/UnrealBuildTool/ProjectFiles/Xcode/XcodeProject.cs
@@ -2362,8 +2362,7 @@ namespace UnrealBuildTool.XcodeProjectXcconfig
 			BackgroundAssets,
 			NotificationService
 		}
-		
-		private string ProductGuid = XcodeProjectFileGenerator.MakeXcodeGuid();
+
 		private UnrealData UnrealData;
 		private DirectoryInfo ExtensionPath;
 		private UnrealTargetPlatform Platform;
