Fix issue CA1820 (Test for empty strings using string length) in UBT.

diff --git a/Engine/Source/Programs/UnrealBuildTool/Executors/UnrealBuildAccelerator/UBAAgentCoordinatorHorde.cs b/Engine/Source/Programs/UnrealBuildTool/Executors/UnrealBuildAccelerator/UBAAgentCoordinatorHorde.cs
--- a/Engine/Source/Programs/UnrealBuildTool/Executors/UnrealBuildAccelerator/UBAAgentCoordinatorHorde.cs
+++ b/Engine/Source/Programs/UnrealBuildTool/Executors/UnrealBuildAccelerator/UBAAgentCoordinatorHorde.cs
@@ -833,7 +833,7 @@ namespace UnrealBuildTool
 			if (_hordeConfig.HordeServer != null)
 			{
 				link = $"{_hordeConfig.HordeServer!.TrimEnd('/')}/agents";
-				if (_hordeConfig.HordePool != string.Empty)
+				if (!string.IsNullOrEmpty(_hordeConfig.HordePool))
 				{
 					link += $"?agent={_hordeConfig.HordePool}";
 				}
diff --git a/Engine/Source/Programs/UnrealBuildTool/System/CppDependencyCache.cs b/Engine/Source/Programs/UnrealBuildTool/System/CppDependencyCache.cs
--- a/Engine/Source/Programs/UnrealBuildTool/System/CppDependencyCache.cs
+++ b/Engine/Source/Programs/UnrealBuildTool/System/CppDependencyCache.cs
@@ -542,8 +542,10 @@ namespace UnrealBuildTool
 					string FilePath = Tokens[TokenIdx];
 					if (!Path.IsPathFullyQualified(FilePath))
 					{
-						if (CurrentDir == string.Empty)
+						if (CurrentDir.Length == 0)
+						{
 							CurrentDir = Directory.GetCurrentDirectory(); // Get current directory has crazy contention. So we only want to call it once
+						}
 						FilePath = Path.Combine(CurrentDir, FilePath);
 					}
 					NewDependencyFiles.Add(FileItem.GetItemByPath(FilePath));
diff --git a/Engine/Source/Programs/UnrealBuildTool/System/FileHashCache.cs b/Engine/Source/Programs/UnrealBuildTool/System/FileHashCache.cs
--- a/Engine/Source/Programs/UnrealBuildTool/System/FileHashCache.cs
+++ b/Engine/Source/Programs/UnrealBuildTool/System/FileHashCache.cs
@@ -488,7 +488,7 @@ namespace UnrealBuildTool
 				if (partition == null)
 				{
 					FileReference location;
-					if (filterPath != string.Empty)
+					if (filterPath.Length != 0)
 					{
 						location = FileReference.Combine(rootPath, "Intermediate", "Build", filterPath, "FileHashCache.bin");
 					}
@@ -525,7 +525,7 @@ namespace UnrealBuildTool
 			int index = fullName.IndexOf(Intermediate, FileReference.Comparison);
 			if (index == -1)
 			{
-				if (_roots.Count == 1 && _roots[0].root.FullName == string.Empty) // MountAll
+				if (_roots.Count == 1 && _roots[0].root.FullName.Length == 0) // MountAll
 				{
 					return _roots[0].Item2[0].partition;
 				}
diff --git a/Engine/Source/Programs/UnrealBuildTool/ToolChain/AppleToolChain.cs b/Engine/Source/Programs/UnrealBuildTool/ToolChain/AppleToolChain.cs
--- a/Engine/Source/Programs/UnrealBuildTool/ToolChain/AppleToolChain.cs
+++ b/Engine/Source/Programs/UnrealBuildTool/ToolChain/AppleToolChain.cs
@@ -1008,7 +1008,7 @@ namespace UnrealBuildTool
 					//   - when no destinations were supplied
 					//   - when trying to build iOS app to run on local mac 
 					//   - when using manual codesigning because the main use of this is to update provisioning information (reducing potential for issues)
-					if (Platform == UnrealTargetPlatform.Mac || DestinationIds == null || DestinationIds.Count == 0 || DestinationIds[0] == "" || DestinationIds[0].Contains(AppleToolChainSettings.LocalMacUUID) || !bUseAutomaticCodeSigning)
+					if (Platform == UnrealTargetPlatform.Mac || DestinationIds == null || DestinationIds.Count == 0 || DestinationIds[0].Length == 0 || DestinationIds[0].Contains(AppleToolChainSettings.LocalMacUUID) || !bUseAutomaticCodeSigning)
 					{
 						DestinationParam = $"-destination generic/platform=\"{AppleExports.GetDestinationPlatform(Platform, Architectures)}\"";
 					}
