Fix issue CA1715 (Identifiers should have correct prefix) in UBT.

diff --git a/Engine/Source/Programs/UnrealBuildTool/System/Utils.cs b/Engine/Source/Programs/UnrealBuildTool/System/Utils.cs
--- a/Engine/Source/Programs/UnrealBuildTool/System/Utils.cs
+++ b/Engine/Source/Programs/UnrealBuildTool/System/Utils.cs
@@ -944,12 +944,12 @@ namespace UnrealBuildTool
 		/// <summary>
 		/// Checks if given type implements given interface.
 		/// </summary>
-		/// <typeparam name="InterfaceType">Interface to check.</typeparam>
+		/// <typeparam name="TInterface">Interface to check.</typeparam>
 		/// <param name="TestType">Type to check.</param>
-		/// <returns>True if TestType implements InterfaceType. False otherwise.</returns>
-		public static bool ImplementsInterface<InterfaceType>(Type TestType)
+		/// <returns>True if TestType implements TInterface. False otherwise.</returns>
+		public static bool ImplementsInterface<TInterface>(Type TestType)
 		{
-			return Array.IndexOf(TestType.GetInterfaces(), typeof(InterfaceType)) != -1;
+			return Array.IndexOf(TestType.GetInterfaces(), typeof(TInterface)) != -1;
 		}
 
 		/// <summary>
