Selectively disable issue CA1724 (Type names should not match namespaces) in UBT.

diff --git a/Engine/Source/Programs/UnrealBuildTool/Configuration/ModuleRules.cs b/Engine/Source/Programs/UnrealBuildTool/Configuration/ModuleRules.cs
--- a/Engine/Source/Programs/UnrealBuildTool/Configuration/ModuleRules.cs
+++ b/Engine/Source/Programs/UnrealBuildTool/Configuration/ModuleRules.cs
@@ -455,7 +455,9 @@ namespace UnrealBuildTool
 		/// <summary>
 		/// Stores information about a framework on IOS or MacOS
 		/// </summary>
+#pragma warning disable CA1724 // Type names should not match namespaces
 		public class Framework
+#pragma warning restore CA1724
 		{
 			/// <summary>
 			/// Name of the framework
diff --git a/Engine/Source/Programs/UnrealBuildTool/Configuration/TestModuleRules.cs b/Engine/Source/Programs/UnrealBuildTool/Configuration/TestModuleRules.cs
--- a/Engine/Source/Programs/UnrealBuildTool/Configuration/TestModuleRules.cs
+++ b/Engine/Source/Programs/UnrealBuildTool/Configuration/TestModuleRules.cs
@@ -470,7 +470,9 @@ namespace UnrealBuildTool
 		/// <summary>
 		/// Test metadata class.
 		/// </summary>
+#pragma warning disable CA1724 // Type names should not match namespaces
 		public class Metadata
+#pragma warning restore CA1724
 		{
 			/// <summary>
 			/// Test long name.
diff --git a/Engine/Source/Programs/UnrealBuildTool/System/Utils.cs b/Engine/Source/Programs/UnrealBuildTool/System/Utils.cs
--- a/Engine/Source/Programs/UnrealBuildTool/System/Utils.cs
+++ b/Engine/Source/Programs/UnrealBuildTool/System/Utils.cs
@@ -21,7 +21,9 @@ namespace UnrealBuildTool
 	/// <summary>
 	/// Utility functions
 	/// </summary>
+#pragma warning disable CA1724 // Type names should not match namespaces
 	public static class Utils
+#pragma warning restore CA1724
 	{
 		/// <summary>
 		/// Searches for a flag in a set of command-line arguments.
