Fix issue CA1727 (Use PascalCase for named placeholders) in UBT.

diff --git a/Engine/Source/Programs/UnrealBuildTool/Modes/PipInstallMode.cs b/Engine/Source/Programs/UnrealBuildTool/Modes/PipInstallMode.cs
--- a/Engine/Source/Programs/UnrealBuildTool/Modes/PipInstallMode.cs
+++ b/Engine/Source/Programs/UnrealBuildTool/Modes/PipInstallMode.cs
@@ -213,7 +213,7 @@ namespace UnrealBuildTool.Modes
 		{
 			if (InTargetDesc.Architectures.bIsMultiArch)
 			{
-				InLogger.LogWarning("PipInstall expects single-arch target(host): {arch}", InTargetDesc.Architectures.ToString());
+				InLogger.LogWarning("PipInstall expects single-arch target(host): {Arch}", InTargetDesc.Architectures.ToString());
 			}
 			
 			Logger = InLogger;
diff --git a/Engine/Source/Programs/UnrealBuildTool/System/FileHashCache.cs b/Engine/Source/Programs/UnrealBuildTool/System/FileHashCache.cs
--- a/Engine/Source/Programs/UnrealBuildTool/System/FileHashCache.cs
+++ b/Engine/Source/Programs/UnrealBuildTool/System/FileHashCache.cs
@@ -395,13 +395,13 @@ namespace UnrealBuildTool
 				FileReference backupFile = new FileReference(fileItem.FullName + ".old");
 				try
 				{
-					logger.LogDebug("Updating {File}: contents have changed. Saving previous version to {backupFile}.", fileItem.Location, backupFile);
+					logger.LogDebug("Updating {File}: contents have changed. Saving previous version to {BackupFile}.", fileItem.Location, backupFile);
 					FileReference.Delete(backupFile);
 					FileReference.Move(fileItem.Location, backupFile);
 				}
 				catch (Exception Ex)
 				{
-					logger.LogWarning("Unable to rename {fileItem} to {backupFile}", fileItem, backupFile);
+					logger.LogWarning("Unable to rename {FileItem} to {BackupFile}", fileItem, backupFile);
 					logger.LogDebug(Ex, "{Ex}", ExceptionUtils.FormatExceptionDetails(Ex));
 				}
 			}
