Fix issue CA1802 (Use Literals Where Appropriate) in UBT.

diff --git a/Engine/Source/Programs/UnrealBuildTool/ProjectFiles/PlatformProjectGenerator.cs b/Engine/Source/Programs/UnrealBuildTool/ProjectFiles/PlatformProjectGenerator.cs
--- a/Engine/Source/Programs/UnrealBuildTool/ProjectFiles/PlatformProjectGenerator.cs
+++ b/Engine/Source/Programs/UnrealBuildTool/ProjectFiles/PlatformProjectGenerator.cs
@@ -14,7 +14,7 @@ namespace UnrealBuildTool
 	/// </summary>
 	abstract class PlatformProjectGenerator
 	{
-		public static readonly string DefaultPlatformConfigurationType = "Makefile";
+		public const string DefaultPlatformConfigurationType = "Makefile";
 		protected readonly ILogger Logger;
 
 		/// <summary>
diff --git a/Engine/Source/Programs/UnrealBuildTool/ProjectFiles/VisualStudio/VCProject.cs b/Engine/Source/Programs/UnrealBuildTool/ProjectFiles/VisualStudio/VCProject.cs
--- a/Engine/Source/Programs/UnrealBuildTool/ProjectFiles/VisualStudio/VCProject.cs
+++ b/Engine/Source/Programs/UnrealBuildTool/ProjectFiles/VisualStudio/VCProject.cs
@@ -77,16 +77,16 @@ namespace UnrealBuildTool
 	abstract class MSBuildProjectFile : ProjectFile
 	{
 		/// The project file version string
-		public static readonly string VCProjectFileVersionString = "10.0.30319.1";
+		public const string VCProjectFileVersionString = "10.0.30319.1";
 
 		/// The build configuration name to use for stub project configurations.  These are projects whose purpose
 		/// is to make it easier for developers to find source files and to provide IntelliSense data for the module
 		/// to Visual Studio
-		public static readonly string StubProjectConfigurationName = "BuiltWithUnrealBuildTool";
+		public const string StubProjectConfigurationName = "BuiltWithUnrealBuildTool";
 
 		/// The name of the Visual C++ platform to use for stub project configurations
 		/// NOTE: We always use Win32 for the stub project's platform, since that is guaranteed to be supported by Visual Studio
-		public static readonly string StubProjectPlatformName = "Win64";
+		public const string StubProjectPlatformName = "Win64";
 
 		/// <summary>
 		/// The Guid representing the project type e.g. C# or C++
diff --git a/Engine/Source/Programs/UnrealBuildTool/ProjectFiles/Xcode/XcodeFrameworkWrapperProject.cs b/Engine/Source/Programs/UnrealBuildTool/ProjectFiles/Xcode/XcodeFrameworkWrapperProject.cs
--- a/Engine/Source/Programs/UnrealBuildTool/ProjectFiles/Xcode/XcodeFrameworkWrapperProject.cs
+++ b/Engine/Source/Programs/UnrealBuildTool/ProjectFiles/Xcode/XcodeFrameworkWrapperProject.cs
@@ -15,11 +15,11 @@ namespace UnrealBuildTool.ProjectFiles.Xcode
 	/// </summary>
 	class XcodeFrameworkWrapperProject
 	{
-		private static readonly string PROJECT_FILE_SEARCH_EXPRESSION = "*.pbxproj";
-		private static readonly string TEMPLATE_NAME = "FrameworkWrapper";
+		private const string PROJECT_FILE_SEARCH_EXPRESSION = "*.pbxproj";
+		private const string TEMPLATE_NAME = "FrameworkWrapper";
 		private static readonly string FRAMEWORK_WRAPPER_TEMPLATE_DIRECTORY = Path.Combine(Unreal.EngineDirectory.ToNormalizedPath(), "Build", "IOS", "Resources", TEMPLATE_NAME);
-		private static readonly string TEMPLATE_PROJECT_NAME = "PROJECT_NAME";
-		private static readonly string COMMANDLINE_FILENAME = "uecommandline.txt";
+		private const string TEMPLATE_PROJECT_NAME = "PROJECT_NAME";
+		private const string COMMANDLINE_FILENAME = "uecommandline.txt";
 
 		/// <summary>
 		/// Recursively copies all of the files and directories that are inside <paramref name="SourceDirectory"/> into <paramref name="DestinationDirectory"/>.
diff --git a/Engine/Source/Programs/UnrealBuildTool/System/ProtoExecution.cs b/Engine/Source/Programs/UnrealBuildTool/System/ProtoExecution.cs
--- a/Engine/Source/Programs/UnrealBuildTool/System/ProtoExecution.cs
+++ b/Engine/Source/Programs/UnrealBuildTool/System/ProtoExecution.cs
@@ -122,8 +122,8 @@ namespace UnrealBuildTool
 
 	class ProtoExecution
 	{
-		public static readonly string ProtobufVersion = "30.0";
-		public static readonly string GrpcVersion = "1.72.1";
+		public const string ProtobufVersion = "30.0";
+		public const string GrpcVersion = "1.72.1";
 
 		public static ProtocOutput CompileProtoFiles(IEnumerable<FileItem> ProtoFiles, DirectoryReference ProtoModuleDir, DirectoryReference OutputDir, IActionGraphBuilder Graph)
 		{
