Fix or selectively disable issue CA1810 (Initialize reference type static fields inline) in UBT.

diff --git a/Engine/Source/Programs/UnrealBuildTool/Preprocessor/PreprocessorExpression.cs b/Engine/Source/Programs/UnrealBuildTool/Preprocessor/PreprocessorExpression.cs
--- a/Engine/Source/Programs/UnrealBuildTool/Preprocessor/PreprocessorExpression.cs
+++ b/Engine/Source/Programs/UnrealBuildTool/Preprocessor/PreprocessorExpression.cs
@@ -16,6 +16,7 @@ namespace UnrealBuildTool
 		/// </summary>
 		static readonly byte[] s_precedence;
 
+#pragma warning disable CA1810 // Initialize reference type static fields inline
 		/// <summary>
 		/// Static constructor. Initializes the Precedence table.
 		/// </summary>
@@ -41,6 +42,7 @@ namespace UnrealBuildTool
 			s_precedence[(int)TokenType.LogicalAnd] = 2;
 			s_precedence[(int)TokenType.LogicalOr] = 1;
 		}
+#pragma warning restore CA1810
 
 		/// <summary>
 		/// Evaluate a preprocessor expression
diff --git a/Engine/Source/Programs/UnrealBuildTool/Preprocessor/Token.cs b/Engine/Source/Programs/UnrealBuildTool/Preprocessor/Token.cs
--- a/Engine/Source/Programs/UnrealBuildTool/Preprocessor/Token.cs
+++ b/Engine/Source/Programs/UnrealBuildTool/Preprocessor/Token.cs
@@ -454,6 +454,7 @@ namespace UnrealBuildTool
 			}
 		}
 
+#pragma warning disable CA1810 // Initialize reference type static fields inline
 		/// <summary>
 		/// Static constructor. Initializes the StaticTokens array.
 		/// </summary>
@@ -494,6 +495,7 @@ namespace UnrealBuildTool
 			s_staticTokens[(int)TokenType.HashHash] = "##";
 		}
 	}
+#pragma warning restore CA1810
 
 	/// <summary>
 	/// Helper functions for serialization
diff --git a/Engine/Source/Programs/UnrealBuildTool/Preprocessor/TokenReader.cs b/Engine/Source/Programs/UnrealBuildTool/Preprocessor/TokenReader.cs
--- a/Engine/Source/Programs/UnrealBuildTool/Preprocessor/TokenReader.cs
+++ b/Engine/Source/Programs/UnrealBuildTool/Preprocessor/TokenReader.cs
@@ -700,6 +700,7 @@ namespace UnrealBuildTool
 			}
 		}
 
+#pragma warning disable CA1810 // Initialize reference type static fields inline
 		/// <summary>
 		/// Static constructor. Initializes the lookup tables used by the lexer.
 		/// </summary>
@@ -809,5 +810,6 @@ namespace UnrealBuildTool
 			s_characters['\t'] |= CharacterFlags.Whitespace;
 			s_characters['\v'] |= CharacterFlags.Whitespace;
 		}
+#pragma warning restore CA1810
 	}
 }
diff --git a/Engine/Source/Programs/UnrealBuildTool/System/ActionSerializer.cs b/Engine/Source/Programs/UnrealBuildTool/System/ActionSerializer.cs
--- a/Engine/Source/Programs/UnrealBuildTool/System/ActionSerializer.cs
+++ b/Engine/Source/Programs/UnrealBuildTool/System/ActionSerializer.cs
@@ -77,6 +77,7 @@ namespace UnrealBuildTool
 		/// </summary>
 		static IReadOnlyDictionary<string, IActionSerializer> NameToSerializer;
 
+#pragma warning disable CA1810 // Initialize reference type static fields inline
 		/// <summary>
 		/// Creates a map of type name to constructor
 		/// </summary>
@@ -100,6 +101,7 @@ namespace UnrealBuildTool
 			TypeToSerializer = TypeToSerializerDict;
 			NameToSerializer = NameToSerializerDict;
 		}
+#pragma warning restore CA1810
 
 		/// <summary>
 		/// Read an action from the given archive
diff --git a/Engine/Source/Programs/UnrealBuildTool/System/InstalledPlatformInfo.cs b/Engine/Source/Programs/UnrealBuildTool/System/InstalledPlatformInfo.cs
--- a/Engine/Source/Programs/UnrealBuildTool/System/InstalledPlatformInfo.cs
+++ b/Engine/Source/Programs/UnrealBuildTool/System/InstalledPlatformInfo.cs
@@ -119,6 +119,7 @@ namespace UnrealBuildTool
 
 		private static List<InstalledPlatformConfiguration>? InstalledPlatformConfigurations;
 
+#pragma warning disable CA1810 // Initialize reference type static fields inline
 		static InstalledPlatformInfo()
 		{
 			List<string>? InstalledPlatforms;
@@ -137,6 +138,7 @@ namespace UnrealBuildTool
 				}
 			}
 		}
+#pragma warning restore CA1810
 
 		/// <summary>
 		/// Initializes the InstalledPlatformInfo. While this is not necessary to be called, it allows timing the static constructor call.
diff --git a/Engine/Source/Programs/UnrealBuildTool/ToolChain/ClangToolChain.cs b/Engine/Source/Programs/UnrealBuildTool/ToolChain/ClangToolChain.cs
--- a/Engine/Source/Programs/UnrealBuildTool/ToolChain/ClangToolChain.cs
+++ b/Engine/Source/Programs/UnrealBuildTool/ToolChain/ClangToolChain.cs
@@ -217,7 +217,9 @@ namespace UnrealBuildTool
 		protected bool bOptimizeForDebugging => Options.HasFlag(ClangToolChainOptions.OptimizeForDebugging);
 
 		// Dummy define to work around clang compilation related to the windows maximum path length limitation
-		protected static string ClangDummyDefine;
+		private const string DummyStart = "-D \"DUMMY_DEFINE";
+		private const string DummyEnd = "\"";
+		protected static readonly string ClangDummyDefine = DummyStart + "_".PadRight(ClangCmdLineMaxSize - ClangCmdlineDangerZone - DummyStart.Length - DummyEnd.Length, 'X') + DummyEnd;
 		protected const int ClangCmdLineMaxSize = 32 * 1024;
 		protected const int ClangCmdlineDangerZone = 30 * 1024;
 
@@ -228,13 +230,6 @@ namespace UnrealBuildTool
 		protected StaticAnalyzerOutputType StaticAnalyzerOutputType = StaticAnalyzerOutputType.Text;
 		protected float CompileActionWeight = 1.0f;
 
-		static ClangToolChain()
-		{
-			const string DummyStart = "-D \"DUMMY_DEFINE";
-			const string DummyEnd = "\"";
-			ClangDummyDefine = DummyStart + "_".PadRight(ClangCmdLineMaxSize - ClangCmdlineDangerZone - DummyStart.Length - DummyEnd.Length, 'X') + DummyEnd;
-		}
-
 		public ClangToolChain(ClangToolChainOptions InOptions, ILogger InLogger)
 			: base(InLogger)
 		{
