Fix issue CA1816 (Call GC.SuppressFinalize correctly) in UBT.

diff --git a/Engine/Source/Programs/UnrealBuildTool/System/ProgressWriter.cs b/Engine/Source/Programs/UnrealBuildTool/System/ProgressWriter.cs
--- a/Engine/Source/Programs/UnrealBuildTool/System/ProgressWriter.cs
+++ b/Engine/Source/Programs/UnrealBuildTool/System/ProgressWriter.cs
@@ -9,7 +9,7 @@ namespace UnrealBuildTool
 	/// <summary>
 	/// Class to display an incrementing progress percentage. Handles progress markup and direct console output.
 	/// </summary>
-	public class ProgressWriter : IDisposable
+	public sealed class ProgressWriter : IDisposable
 	{
 		/// <summary>
 		/// Global setting controlling whether to output markup
diff --git a/Engine/Source/Programs/UnrealBuildTool/System/ScopedTimer.cs b/Engine/Source/Programs/UnrealBuildTool/System/ScopedTimer.cs
--- a/Engine/Source/Programs/UnrealBuildTool/System/ScopedTimer.cs
+++ b/Engine/Source/Programs/UnrealBuildTool/System/ScopedTimer.cs
@@ -9,7 +9,7 @@ namespace UnrealBuildTool
 	/// <summary>
 	/// Scoped timer, start is in the constructor, end in Dispose. Best used with using(ScopedTimer Timer = new ScopedTimer()). Suports nesting.
 	/// </summary>
-	public class ScopedTimer : IDisposable
+	public sealed class ScopedTimer : IDisposable
 	{
 		DateTime StartTime;
 		string Name;
