Fix missing variable assignment in merge modules build setup

#rnx

diff --git a/Engine/Source/Programs/UnrealBuildTool/Configuration/UEBuildTarget.cs b/Engine/Source/Programs/UnrealBuildTool/Configuration/UEBuildTarget.cs
--- a/Engine/Source/Programs/UnrealBuildTool/Configuration/UEBuildTarget.cs
+++ b/Engine/Source/Programs/UnrealBuildTool/Configuration/UEBuildTarget.cs
@@ -4163,7 +4163,7 @@ namespace UnrealBuildTool
 
 									if (sharedGroupCandidates.Any())
 									{
-										sharedGroupCandidates.OrderBy(x => x.Value.Count()).Reverse();
+										sharedGroupCandidates = sharedGroupCandidates.OrderByDescending(x => x.Value.Count());
 										KeyValuePair<string, IEnumerable<string>> sharedGroupRule = sharedGroupCandidates.First();
 
 										AddPluginToGroup(plugin, sharedGroupRule.Key);
