PR #13917: Mitigate race condition in ImmediateActionQueue

diff --git a/Engine/Source/Programs/UnrealBuildTool/Executors/ImmediateActionQueue.cs b/Engine/Source/Programs/UnrealBuildTool/Executors/ImmediateActionQueue.cs
--- a/Engine/Source/Programs/UnrealBuildTool/Executors/ImmediateActionQueue.cs
+++ b/Engine/Source/Programs/UnrealBuildTool/Executors/ImmediateActionQueue.cs
@@ -555,6 +555,27 @@ namespace UnrealBuildTool
 							runner.LastActionChange = _lastActionChange;
 						}
 					}
+
+					// If we have no completed actions and have nothing to run, then we need to check for a stall
+					if (completedActions == 0 && (runAction == null || action == null))
+					{
+						// We found nothing, check to see if we have no active running tasks and no manual runners.
+						// If we don't, then it is assumed we can't queue any more work.
+						bool prematureDone = true;
+						foreach (ImmediateActionQueueRunner tryRunner in _runners)
+						{
+							if ((tryRunner.Type == ImmediateActionQueueRunnerType.Manual && !hasCanceled) || tryRunner.ActiveActions != 0)
+							{
+								prematureDone = false;
+								break;
+							}
+						}
+						if (prematureDone)
+						{
+							AddCompletedActions(Int32.MaxValue);
+						}
+						return false;
+					}
 				}
 
 				// If we have an action, run it and account for any completed actions
@@ -589,27 +610,6 @@ namespace UnrealBuildTool
 					return true;
 				}
 
-				// If we have no completed actions, then we need to check for a stall
-				if (completedActions == 0)
-				{
-					// We found nothing, check to see if we have no active running tasks and no manual runners.
-					// If we don't, then it is assumed we can't queue any more work.
-					bool prematureDone = true;
-					foreach (ImmediateActionQueueRunner tryRunner in _runners)
-					{
-						if ((tryRunner.Type == ImmediateActionQueueRunnerType.Manual && !hasCanceled) || tryRunner.ActiveActions != 0)
-						{
-							prematureDone = false;
-							break;
-						}
-					}
-					if (prematureDone)
-					{
-						AddCompletedActions(Int32.MaxValue);
-					}
-					return false;
-				}
-
 				// Otherwise add the completed actions and test again for the possibility that errors still need propagating
 				AddCompletedActions(completedActions);
 			}
