[UBT/UHT]
* Fixed bug where UHTTargetSettings.GenerateEnumMaxValues was not saved into the makefile and .uhtsettings file causing different results if makefile was up-to-date or needed to be recreated

diff --git a/Engine/Source/Programs/Shared/EpicGames.Core/UHTTypes.cs b/Engine/Source/Programs/Shared/EpicGames.Core/UHTTypes.cs
--- a/Engine/Source/Programs/Shared/EpicGames.Core/UHTTypes.cs
+++ b/Engine/Source/Programs/Shared/EpicGames.Core/UHTTypes.cs
@@ -147,11 +147,13 @@ namespace EpicGames.Core
 		public UHTTargetSettings(BinaryArchiveReader reader) : this()
 		{
 			CompiledInObjectFormat = Enum.Parse<UhtCompiledInObjectFormat>(reader.ReadString()!);
+			GenerateEnumMaxValues = reader.ReadBool();
 		}
 
 		public void Write(BinaryArchiveWriter writer)
 		{
 			writer.WriteString(CompiledInObjectFormat.ToString());
+			writer.WriteBool(GenerateEnumMaxValues);
 		}
 
 		/// <summary>
diff --git a/Engine/Source/Programs/UnrealBuildTool/System/UHTExecution.cs b/Engine/Source/Programs/UnrealBuildTool/System/UHTExecution.cs
--- a/Engine/Source/Programs/UnrealBuildTool/System/UHTExecution.cs
+++ b/Engine/Source/Programs/UnrealBuildTool/System/UHTExecution.cs
@@ -1312,8 +1312,14 @@ namespace UnrealBuildTool
 						{
 							// Note that even if default/omitted values cause the settings to compare equal here, the UHT module version check below 
 							// will ensure that we rerun UHT when new properties are added to UHTTargetSettings
-							UHTTargetSettings settings = new(reader);
-							if (settings != Makefile.UHTTargetSettings)
+							try
+							{
+								UHTTargetSettings settings = new(reader);
+								if (settings != Makefile.UHTTargetSettings)
+								{
+									bUHTNeedsToRun = true;
+								}
+							} catch (Exception)
 							{
 								bUHTNeedsToRun = true;
 							}
