UnrealBuildTool: Update IOS linking to use a response file instead of command line args (first pass before adding vfs support)
Properly delete IOS output item when required using Action.DeleteFiles rather than a shell script

#rnx

diff --git a/Engine/Source/Programs/UnrealBuildTool/Platform/IOS/IOSToolChain.cs b/Engine/Source/Programs/UnrealBuildTool/Platform/IOS/IOSToolChain.cs
--- a/Engine/Source/Programs/UnrealBuildTool/Platform/IOS/IOSToolChain.cs
+++ b/Engine/Source/Programs/UnrealBuildTool/Platform/IOS/IOSToolChain.cs
@@ -307,11 +307,12 @@ namespace UnrealBuildTool
 			Arguments.Add(FormatArchitectureArg(LinkEnvironment.Architectures));
 
 			DirectoryReference SDKPath = Settings.GetSDKPath(LinkEnvironment.Architecture);
-			Arguments.Add($" -isysroot \"{SDKPath}\"");
+			Arguments.Add($"-isysroot \"{SDKPath}\"");
 
 			// needed to find swift directories for some versions of the SDK
 			string PlatformName = LinkEnvironment.Architecture == UnrealArch.IOSSimulator ? Settings.SimulatorPlatformDirName : Settings.PlatformDirName;
-			Arguments.Add($"-L{ApplePlatformSDK.GetToolchainDirectory()}/../lib/swift/{PlatformName}");
+			FileReference PlatformLibDirectory = FileReference.Combine(ApplePlatformSDK.GetToolchainDirectory(), "..", "lib", "swift", PlatformName);
+			Arguments.Add($"-L{PlatformLibDirectory}");
 
 			if (!LinkEnvironment.bIsBuildingLibrary && !LinkEnvironment.bIsBuildingDLL)
 			{
@@ -390,6 +391,8 @@ namespace UnrealBuildTool
 
 			// Create an action that invokes the linker.
 			Action LinkAction = Graph.CreateAction(ActionType.Link);
+			string Arch = UnrealArchitectureConfig.ForPlatform(LinkEnvironment.Platform).ConvertToReadableArchitecture(LinkEnvironment.Architecture);
+			LinkAction.CommandDescription = (bBuildImportLibraryOnly ? "LinkStub" : "Link") + $" [{Arch}]";
 			LinkAction.RootPaths = LinkEnvironment.RootPaths;
 
 			// RPC utility parameters are in terms of the Mac side
@@ -405,18 +408,17 @@ namespace UnrealBuildTool
 			{
 				GetLinkArguments_Global(LinkEnvironment, LinkArguments);
 			}
-			string LinkCommandArguments = String.Join(' ', LinkArguments);
 			if (LinkEnvironment.bIsBuildingDLL)
 			{
 				// @todo roll this put into GetLinkArguments_Global
 				LinkerPath = Info.Clang;
-				LinkCommandArguments += " -dynamiclib -Xlinker -export_dynamic -Xlinker -no_deduplicate";
+				LinkArguments.Add("-dynamiclib -Xlinker -export_dynamic -Xlinker -no_deduplicate");
 
 				string InstallName = LinkEnvironment.InstallName ?? String.Format("@executable_path/Frameworks/{0}", LinkEnvironment.OutputFilePath.MakeRelativeTo(LinkEnvironment.OutputFilePath.Directory.ParentDirectory!));
-				LinkCommandArguments += String.Format(" -Xlinker -install_name -Xlinker {0}", InstallName);
+				LinkArguments.Add($"-Xlinker -install_name -Xlinker {InstallName}");
 
-				LinkCommandArguments += " -Xlinker -rpath -Xlinker @executable_path/Frameworks";
-				LinkCommandArguments += " -Xlinker -rpath -Xlinker @loader_path/Frameworks";
+				LinkArguments.Add("-Xlinker -rpath -Xlinker @executable_path/Frameworks");
+				LinkArguments.Add("-Xlinker -rpath -Xlinker @loader_path/Frameworks");
 			}
 
 			if (!LinkEnvironment.bIsBuildingLibrary)
@@ -424,13 +426,21 @@ namespace UnrealBuildTool
 				// Add the library paths to the argument list.
 				foreach (DirectoryReference LibraryPath in LinkEnvironment.SystemLibraryPaths)
 				{
-					LinkCommandArguments += String.Format(" -L\\\"{0}\\\"", LibraryPath.FullName);
+					LinkArguments.Add($"-L\"{LibraryPath.FullName}\"");
 				}
 
 				// Add the additional libraries to the argument list.
 				foreach (string AdditionalLibrary in LinkEnvironment.SystemLibraries)
 				{
-					LinkCommandArguments += String.Format(" -l\\\"{0}\\\"", AdditionalLibrary);
+					if (String.IsNullOrEmpty(Path.GetDirectoryName(AdditionalLibrary)) && String.IsNullOrEmpty(Path.GetExtension(AdditionalLibrary)))
+					{
+						LinkArguments.Add($"-l\"{AdditionalLibrary}\"");
+					}
+					else
+					{
+                        FileReference LibraryPath = FileReference.FromString(AdditionalLibrary);
+						LinkArguments.Add($"\"{LibraryPath}\"");
+                    }					
 				}
 
 				foreach (FileReference Library in LinkEnvironment.Libraries)
@@ -441,7 +451,7 @@ namespace UnrealBuildTool
 					LinkAction.PrerequisiteItems.Add(LibFile);
 
 					// and add to the commandline
-					LinkCommandArguments += String.Format(" \\\"{0}\\\"", Library.FullName);
+					LinkArguments.Add($"\"{Library.FullName}\"");
 				}
 			}
 
@@ -463,7 +473,7 @@ namespace UnrealBuildTool
 			if ((!LinkEnvironment.bIsBuildingLibrary || LinkEnvironment.bIsBuildingDLL) && LinkEnvironment.bCreateMapFile)
 			{
 				FileItem MapFile = FileItem.GetItemByFileReference(new FileReference(OutputFile.Location.FullName + ".map"));
-				LinkCommandArguments += String.Format(" -Wl,-map,\\\"{0}\\\"", MapFile.Location.FullName);
+				LinkArguments.Add($"-Wl,-map,\"{MapFile.Location.FullName}\"");
 				LinkAction.ProducedItems.Add(MapFile);
 			}
 
@@ -471,31 +481,10 @@ namespace UnrealBuildTool
 			List<string> InputFileNames = new List<string>();
 			foreach (FileItem InputFile in LinkEnvironment.InputFiles)
 			{
-				string FileNamePath = String.Format("\"{0}\"", InputFile.AbsolutePath);
-				InputFileNames.Add(FileNamePath);
+				LinkArguments.Add($"\"{InputFile.AbsolutePath}\"");
 				LinkAction.PrerequisiteItems.Add(InputFile);
 			}
 
-			// Write the list of input files to a response file, with a tempfilename, on remote machine
-			if (LinkEnvironment.bIsBuildingLibrary)
-			{
-				foreach (string Filename in InputFileNames)
-				{
-					LinkCommandArguments += " " + Filename;
-				}
-				// @todo rocket lib: the -filelist command should take a response file (see else condition), except that it just says it can't
-				// find the file that's in there. Rocket.lib may overflow the commandline by putting all files on the commandline, so this 
-				// may be needed:
-				// LinkCommandArguments += string.Format(" -filelist \"{0}\"", ConvertPath(ResponsePath));
-			}
-			else
-			{
-				FileReference ResponsePath = FileReference.Combine(LinkEnvironment.IntermediateDirectory!, "LinkFileList_" + LinkEnvironment.OutputFilePath.GetFileNameWithoutExtension() + ".tmp");
-				FileItem fileItem = Graph.CreateIntermediateTextFile(ResponsePath, InputFileNames);
-				LinkAction.ProducedItems.Add(fileItem);
-				LinkCommandArguments += String.Format(" \\\"@{0}\\\"", ResponsePath.FullName);
-			}
-
 			// if we are making an LTO build, write the lto file next to build so dsymutil can find it
 			if (LinkEnvironment.bAllowLTCG)
 			{
@@ -509,42 +498,47 @@ namespace UnrealBuildTool
 				{
 					LtoObjectFile = OutputFile.AbsolutePath + ".lto.o";
 				}
-				LinkCommandArguments += String.Format(" -flto -Xlinker -object_path_lto -Xlinker \\\"{0}\\\"", LtoObjectFile);
+				LinkArguments.Add($"-flto -Xlinker -object_path_lto -Xlinker \"{LtoObjectFile}\"");
 			}
 
 			// Add the output file to the command-line.
-			LinkCommandArguments += String.Format(" -o \\\"{0}\\\"", OutputFile.AbsolutePath);
+			LinkArguments.Add($"-o \"{OutputFile.AbsolutePath}\"");
 
 			// Add the additional arguments specified by the environment.
-			LinkCommandArguments += LinkEnvironment.AdditionalArguments;
+			LinkArguments.Add(LinkEnvironment.AdditionalArguments);
+
+			FileReference ResponseFileName;
+			if (bBuildImportLibraryOnly)
+			{
+				ResponseFileName = FileReference.FromString(OutputFile.FullName + ResponseExt);
+			}
+			else
+			{
+				ResponseFileName = GetResponseFileName(LinkEnvironment, OutputFile);
+			}
+			
+			FileItem ResponseFileItem = Graph.CreateIntermediateTextFile(ResponseFileName, LinkArguments);
+			string ResponsFileArgument = GetResponseFileArgument(ResponseFileItem, LinkEnvironment.RootPaths);
+
+			LinkAction.PrerequisiteItems.Add(ResponseFileItem);
+
+			LinkAction.CommandArguments = ResponsFileArgument;
 
 			// Only execute linking on the local PC.
 			LinkAction.bCanExecuteRemotely = false;
 
-			LinkAction.StatusDescription = String.Format("{0}", OutputFile.AbsolutePath);
+			LinkAction.StatusDescription = OutputFile.Location.GetFileName();
 
 			LinkAction.CommandVersion = Info.ClangVersionString;
 
-			LinkAction.CommandPath = BuildHostPlatform.Current.Shell;
+			LinkAction.CommandPath = LinkerPath;
 			if (LinkEnvironment.Configuration == CppConfiguration.Shipping && Path.GetExtension(OutputFile.AbsolutePath) != ".a")
 			{
 				// When building a shipping package, symbols are stripped from the exe as the last build step. This is a problem
 				// when re-packaging and no source files change because the linker skips symbol generation and dsymutil will 
 				// recreate a new .dsym file from a symboless exe file. It's just sad. To make things happy we need to delete 
 				// the output file to force the linker to recreate it with symbols again.
-				string linkCommandArguments = "-c \"";
-
-				linkCommandArguments += String.Format("rm -f \\\"{0}\\\";", OutputFile.AbsolutePath);
-				linkCommandArguments += "\\\"" + LinkerPath + "\\\" " + LinkCommandArguments + ";";
-
-				linkCommandArguments += "\"";
-
-				LinkAction.CommandArguments = linkCommandArguments;
-			}
-			else
-			{
-				// This is not a shipping build so no need to delete the output file since symbols will not have been stripped from it.
-				LinkAction.CommandArguments = String.Format("-c \"\\\"{0}\\\" {1}\"", LinkerPath, LinkCommandArguments);
+				LinkAction.DeleteItems.Add(OutputFile);
 			}
 
 			return OutputFile;
