Fix issue CA2213 (Disposable fields should be disposed) in UBT.

diff --git a/Engine/Source/Programs/UnrealBuildTool/Executors/Experimental/FASTBuild.cs b/Engine/Source/Programs/UnrealBuildTool/Executors/Experimental/FASTBuild.cs
--- a/Engine/Source/Programs/UnrealBuildTool/Executors/Experimental/FASTBuild.cs
+++ b/Engine/Source/Programs/UnrealBuildTool/Executors/Experimental/FASTBuild.cs
@@ -172,13 +172,6 @@ namespace UnrealBuildTool
 			LocalExecutor = new ParallelExecutor(MaxLocalActions, bAllCores, bCompactOutput, Logger);
 		}
 
-		/// <inheritdoc/>
-		public new void Dispose()
-		{
-			LocalExecutor.Dispose();
-			base.Dispose();
-		}
-
 		public override string Name => "FASTBuild";
 
 		public static string GetExecutableName()
@@ -1434,5 +1427,20 @@ namespace UnrealBuildTool
 				return false;
 			}
 		}
+
+		private bool _disposedValue;
+		protected override void Dispose(bool disposing)
+		{
+			if (!_disposedValue)
+			{
+				if (disposing)
+				{
+					LocalExecutor.Dispose();
+				}
+
+				base.Dispose(disposing);
+				_disposedValue = true;
+			}
+		}
 	}
 }
diff --git a/Engine/Source/Programs/UnrealBuildTool/Executors/UnrealBuildAccelerator/UBAAgentCoordinatorHorde.cs b/Engine/Source/Programs/UnrealBuildTool/Executors/UnrealBuildAccelerator/UBAAgentCoordinatorHorde.cs
--- a/Engine/Source/Programs/UnrealBuildTool/Executors/UnrealBuildAccelerator/UBAAgentCoordinatorHorde.cs
+++ b/Engine/Source/Programs/UnrealBuildTool/Executors/UnrealBuildAccelerator/UBAAgentCoordinatorHorde.cs
@@ -126,6 +126,8 @@ namespace UnrealBuildTool
 				await UpdateCpuCoreNeedAsync(0);
 			}
 
+			_hordeHttpClient?.Dispose();
+
 			await _serviceProvider.DisposeAsync();
 			_cancellationTokenSource.Dispose();
 		}
diff --git a/Engine/Source/Programs/UnrealBuildTool/Executors/UnrealBuildAccelerator/UBAExecutor.cs b/Engine/Source/Programs/UnrealBuildTool/Executors/UnrealBuildAccelerator/UBAExecutor.cs
--- a/Engine/Source/Programs/UnrealBuildTool/Executors/UnrealBuildAccelerator/UBAExecutor.cs
+++ b/Engine/Source/Programs/UnrealBuildTool/Executors/UnrealBuildAccelerator/UBAExecutor.cs
@@ -33,7 +33,9 @@ namespace UnrealBuildTool
 		IServer? _server;
 		IStorageServer? _storageServer;
 		ISessionServer? _session;
+#pragma warning disable CA2213 // Disposable fields should be disposed - _trace can be a global trace sometimes, so not safe to dispose.
 		ITrace? _trace;
+#pragma warning restore CA2213
 		string? _sessionHyperlink;
 
 		readonly object _sessionLock = new();
@@ -75,6 +77,8 @@ namespace UnrealBuildTool
 				}
 
 				_server?.Dispose();
+
+				_cancellationSource.Dispose();
 			}
 			base.Dispose(disposing);
 		}
diff --git a/Engine/Source/Programs/UnrealBuildTool/System/SourceFileWorkingSet.cs b/Engine/Source/Programs/UnrealBuildTool/System/SourceFileWorkingSet.cs
--- a/Engine/Source/Programs/UnrealBuildTool/System/SourceFileWorkingSet.cs
+++ b/Engine/Source/Programs/UnrealBuildTool/System/SourceFileWorkingSet.cs
@@ -186,6 +186,8 @@ namespace UnrealBuildTool
 					try
 					{
 						BackgroundProcess.Kill();
+						BackgroundProcess.Dispose();
+						BackgroundProcess = null;
 					}
 					catch
 					{
