remove deprecated CultureId and IntelligentInstallChunks support

#rnx

diff --git a/Engine/Source/Programs/UnrealBuildTool/Platform/Windows/AppXManifestGeneratorBase.cs b/Engine/Source/Programs/UnrealBuildTool/Platform/Windows/AppXManifestGeneratorBase.cs
--- a/Engine/Source/Programs/UnrealBuildTool/Platform/Windows/AppXManifestGeneratorBase.cs
+++ b/Engine/Source/Programs/UnrealBuildTool/Platform/Windows/AppXManifestGeneratorBase.cs
@@ -64,9 +64,6 @@ namespace UnrealBuildTool
 		/// Logger for output
 		protected readonly ILogger Logger;
 
-		/// Whether we have logged the deprecation warning for PerCultureResources CultureId being replaced by StageIdOverrides
-		protected static bool bHasWarnedAboutDeprecatedCultureId = false;
-
 		/// CustomConfig to use when reading the ini files
 		public string CustomConfig = "";
 
@@ -247,11 +244,9 @@ namespace UnrealBuildTool
 			}
 
 			// read StageId overrides
-			bool bHasStageIdOverrides = false;
 			if (EngineIni!.GetString(IniSection_PlatformTargetSettings, "StageIdOverrides", out string? StageIdOverridesString) &&
 				ConfigHierarchy.TryParseAsMap(StageIdOverridesString, out Dictionary<string, string>? StageIdOverrides))
 			{
-				bHasStageIdOverrides = true;
 				UEStageIdToAppXCultureId = StageIdOverrides;
 				AppXResources!.AddCultures(UEStageIdToAppXCultureId.Values);
 			}
@@ -279,34 +274,16 @@ namespace UnrealBuildTool
 					}
 
 					string CultureId = "";
-					if (bHasStageIdOverrides)
+					if (UEStageIdToAppXCultureId.ContainsKey(StageId))
 					{
-						CultureId = UEStageIdToAppXCultureId.ContainsKey(StageId) ? UEStageIdToAppXCultureId[StageId] : StageId;
-					}
-					else if (!CultureProperties.ContainsKey("CultureId"))
-					{
-						Logger.LogWarning("Invalid per-culture resource value: {Culture}", CultureResources);
+						CultureId = UEStageIdToAppXCultureId[StageId];
 					}
 					else
 					{
-						CultureId = CultureProperties["CultureId"];
-						if (String.IsNullOrEmpty(CultureId))
-						{
-							Logger.LogWarning("Missing CultureId value: {Culture}", CultureResources);
-							continue;
-						}
-
-						// PerCultureResources CultureId is deprecated in favor of new StageIdOverrides property
-						// if the CultureId field contains an overridden property, warn they need to update the data
-						if (CultureId != StageId && !bHasWarnedAboutDeprecatedCultureId)
-						{
-							Logger.LogWarning("PerCultureResources is out of date - please re-save this project's {Platform}Engine.ini in the editor to update StageIdOverrides. This must be done before UE5.5 to avoid losing data to deprecation", Platform);
-							bHasWarnedAboutDeprecatedCultureId = true;
-						}
-
+						CultureId = StageId;
 						UEStageIdToAppXCultureId[StageId] = CultureId;
+						AppXResources.AddCulture(CultureId);
 					}
-					AppXResources.AddCulture(CultureId);
 
 					// read culture strings
 					if (!ConfigHierarchy.TryParse(CultureProperties["CultureStringResources"], out Dictionary<string, string>? CultureStringResources))
