[Core] Add ability to enable compact set via build config or command line

#tests Compiled locally, viewed output on preflight (had message from container tests reported on some builds where it was present to confirm)

diff --git a/Engine/Source/Programs/UnrealBuildTool/Configuration/ReadOnlyTargetRules.cs b/Engine/Source/Programs/UnrealBuildTool/Configuration/ReadOnlyTargetRules.cs
--- a/Engine/Source/Programs/UnrealBuildTool/Configuration/ReadOnlyTargetRules.cs
+++ b/Engine/Source/Programs/UnrealBuildTool/Configuration/ReadOnlyTargetRules.cs
@@ -269,6 +269,8 @@ namespace UnrealBuildTool
 
 		public bool bDefineForceInlineHintToInline => Inner.bDefineForceInlineHintToInline;
 
+		public bool bUseCompactSetAsDefault => Inner.bUseCompactSetAsDefault;
+
 		public bool bUseAndroidJni => Inner.bUseAndroidJni;
 
 		public FPSemanticsMode FPSemantics => Inner.FPSemantics;
diff --git a/Engine/Source/Programs/UnrealBuildTool/Configuration/TargetRules.cs b/Engine/Source/Programs/UnrealBuildTool/Configuration/TargetRules.cs
--- a/Engine/Source/Programs/UnrealBuildTool/Configuration/TargetRules.cs
+++ b/Engine/Source/Programs/UnrealBuildTool/Configuration/TargetRules.cs
@@ -1598,6 +1598,14 @@ namespace UnrealBuildTool
 		}
 		private bool? bDefineForceInlineHintToInlinePrivate;
 
+		/// <summary>
+		/// By default TSparseSet is used as the underlying TSet container. Use this flag to change that to use TCompactSet instead.
+		/// </summary>
+		[RequiresUniqueBuildEnvironment]
+		[CommandLine("-UseCompactSetAsDefault")]
+		[XmlConfigFile(Category = "BuildConfiguration")]
+		public bool bUseCompactSetAsDefault { get; set; } = false;
+
 		/// <summary>
 		/// Whether to use Android Java Native Interface.
 		/// </summary>
diff --git a/Engine/Source/Programs/UnrealBuildTool/Configuration/UEBuildTarget.cs b/Engine/Source/Programs/UnrealBuildTool/Configuration/UEBuildTarget.cs
--- a/Engine/Source/Programs/UnrealBuildTool/Configuration/UEBuildTarget.cs
+++ b/Engine/Source/Programs/UnrealBuildTool/Configuration/UEBuildTarget.cs
@@ -6504,6 +6504,15 @@ namespace UnrealBuildTool
 				GlobalCompileEnvironment.Definitions.Add("UE_DEFINE_FORCEINLINE_HINT_TO_INLINE=0");
 			}
 
+			if (Rules.bUseCompactSetAsDefault)
+			{
+				GlobalCompileEnvironment.Definitions.Add("UE_USE_COMPACT_SET_AS_DEFAULT=1");
+			}
+			else
+			{
+				GlobalCompileEnvironment.Definitions.Add("UE_USE_COMPACT_SET_AS_DEFAULT=0");
+			}
+
 			// Define the custom config if specified, this should not be set in an editor build
 			if (!bUseSharedBuildEnvironment && !String.IsNullOrEmpty(Rules.CustomConfig))
 			{
