UnrealBuildTool: VFS fixes for apple toolchains
Move a few common link arguments out of mac & ios toolchain to apple toolchain

#rnx

diff --git a/Engine/Source/Programs/UnrealBuildTool/Platform/IOS/IOSToolChain.cs b/Engine/Source/Programs/UnrealBuildTool/Platform/IOS/IOSToolChain.cs
--- a/Engine/Source/Programs/UnrealBuildTool/Platform/IOS/IOSToolChain.cs
+++ b/Engine/Source/Programs/UnrealBuildTool/Platform/IOS/IOSToolChain.cs
@@ -304,15 +304,6 @@ namespace UnrealBuildTool
 		protected override void GetLinkArguments_Global(LinkEnvironment LinkEnvironment, List<string> Arguments)
 		{
 			base.GetLinkArguments_Global(LinkEnvironment, Arguments);
-			Arguments.Add(FormatArchitectureArg(LinkEnvironment.Architectures));
-
-			DirectoryReference SDKPath = Settings.GetSDKPath(LinkEnvironment.Architecture);
-			Arguments.Add($"-isysroot \"{SDKPath}\"");
-
-			// needed to find swift directories for some versions of the SDK
-			string PlatformName = LinkEnvironment.Architecture == UnrealArch.IOSSimulator ? Settings.SimulatorPlatformDirName : Settings.PlatformDirName;
-			FileReference PlatformLibDirectory = FileReference.Combine(ApplePlatformSDK.GetToolchainDirectory(), "..", "lib", "swift", PlatformName);
-			Arguments.Add($"-L\"{PlatformLibDirectory}\"");
 
 			if (!LinkEnvironment.bIsBuildingLibrary && !LinkEnvironment.bIsBuildingDLL)
 			{
@@ -338,8 +329,6 @@ namespace UnrealBuildTool
 				Arguments.Add("-Wl,-exported_symbol,__ZdaPvRKSt9nothrow_t");
 			}
 
-			Arguments.Add("-dead_strip");
-			Arguments.Add($"-target {Settings.GetTargetTuple(LinkEnvironment.Architecture, Target!.Platform, Target!.Type)}");
 			Arguments.Add("-Wl-no_pie");
 			Arguments.Add("-stdlib=libc++");
 			Arguments.Add("-ObjC");
@@ -366,15 +355,14 @@ namespace UnrealBuildTool
 				if (FrameworkDirectory != null)
 				{
 					// If this framework has a directory specified, we'll need to setup the path as well
-					string FrameworkDir = FrameworkDirectory.FullName;
 
 					// embedded frameworks have a framework inside of this directory, so we use this directory. regular frameworks need to go one up to point to the 
 					// directory containing the framework. -F gives a path to look for the -framework
-					if (FrameworkDir.EndsWith(".framework"))
+					if (FrameworkDirectory.FullName.EndsWith(".framework"))
 					{
-						FrameworkDir = Path.GetDirectoryName(FrameworkDir)!;
+						FrameworkDirectory = FrameworkDirectory.ParentDirectory!;
 					}
-					Arguments.Add($"-F\"{FrameworkDir}\"");
+					Arguments.Add($"-F\"{NormalizeCommandLinePath(FrameworkDirectory, LinkEnvironment.RootPaths)}\"");
 				}
 
 				Arguments.Add($"-framework {Framework.Name}");
@@ -392,7 +380,13 @@ namespace UnrealBuildTool
 			// Create an action that invokes the linker.
 			Action LinkAction = Graph.CreateAction(ActionType.Link);
 			string Arch = UnrealArchitectureConfig.ForPlatform(LinkEnvironment.Platform).ConvertToReadableArchitecture(LinkEnvironment.Architecture);
-			LinkAction.CommandDescription = (bBuildImportLibraryOnly ? "LinkStub" : "Link") + $" [{Arch}]";
+			LinkAction.CommandPath = LinkerPath;
+			LinkAction.CommandDescription = $"{(bBuildImportLibraryOnly ? "LinkStub" : "Link")} [{Arch}]";
+			LinkAction.CommandVersion = Info.ClangVersionString;
+			LinkAction.bProducesImportLibrary = bBuildImportLibraryOnly || LinkEnvironment.bIsBuildingDLL;
+			LinkAction.bCanExecuteInUBA = true;
+			LinkAction.CacheBucket = GetCacheBucket(Target, LinkEnvironment);
+			LinkAction.ArtifactMode = ArtifactMode.Enabled;
 			LinkAction.RootPaths = LinkEnvironment.RootPaths;
 
 			// RPC utility parameters are in terms of the Mac side
@@ -414,7 +408,7 @@ namespace UnrealBuildTool
 				LinkerPath = Info.Clang;
 				LinkArguments.Add("-dynamiclib -Xlinker -export_dynamic -Xlinker -no_deduplicate");
 
-				string InstallName = LinkEnvironment.InstallName ?? String.Format("@executable_path/Frameworks/{0}", LinkEnvironment.OutputFilePath.MakeRelativeTo(LinkEnvironment.OutputFilePath.Directory.ParentDirectory!));
+				string InstallName = LinkEnvironment.InstallName ?? $"@executable_path/Frameworks/{LinkEnvironment.OutputFilePath.MakeRelativeTo(LinkEnvironment.OutputFilePath.Directory.ParentDirectory!)}";
 				LinkArguments.Add($"-Xlinker -install_name -Xlinker {InstallName}");
 
 				LinkArguments.Add("-Xlinker -rpath -Xlinker @executable_path/Frameworks");
@@ -426,7 +420,7 @@ namespace UnrealBuildTool
 				// Add the library paths to the argument list.
 				foreach (DirectoryReference LibraryPath in LinkEnvironment.SystemLibraryPaths)
 				{
-					LinkArguments.Add($"-L\"{LibraryPath.FullName}\"");
+					LinkArguments.Add($"-L\"{NormalizeCommandLinePath(LibraryPath, LinkEnvironment.RootPaths)}\"");
 				}
 
 				// Add the additional libraries to the argument list.
@@ -443,19 +437,18 @@ namespace UnrealBuildTool
 					else
 					{
                         FileReference LibraryPath = FileReference.FromString(AdditionalLibrary);
-						LinkArguments.Add($"-l\"{LibraryPath}\"");
+						LinkArguments.Add($"-l\"{NormalizeCommandLinePath(LibraryPath, LinkEnvironment.RootPaths)}\"");
                     }					
 				}
 
-				foreach (FileReference Library in LinkEnvironment.Libraries)
+				foreach (FileReference LibraryPath in LinkEnvironment.Libraries)
 				{
 					// for absolute library paths, convert to remote filename
 					// add it to the prerequisites to make sure it's built first (this should be the case of non-system libraries)
-					FileItem LibFile = FileItem.GetItemByFileReference(Library);
-					LinkAction.PrerequisiteItems.Add(LibFile);
+					LinkAction.PrerequisiteItems.Add(FileItem.GetItemByFileReference(LibraryPath));
 
 					// and add to the commandline
-					LinkArguments.Add($"\"{Library.FullName}\"");
+					LinkArguments.Add($"\"{NormalizeCommandLinePath(LibraryPath, LinkEnvironment.RootPaths)}\"");
 				}
 			}
 
@@ -477,36 +470,35 @@ namespace UnrealBuildTool
 			if ((!LinkEnvironment.bIsBuildingLibrary || LinkEnvironment.bIsBuildingDLL) && LinkEnvironment.bCreateMapFile)
 			{
 				FileItem MapFile = FileItem.GetItemByFileReference(new FileReference(OutputFile.Location.FullName + ".map"));
-				LinkArguments.Add($"-Wl,-map,\"{MapFile.Location.FullName}\"");
+				LinkArguments.Add($"-Wl,-map,\"{NormalizeCommandLinePath(MapFile, LinkEnvironment.RootPaths)}\"");
 				LinkAction.ProducedItems.Add(MapFile);
 			}
 
 			// Add the input files to a response file, and pass the response file on the command-line.
-			List<string> InputFileNames = new List<string>();
 			foreach (FileItem InputFile in LinkEnvironment.InputFiles)
 			{
-				LinkArguments.Add($"\"{InputFile.AbsolutePath}\"");
+				LinkArguments.Add($"\"{NormalizeCommandLinePath(InputFile, LinkEnvironment.RootPaths)}\"");
 				LinkAction.PrerequisiteItems.Add(InputFile);
 			}
 
 			// if we are making an LTO build, write the lto file next to build so dsymutil can find it
 			if (LinkEnvironment.bAllowLTCG)
 			{
-				string LtoObjectFile;
+				DirectoryReference LtoObjectDir;
 				if (Target!.bShouldCompileAsDLL)
 				{
 					// go up a directory so we don't put this big file into the framework
-					LtoObjectFile = Path.Combine(Path.GetDirectoryName(Path.GetDirectoryName(OutputFile.AbsolutePath))!, Path.GetFileName(OutputFile.AbsolutePath) + ".lto.o");
+					LtoObjectDir = DirectoryReference.Combine(OutputFile.Directory.Location.ParentDirectory!, $"{OutputFile.Name}.lto.o");
 				}
 				else
 				{
-					LtoObjectFile = OutputFile.AbsolutePath + ".lto.o";
+					LtoObjectDir = DirectoryReference.Combine(OutputFile.Directory.Location, $"{OutputFile.Name}.lto.o");
 				}
-				LinkArguments.Add($"-flto -Xlinker -object_path_lto -Xlinker \"{LtoObjectFile}\"");
+				LinkArguments.Add($"-flto -Xlinker -object_path_lto -Xlinker \"{NormalizeCommandLinePath(LtoObjectDir, LinkEnvironment.RootPaths)}\"");
 			}
 
 			// Add the output file to the command-line.
-			LinkArguments.Add($"-o \"{OutputFile.AbsolutePath}\"");
+			LinkArguments.Add($"-o \"{NormalizeCommandLinePath(OutputFile, LinkEnvironment.RootPaths)}\"");
 
 			// Add the additional arguments specified by the environment.
 			LinkArguments.Add(LinkEnvironment.AdditionalArguments);
@@ -533,9 +525,6 @@ namespace UnrealBuildTool
 
 			LinkAction.StatusDescription = OutputFile.Location.GetFileName();
 
-			LinkAction.CommandVersion = Info.ClangVersionString;
-
-			LinkAction.CommandPath = LinkerPath;
 			if (LinkEnvironment.Configuration == CppConfiguration.Shipping && Path.GetExtension(OutputFile.AbsolutePath) != ".a")
 			{
 				// When building a shipping package, symbols are stripped from the exe as the last build step. This is a problem
diff --git a/Engine/Source/Programs/UnrealBuildTool/Platform/Mac/MacToolChain.cs b/Engine/Source/Programs/UnrealBuildTool/Platform/Mac/MacToolChain.cs
--- a/Engine/Source/Programs/UnrealBuildTool/Platform/Mac/MacToolChain.cs
+++ b/Engine/Source/Programs/UnrealBuildTool/Platform/Mac/MacToolChain.cs
@@ -220,14 +220,15 @@ namespace UnrealBuildTool
 			}
 		}
 
-		void AddFrameworkToLinkCommand(List<string> LinkArguments, string FrameworkName, string Arg = "-framework")
+		void AddFrameworkToLinkCommand(LinkEnvironment linkEnvironment, List<string> LinkArguments, string FrameworkName, string Arg = "-framework")
 		{
 			if (FrameworkName.EndsWith(".framework"))
 			{
-				LinkArguments.Add("-F \"" + Path.GetDirectoryName(Path.GetFullPath(FrameworkName)) + "\"");
+				FileReference frameworkPath = new FileReference(FrameworkName);
+				LinkArguments.Add($"-F \"{NormalizeCommandLinePath(frameworkPath.Directory, linkEnvironment.RootPaths)}\"");
 				FrameworkName = Path.GetFileNameWithoutExtension(FrameworkName);
 			}
-			LinkArguments.Add(Arg + " \"" + FrameworkName + "\"");
+			LinkArguments.Add($"{Arg} \"{FrameworkName}\"");
 		}
 
 		protected override void GetLinkArguments_Sanitizers(LinkEnvironment linkEnvironment, List<string> arguments)
@@ -246,11 +247,6 @@ namespace UnrealBuildTool
 		protected override void GetLinkArguments_Global(LinkEnvironment LinkEnvironment, List<string> Arguments)
 		{
 			base.GetLinkArguments_Global(LinkEnvironment, Arguments);
-			// Pass through architecture and OS info		
-			Arguments.Add(FormatArchitectureArg(LinkEnvironment.Architectures));
-			Arguments.Add(String.Format("-isysroot \"{0}\"", NormalizeCommandLinePath(Settings.GetSDKPath(LinkEnvironment.Architecture), LinkEnvironment.RootPaths)));
-			Arguments.Add($"-target {Settings.GetTargetTuple(LinkEnvironment.Architecture, Target!.Platform, Target!.Type)}");
-			Arguments.Add("-dead_strip");
 
 			// Temporary workaround for linker warning with Xcode 14:
 			//		'ld: warning: could not create compact unwind for _inflate_fast: registers 27 not saved contiguously in frame'
@@ -687,7 +683,7 @@ namespace UnrealBuildTool
 			LinkAction.WorkingDirectory = GetMacDevSrcRoot();
 			LinkAction.CommandPath = LinkerPath;
 			string Arch = UnrealArchitectureConfig.ForPlatform(UnrealTargetPlatform.Mac).ConvertToReadableArchitecture(LinkEnvironment.Architecture);
-			LinkAction.CommandDescription = (bBuildImportLibraryOnly ? "LinkStub" : "Link") + $" [{Arch}]";
+			LinkAction.CommandDescription = $"{(bBuildImportLibraryOnly ? "LinkStub" : "Link")} [{Arch}]";
 			LinkAction.CommandVersion = Info.ClangVersionString;
 			LinkAction.bProducesImportLibrary = bBuildImportLibraryOnly || LinkEnvironment.bIsBuildingDLL;
 			LinkAction.bCanExecuteInUBA = true;
@@ -830,7 +826,7 @@ namespace UnrealBuildTool
 					}
 					else if (AdditionalLibrary.Contains(".framework/"))
 					{
-						LinkArguments.Add($"\"{NormalizeCommandLinePath(FileReference.FromString(AdditionalLibrary), LinkEnvironment.RootPaths)}\"");
+						LinkArguments.Add($"\"{NormalizeCommandLinePath(FileItem.GetItemByPath(AdditionalLibrary), LinkEnvironment.RootPaths)}\"");
 					}
 					else if (String.IsNullOrEmpty(Path.GetDirectoryName(AdditionalLibrary)) && String.IsNullOrEmpty(Path.GetExtension(AdditionalLibrary)))
 					{
@@ -849,7 +845,7 @@ namespace UnrealBuildTool
 
 				foreach (string AdditionalLibrary in LinkEnvironment.DelayLoadDLLs)
 				{
-					LinkArguments.Add(String.Format("-weak_library \"{0}\"", Path.GetFullPath(AdditionalLibrary)));
+					LinkArguments.Add($"-weak_library \"{NormalizeCommandLinePath(FileItem.GetItemByPath(AdditionalLibrary), LinkEnvironment.RootPaths)}\"");
 
 					if (!bCanUseSimpleRPaths)
 					{
@@ -889,7 +885,7 @@ namespace UnrealBuildTool
 
 				foreach (KeyValuePair<string, bool> Framework in AllFrameworks)
 				{
-					AddFrameworkToLinkCommand(LinkArguments, Framework.Key, Framework.Value ? "-weak_framework" : "-framework");
+					AddFrameworkToLinkCommand(LinkEnvironment, LinkArguments, Framework.Key, Framework.Value ? "-weak_framework" : "-framework");
 					AddLibraryPathToRPaths(Framework.Key, AbsolutePath, ref RPaths, LinkArguments, bIsBuildingAppBundle, Logger);
 				}
 
diff --git a/Engine/Source/Programs/UnrealBuildTool/ToolChain/AppleToolChain.cs b/Engine/Source/Programs/UnrealBuildTool/ToolChain/AppleToolChain.cs
--- a/Engine/Source/Programs/UnrealBuildTool/ToolChain/AppleToolChain.cs
+++ b/Engine/Source/Programs/UnrealBuildTool/ToolChain/AppleToolChain.cs
@@ -113,9 +113,6 @@ namespace UnrealBuildTool
 			SDKVersion = ApplePlatformSDK.GetPlatformSDKVersion(OSPrefix) ?? "";
 			SDKVersionFloat = ApplePlatformSDK.GetPlatformSDKVersionFloat(OSPrefix);
 
-			// Xcode sdk files are not allowed to be transferred over the network, ensure UBA will not do so
-			EpicGames.UBA.Utils.RegisterDisallowedPaths(ApplePlatformSDK.DeveloperDir.FullName);
-
 			TestXcode(bVerbose, Logger);
 		}
 
@@ -266,6 +263,28 @@ namespace UnrealBuildTool
 			AppleExports.GetSwiftIntegrationSettings(ProjectFile, Target == null ? TargetType.Game : Target.Type, Target == null ? UnrealTargetPlatform.Mac : Target.Platform, out bUseSwiftUIMain, out bCreateSwiftBridgingHeader);
 		}
 
+		/// <inheritdoc/>
+		public override void SetUpGlobalEnvironment(ReadOnlyTargetRules Target, CppCompileEnvironment GlobalCompileEnvironment, LinkEnvironment GlobalLinkEnvironment)
+		{
+			base.SetUpGlobalEnvironment(Target, GlobalCompileEnvironment, GlobalLinkEnvironment);
+
+			DirectoryReference XcodeDeveloperDir = ApplePlatformSDK.DeveloperDir;
+			DirectoryReference? XcodeAppDir = XcodeDeveloperDir.ParentDirectory?.ParentDirectory;
+
+			if (XcodeAppDir != null)
+			{
+				GlobalCompileEnvironment.RootPaths.AddFolderName(CppRootPathFolder.PlatformSDK, "Xcode.app");
+				GlobalLinkEnvironment.RootPaths.AddFolderName(CppRootPathFolder.PlatformSDK, "Xcode.app");
+				GlobalCompileEnvironment.RootPaths[CppRootPathFolder.PlatformSDK] = XcodeAppDir;
+				GlobalLinkEnvironment.RootPaths[CppRootPathFolder.PlatformSDK] = XcodeAppDir;
+
+				EpicGames.UBA.Utils.RegisterPathHash(XcodeAppDir.FullName, ApplePlatformSDK.InstalledXcodeVersion.Value ?? "0");
+
+				// Xcode sdk files are not allowed to be transferred over the network, ensure UBA will not do so
+				EpicGames.UBA.Utils.RegisterDisallowedPaths(XcodeAppDir.FullName);
+			}
+		}
+
 		/// <summary>
 		/// Takes an architecture string as provided by UBT for the target and formats it for Clang. Supports
 		/// multiple architectures joined with '+'
@@ -871,22 +890,29 @@ namespace UnrealBuildTool
 		{
 			base.GetLinkArguments_Global(LinkEnvironment, Arguments);
 
+			Arguments.Add($"-isysroot \"{NormalizeCommandLinePath(Settings.GetSDKPath(LinkEnvironment.Architecture), LinkEnvironment.RootPaths)}\"");
+
 			VersionNumber AppleClangVersion = ((AppleToolChainInfo)Info).AppleClangVersion;
 			// Temp solution for UE-191350
 			if (AppleClangVersion >= new VersionNumber(15) && AppleClangVersion < new VersionNumber(16))
 			{
-				Arguments.Add(" -ld_classic");
+				Arguments.Add("-ld_classic");
 			}
 
 			if (!LinkEnvironment.bIsBuildingDLL)
 			{
-				// Todo: This is currently only used when linking binary and swift files have not been tested with dylibs.
-				// The right fix is to move the vfs root one step up to be able to handle the ../lib/ path
-				DirectoryReference dir = DirectoryReference.Combine(ApplePlatformSDK.GetToolchainDirectory(), "../lib/swift", Settings.PlatformDirName);
-				Arguments.Add($"-L{dir}");
+				string PlatformName = LinkEnvironment.Architecture == UnrealArch.IOSSimulator ? Settings.SimulatorPlatformDirName : Settings.PlatformDirName;
+				DirectoryReference libSwiftDir = DirectoryReference.Combine(ApplePlatformSDK.GetToolchainDirectory(), "..", "lib", "swift", PlatformName);
+				Arguments.Add($"-L\"{NormalizeCommandLinePath(libSwiftDir, LinkEnvironment.RootPaths)}\"");
 			}
 
-			Arguments.Add("-L/usr/lib/swift");
+			Arguments.Add("-L\"/usr/lib/swift\"");
+
+			// Pass through architecture and OS info		
+			Arguments.Add(FormatArchitectureArg(LinkEnvironment.Architectures));
+			Arguments.Add($"-target {Settings.GetTargetTuple(LinkEnvironment.Architecture, Target!.Platform, Target!.Type)}");
+
+			Arguments.Add("-dead_strip");
 		}
 
 		protected override void GetArchiveArguments_Global(LinkEnvironment LinkEnvironment, List<string> Arguments)
