Updated Build Target logic to add "lib" prefix to libraries on Android, iOS, Linux & Mac platform groups.
- Previously, "lib" prefix was only added to Linux targets.
- Now, "lib" prefix is added to Android, iOS, Linux, & Mac platform groups.
- Checking for Linux platform group now includes LinuxArm64 instead of only Linux.

diff --git a/Engine/Source/Programs/UnrealBuildTool/Configuration/UEBuildTarget.cs b/Engine/Source/Programs/UnrealBuildTool/Configuration/UEBuildTarget.cs
--- a/Engine/Source/Programs/UnrealBuildTool/Configuration/UEBuildTarget.cs
+++ b/Engine/Source/Programs/UnrealBuildTool/Configuration/UEBuildTarget.cs
@@ -5368,7 +5368,14 @@ namespace UnrealBuildTool
 		{
 			StringBuilder Result = new StringBuilder();
 
-			if (Platform == UnrealTargetPlatform.Linux && (BinaryType == UEBuildBinaryType.DynamicLinkLibrary || BinaryType == UEBuildBinaryType.StaticLibrary))
+			// If the binary is a library and is on a platform that expects the
+			// "lib" library prefix, add the prefix to the library filename.
+			// Platforms that expect this are Android, Apple (Mac & iOS), and Linux.
+			bool bBinaryIsLibrary = BinaryType is UEBuildBinaryType.DynamicLinkLibrary or UEBuildBinaryType.StaticLibrary;
+			bool bRequiresLibraryFilenamePrefix = UEBuildPlatform.IsPlatformInGroup(Platform, UnrealPlatformGroup.Android) ||
+			                                      UEBuildPlatform.IsPlatformInGroup(Platform, UnrealPlatformGroup.Apple) ||
+			                                      UEBuildPlatform.IsPlatformInGroup(Platform, UnrealPlatformGroup.Linux);
+			if (bBinaryIsLibrary && bRequiresLibraryFilenamePrefix)
 			{
 				Result.Append("lib");
 			}
