[UBT]
* Changed so generated file's module dirs are always created when fake-populating. This fix a bug where totally cleaned builds (no generated folders created) end up missing generated files in manifest.
* Made code run in parallel since CreateDirectory is a kernel function and sometimes take time

diff --git a/Engine/Source/Programs/UnrealBuildTool/System/UHTExecution.cs b/Engine/Source/Programs/UnrealBuildTool/System/UHTExecution.cs
--- a/Engine/Source/Programs/UnrealBuildTool/System/UHTExecution.cs
+++ b/Engine/Source/Programs/UnrealBuildTool/System/UHTExecution.cs
@@ -1149,7 +1149,7 @@ namespace UnrealBuildTool
 		/// <param name="moduleInfos"></param>
 		public static void PrepopulateDirectoryCache(List<UHTModuleInfo> moduleInfos)
 		{
-			foreach (UHTModuleInfo mi in moduleInfos)
+			Parallel2.ForEach(moduleInfos, mi =>
 			{
 				List<string> outputFiles = [];
 				outputFiles.Add(mi.ModuleName + ".init.gen.cpp");
@@ -1162,8 +1162,9 @@ namespace UnrealBuildTool
 					outputFiles.Add(fileName + ".gen.cpp");
 				}
 
+				DirectoryReference.CreateDirectory(mi.GeneratedCodeDirectory.Location);
 				mi.GeneratedCodeDirectory.PrepopulateCache(outputFiles);
-			}
+			});
 		}
 
 		static UInt128 CalculateHash(UHTModuleInfo mi)
