[UBT]
* Added so file created with -UBAActionsOutputFile also adds vfs roots per action

diff --git a/Engine/Source/Programs/UnrealBuildTool/Executors/UnrealBuildAccelerator/UBAExecutor.cs b/Engine/Source/Programs/UnrealBuildTool/Executors/UnrealBuildAccelerator/UBAExecutor.cs
--- a/Engine/Source/Programs/UnrealBuildTool/Executors/UnrealBuildAccelerator/UBAExecutor.cs
+++ b/Engine/Source/Programs/UnrealBuildTool/Executors/UnrealBuildAccelerator/UBAExecutor.cs
@@ -441,6 +441,21 @@ namespace UnrealBuildTool
 					await writer.WriteAsync(String.Join(", ", action.PrerequisiteActions.Select(x => x.SortIndex)));
 					await writer.WriteLineAsync("]");
 				}
+				if (action.RootPaths.Any())
+				{
+					await writer.WriteAsync("vfs: [");
+					bool isFirst = true;
+					foreach ((uint id, DirectoryReference vfs, DirectoryReference local) in action.RootPaths)
+					{
+						if (!isFirst)
+							await writer.WriteAsync(", ");
+						isFirst = false;
+						await writer.WriteAsync(vfs.FullName);
+						await writer.WriteAsync(';');
+						await writer.WriteAsync(local.FullName);
+					}
+					await writer.WriteLineAsync("]");
+				}
 				writer.Indent--;
 				await writer.WriteLineNoTabsAsync(null);
 			}
