[UBT]
* Added threadsafe AddOutputFiles and AddInternalDependency to TargetMakefileBuilder and changed so calling code uses them isntead of doing builder.Makefile.OutputFiles.Add etc
* Changed Utils.WriteClass to use using for StreamWriter instead of finally with close

diff --git a/Engine/Source/Programs/UnrealBuildTool/Configuration/UEBuildTarget.cs b/Engine/Source/Programs/UnrealBuildTool/Configuration/UEBuildTarget.cs
--- a/Engine/Source/Programs/UnrealBuildTool/Configuration/UEBuildTarget.cs
+++ b/Engine/Source/Programs/UnrealBuildTool/Configuration/UEBuildTarget.cs
@@ -3011,7 +3011,7 @@ namespace UnrealBuildTool
 			foreach (UEBuildModule Module in Modules.Values)
 			{
 				IEnumerable<FileItem> Items = Module.CopyDebuggerVisualizers(TargetToolChain, MakefileBuilder, Logger);
-				Makefile.OutputItems.AddRange(Items);
+				MakefileBuilder.AddOutputFiles(Items);
 			}
 
 			// Cache module dependencies for StableTopologicalSort
@@ -3025,16 +3025,13 @@ namespace UnrealBuildTool
 				{
 					Binary.PrepareRuntimeDependencies(ExeDir);
 					List<FileItem> BinaryOutputItems = Binary.Build(Rules, TargetToolChain, GlobalCompileEnvironment, GlobalLinkEnvironment, WorkingSet, ExeDir, MakefileBuilder, Logger);
-					lock (Makefile)
-					{
-						Makefile.OutputItems.AddRange(BinaryOutputItems);
-					}
+					MakefileBuilder.AddOutputFiles(BinaryOutputItems);
 				});
 			}
 
 			if (Rules.bStripExports)
 			{
-				CreateStripExportsActions(TargetToolChain, MakefileBuilder, Makefile, GlobalCompileEnvironment);
+				CreateStripExportsActions(TargetToolChain, MakefileBuilder, GlobalCompileEnvironment);
 			}
 
 			// Cache inline gen cpp data
@@ -3405,7 +3402,7 @@ namespace UnrealBuildTool
 			CleanStaleModules(Logger);
 		}
 
-		void CreateStripExportsActions(UEToolChain TargetToolChain, TargetMakefileBuilder MakefileBuilder, TargetMakefile Makefile, CppCompileEnvironment CompileEnvironment)
+		void CreateStripExportsActions(UEToolChain TargetToolChain, TargetMakefileBuilder MakefileBuilder, CppCompileEnvironment CompileEnvironment)
 		{
 			FileReference commandPath;
 			if (OperatingSystem.IsWindows())
@@ -3480,7 +3477,7 @@ namespace UnrealBuildTool
 				exportsAction.RootPaths = CompileEnvironment.RootPaths;
 				exportsAction.CacheBucket = ISPCToolChain.GetCacheBucket(Rules, CompileEnvironment);
 
-				Makefile.OutputItems.AddRange(exportsAction.ProducedItems);
+				MakefileBuilder.AddOutputFiles(exportsAction.ProducedItems);
 			}
 
 			foreach (UEBuildBinary Binary in Binaries.Where(x => x.bStripUnusedExports))
@@ -3505,8 +3502,6 @@ namespace UnrealBuildTool
 
 			if (!ProducedItems.Any(x => Unreal.IsFileInstalled(x)))
 			{
-				TargetMakefile Makefile = MakefileBuilder.Makefile;
-
 				// Save the info file, but deliberately do not add a prerequisite from the action onto it unless we're updating the
 				// target file. Since the outputs are not target specific, but the input file is always beneath the project intermediate
 				// directory (even for the engine), it will cause -NoEngineChanges to block the build when updated. The behavior we want
@@ -3515,7 +3510,7 @@ namespace UnrealBuildTool
 				FileItem InfoFileItem = FileItem.GetItemByFileReference(InfoFile);
 				BinaryFormatterUtils.SaveIfDifferent(InfoFile, Info);
 				InfoFileItem.ResetCachedInfo();
-				Makefile.InternalDependencies.Add(InfoFileItem);
+				MakefileBuilder.AddInternalDependency(InfoFileItem);
 
 				// Get the argument list
 				StringBuilder WriteMetadataArguments = new StringBuilder();
@@ -3544,7 +3539,7 @@ namespace UnrealBuildTool
 				}
 
 				// Add the produced items as leaf nodes to be built
-				Makefile.OutputItems.AddRange(WriteMetadataAction.ProducedItems);
+				MakefileBuilder.AddOutputFiles(WriteMetadataAction.ProducedItems);
 			}
 		}
 
diff --git a/Engine/Source/Programs/UnrealBuildTool/Platform/Windows/VCToolChain.cs b/Engine/Source/Programs/UnrealBuildTool/Platform/Windows/VCToolChain.cs
--- a/Engine/Source/Programs/UnrealBuildTool/Platform/Windows/VCToolChain.cs
+++ b/Engine/Source/Programs/UnrealBuildTool/Platform/Windows/VCToolChain.cs
@@ -2477,7 +2477,7 @@ namespace UnrealBuildTool
 				// Win64 PGO folder is Windows, the rest match the platform name
 				string PGOPlatform = Target.Platform == UnrealTargetPlatform.Win64 ? "Windows" : Target.Platform.ToString();
 				DirectoryReference PGODirectory = DirectoryReference.Combine(Target.ProjectFile?.Directory ?? Unreal.WritableEngineDirectory, "Platforms", PGOPlatform, "Build", "PGO");
-				MakefileBuilder.Makefile.InternalDependencies.Add(FileItem.GetItemByFileReference(FileReference.Combine(PGODirectory, "PGOProfileCompilerInfo.txt")));
+				MakefileBuilder.AddInternalDependency(FileItem.GetItemByFileReference(FileReference.Combine(PGODirectory, "PGOProfileCompilerInfo.txt")));
 			}
 
 			if (Target.bPrintToolChainTimingInfo || Target.WindowsPlatform.bCompilerTrace)
diff --git a/Engine/Source/Programs/UnrealBuildTool/System/TargetMakefile.cs b/Engine/Source/Programs/UnrealBuildTool/System/TargetMakefile.cs
--- a/Engine/Source/Programs/UnrealBuildTool/System/TargetMakefile.cs
+++ b/Engine/Source/Programs/UnrealBuildTool/System/TargetMakefile.cs
@@ -92,10 +92,7 @@ namespace UnrealBuildTool
 		/// <inheritdoc/>
 		public void CreateIntermediateTextFile(FileItem FileItem, string Contents, bool AllowAsync = true)
 		{
-			lock (Makefile.InternalDependencies)
-			{
-				Makefile.InternalDependencies.Add(FileItem);
-			}
+			AddInternalDependency(FileItem);
 
 			Task? task = FileHashCache.WriteFileIfChanged(FileItem, Contents, Logger, AllowAsync);
 			if (task != null)
@@ -107,10 +104,7 @@ namespace UnrealBuildTool
 		/// <inheritdoc/>
 		public void CreateIntermediateTextFile(FileItem FileItem, IEnumerable<string> ContentLines, bool AllowAsync = true)
 		{
-			lock (Makefile.InternalDependencies)
-			{
-				Makefile.InternalDependencies.Add(FileItem);
-			}
+			AddInternalDependency(FileItem);
 
 			Task? task = FileHashCache.WriteFileIfChanged(FileItem, ContentLines, Logger, AllowAsync);
 			if (task != null)
@@ -185,6 +179,22 @@ namespace UnrealBuildTool
 			}
 		}
 
+		public void AddOutputFiles(IEnumerable<FileItem> OutputFiles)
+		{
+			lock (Makefile.OutputItems)
+			{
+				Makefile.OutputItems.AddRange(OutputFiles);
+			}
+		}
+
+		public void AddInternalDependency(FileItem Dependency)
+		{
+			lock (Makefile.InternalDependencies)
+			{
+				Makefile.InternalDependencies.Add(Dependency);
+			}
+		}
+
 		internal void WaitOnWriteTasks()
 		{
 			Task.WaitAll(WriteTasks.ToArray());
diff --git a/Engine/Source/Programs/UnrealBuildTool/System/Utils.cs b/Engine/Source/Programs/UnrealBuildTool/System/Utils.cs
--- a/Engine/Source/Programs/UnrealBuildTool/System/Utils.cs
+++ b/Engine/Source/Programs/UnrealBuildTool/System/Utils.cs
@@ -851,8 +851,6 @@ namespace UnrealBuildTool
 		/// <returns>True if the file was written successfully</returns>
 		public static bool WriteClass<T>(T Data, string FileName, string DefaultNameSpace, ILogger Logger)
 		{
-			bool bSuccess = true;
-			StreamWriter? XmlStream = null;
 			try
 			{
 				FileInfo Info = new FileInfo(FileName);
@@ -864,30 +862,24 @@ namespace UnrealBuildTool
 				// Make sure the output directory exists
 				Directory.CreateDirectory(Path.GetDirectoryName(FileName)!);
 
-				XmlSerializerNamespaces EmptyNameSpace = new XmlSerializerNamespaces();
+				XmlSerializerNamespaces EmptyNameSpace = new();
 				EmptyNameSpace.Add("", DefaultNameSpace);
 
-				XmlStream = new StreamWriter(FileName, false, Encoding.Unicode);
-				XmlSerializer Serialiser = new XmlSerializer(typeof(T));
+				using StreamWriter XmlStream = new(FileName, false, Encoding.Unicode);
+				XmlSerializer Serialiser = new(typeof(T));
 
 				// Add our callbacks for unknown nodes and attributes
 				Serialiser.UnknownNode += new XmlNodeEventHandler(UnknownNodeDelegate);
 				Serialiser.UnknownAttribute += new XmlAttributeEventHandler(UnknownAttributeDelegate);
 
 				Serialiser.Serialize(XmlStream, Data, EmptyNameSpace);
+				return true;
 			}
 			catch (Exception E)
 			{
 				Logger.LogInformation("{Message}", E.Message);
-				bSuccess = false;
-			}
-			finally
-			{
-				// Done with the file so close it
-				XmlStream?.Close();
+				return false;
 			}
-
-			return (bSuccess);
 		}
 
 		/// <summary>
