[UBT]
* Changed so verse bpvm is always false when running iwyu

diff --git a/Engine/Source/Programs/UnrealBuildTool/Configuration/TargetRules.cs b/Engine/Source/Programs/UnrealBuildTool/Configuration/TargetRules.cs
--- a/Engine/Source/Programs/UnrealBuildTool/Configuration/TargetRules.cs
+++ b/Engine/Source/Programs/UnrealBuildTool/Configuration/TargetRules.cs
@@ -3278,6 +3278,7 @@ namespace UnrealBuildTool
 				bUsePCHFiles = false; // No pch for iwyu
 				bUseVFS = false; // Do not use vfs because IWYUMode can't handle virtual paths (yet)
 				bUseUnityBuild = false; // No unity for iwyu
+				bUseVerseBPVM = false;
 				// bAlwaysUseUnityForGeneratedFiles = false; // This should be used but adds too many extra files atm.
 			}
 
