Fix iOS xcode project missing source files in client targeta

#rnx

diff --git a/Engine/Source/Programs/UnrealBuildTool/ProjectFiles/ProjectFileGenerator.cs b/Engine/Source/Programs/UnrealBuildTool/ProjectFiles/ProjectFileGenerator.cs
--- a/Engine/Source/Programs/UnrealBuildTool/ProjectFiles/ProjectFileGenerator.cs
+++ b/Engine/Source/Programs/UnrealBuildTool/ProjectFiles/ProjectFileGenerator.cs
@@ -2735,6 +2735,11 @@ namespace UnrealBuildTool
 						Debug.Assert(bAllowMultiModuleReference, "ProjectFileGenerator assert", $"Unexpected multi projects for module {CurModuleFile.GetFileName()}");
 						// e.g. QAGame module would be add to both QAGame.xcodeproj and QAGameEditor.xcodeproj, use the game project for module map, this way IOS workspace can have indexing
 						ProjectFile? EditorProjectFile = ProjectFiles.FirstOrDefault(x => x.ProjectTargets.Any(x => x.TargetRules!.Type == TargetType.Game));
+						if (EditorProjectFile == null)
+						{
+							// if no game, then try client
+							EditorProjectFile = ProjectFiles.FirstOrDefault(x => x.ProjectTargets.Any(x => x.TargetRules!.Type == TargetType.Client));
+						}
 						if (EditorProjectFile != null)
 						{
 							lock (ModuleToEditorProjectFileMap)
@@ -2940,8 +2945,8 @@ namespace UnrealBuildTool
 									}
 									else
 									{
-										// plugins and game modules add into both editor and game projects
-										if (Pair.Value.ProjectTargets.Any(x => (x.TargetRules!.Type == TargetType.Editor) || (x.TargetRules!.Type == TargetType.Game)))
+										// plugins and game/client modules add into both editor and game/client projects
+										if (Pair.Value.ProjectTargets.Any(x => (x.TargetRules!.Type == TargetType.Editor) || (x.TargetRules!.Type == TargetType.Client) || (x.TargetRules!.Type == TargetType.Game)))
 										{
 											FoundProjects.Add(Pair.Value);
 										}
