UBT: Add missing include paths from the Android NDK to the project file

#rnx

diff --git a/Engine/Source/Programs/UnrealBuildTool/Platform/Android/AndroidProjectGenerator.cs b/Engine/Source/Programs/UnrealBuildTool/Platform/Android/AndroidProjectGenerator.cs
--- a/Engine/Source/Programs/UnrealBuildTool/Platform/Android/AndroidProjectGenerator.cs
+++ b/Engine/Source/Programs/UnrealBuildTool/Platform/Android/AndroidProjectGenerator.cs
@@ -248,5 +248,42 @@ namespace UnrealBuildTool
 
 			return String.Empty;
 		}
+		
+		public override IList<string> GetSystemIncludePaths(UEBuildTarget InTarget)
+		{
+			UnrealArch Architecture = InTarget.Architectures.Architectures.Count >= 1 ? InTarget.Architectures.Architectures[0] : UnrealArch.X64;
+			string NDKArch = UEDeployAndroid.GetNDKArch(Architecture);
+			string NDKTargetTripletName = (NDKArch == "x86_64") ? "x86_64-linux-android" : "aarch64-linux-android";
+			string PlatformNDKHostName = UEDeployAndroid.GetPlatformNDKHostName();
+			string LLVMRoot = Path.Combine(Environment.ExpandEnvironmentVariables("%NDKROOT%"), "toolchains", "llvm", "prebuilt",
+				PlatformNDKHostName);
+			string SysRootInclude = Path.Combine(LLVMRoot, "sysroot", "usr", "include");
+
+			List<string> Result = new List<string>();
+			Result.Add(SysRootInclude);
+			Result.Add(Path.Combine(SysRootInclude, "c++", "v1"));
+			Result.Add(Path.Combine(SysRootInclude, NDKTargetTripletName));
+
+			try
+			{
+				string VersionFileName = Path.Combine(LLVMRoot, "AndroidVersion.txt");
+				using StreamReader VersionFile = new StreamReader(VersionFileName);
+				string LibsVersion = VersionFile.ReadLine()!;
+				int MajorVersionIndex = LibsVersion.IndexOf('.');
+				if (MajorVersionIndex != -1)
+				{
+					LibsVersion = LibsVersion.Substring(0, MajorVersionIndex);
+				}
+				
+				string ClangIncludePath = Path.Combine(LLVMRoot, "lib", "clang", LibsVersion, "include");
+				Result.Add(ClangIncludePath);
+			}
+			catch (Exception e)
+			{
+				Logger.LogWarning(e.Message);
+			}
+			
+			return Result;
+		}
 	}
 }
diff --git a/Engine/Source/Programs/UnrealBuildTool/Platform/Android/UEDeployAndroid.cs b/Engine/Source/Programs/UnrealBuildTool/Platform/Android/UEDeployAndroid.cs
--- a/Engine/Source/Programs/UnrealBuildTool/Platform/Android/UEDeployAndroid.cs
+++ b/Engine/Source/Programs/UnrealBuildTool/Platform/Android/UEDeployAndroid.cs
@@ -1271,7 +1271,7 @@ namespace UnrealBuildTool
 			Logger.LogInformation("==== Write {BuildIDFileName}  ====", BuildIDFileName);
 		}
 
-		private static string GetNDKArch(UnrealArch UnrealArch)
+		public static string GetNDKArch(UnrealArch UnrealArch)
 		{
 			if (UnrealArch == UnrealArch.Arm64)
 			{
@@ -1495,7 +1495,7 @@ namespace UnrealBuildTool
 			}
 		}
 
-		private static string GetPlatformNDKHostName()
+		public static string GetPlatformNDKHostName()
 		{
 			if (OperatingSystem.IsLinux())
 			{
