Add Wine check to EpicGames.Core, remove from UBT and use Linux code path for Horde Agent shutdown.

diff --git a/Engine/Source/Programs/AutomationTool/AutomationUtils/CommandEnvironment.cs b/Engine/Source/Programs/AutomationTool/AutomationUtils/CommandEnvironment.cs
--- a/Engine/Source/Programs/AutomationTool/AutomationUtils/CommandEnvironment.cs
+++ b/Engine/Source/Programs/AutomationTool/AutomationUtils/CommandEnvironment.cs
@@ -207,7 +207,7 @@ namespace AutomationTool
 			// Assume we have the capability co compile.
 			HasCapabilityToCompile = true;
 
-			if (BuildHostPlatform.Current.IsRunningOnWine())
+			if (RuntimePlatform.IsRunningOnWine())
 			{
 				// Just set an empty path as we currently compile .NET Framework/Core dependencies outside Wine
 				FrameworkMsbuildPath = "";
diff --git a/Engine/Source/Programs/AutomationTool/BuildGraph/TempStorage.cs b/Engine/Source/Programs/AutomationTool/BuildGraph/TempStorage.cs
--- a/Engine/Source/Programs/AutomationTool/BuildGraph/TempStorage.cs
+++ b/Engine/Source/Programs/AutomationTool/BuildGraph/TempStorage.cs
@@ -1350,7 +1350,7 @@ namespace AutomationTool
 		/// <returns>Output from the copy operation</returns>
 		static string CopyFile(FileInfo zipFile, DirectoryReference rootDir)
 		{
-			if (BuildHostPlatform.Current.IsRunningOnWine())
+			if (RuntimePlatform.IsRunningOnWine())
 			{
 				string sourceFile = zipFile.FullName;
 				string destFile = Path.Join(rootDir.FullName, zipFile.Name);
@@ -1376,7 +1376,7 @@ namespace AutomationTool
 		/// <exception cref="DirectoryNotFoundException">If source directory wasn't found</exception>
 		static string CopyDirectory(DirectoryReference sourceDir, DirectoryReference destinationDir)
 		{
-			if (BuildHostPlatform.Current.IsRunningOnWine())
+			if (RuntimePlatform.IsRunningOnWine())
 			{
 				return CopyDirectoryDotNet(sourceDir, destinationDir);
 			}
diff --git a/Engine/Source/Programs/Horde/HordeAgent/Utility/Shutdown.cs b/Engine/Source/Programs/Horde/HordeAgent/Utility/Shutdown.cs
--- a/Engine/Source/Programs/Horde/HordeAgent/Utility/Shutdown.cs
+++ b/Engine/Source/Programs/Horde/HordeAgent/Utility/Shutdown.cs
@@ -4,6 +4,7 @@ using System.Diagnostics;
 using System.Diagnostics.CodeAnalysis;
 using System.Runtime.InteropServices;
 using Microsoft.Extensions.Logging;
+using EpicGames.Core;
 
 namespace HordeAgent.Utility
 {
@@ -72,7 +73,7 @@ namespace HordeAgent.Utility
 		/// <returns></returns>
 		public static bool InitiateShutdown(bool restartAfterShutdown, ILogger logger)
 		{
-			if (RuntimeInformation.IsOSPlatform(OSPlatform.Windows))
+			if (RuntimeInformation.IsOSPlatform(OSPlatform.Windows) && !RuntimePlatform.IsRunningOnWine())
 			{
 				if (restartAfterShutdown)
 				{
@@ -118,7 +119,7 @@ namespace HordeAgent.Utility
 				}
 				return true;
 			}
-			else if (RuntimeInformation.IsOSPlatform(OSPlatform.OSX) || RuntimeInformation.IsOSPlatform(OSPlatform.Linux))
+			else if (RuntimeInformation.IsOSPlatform(OSPlatform.OSX) || RuntimeInformation.IsOSPlatform(OSPlatform.Linux) || RuntimePlatform.IsRunningOnWine())
 			{
 				string shutdownArgs;
 				if (restartAfterShutdown)
diff --git a/Engine/Source/Programs/Shared/EpicGames.Core/RuntimePlatform.cs b/Engine/Source/Programs/Shared/EpicGames.Core/RuntimePlatform.cs
--- a/Engine/Source/Programs/Shared/EpicGames.Core/RuntimePlatform.cs
+++ b/Engine/Source/Programs/Shared/EpicGames.Core/RuntimePlatform.cs
@@ -1,6 +1,7 @@
 // Copyright Epic Games, Inc. All Rights Reserved.
 
 using System;
+using System.Runtime.InteropServices;
 
 namespace EpicGames.Core
 {
@@ -57,5 +58,24 @@ namespace EpicGames.Core
 		/// The extension executables have on the current platform
 		/// </summary>
 		public static readonly string ExeExtension = OperatingSystem.IsWindows() ? ".exe" : "";
+
+		[DllImport("kernel32.dll", CharSet = CharSet.Auto)]
+		private static extern IntPtr GetModuleHandle(string lpModuleName);
+
+		[DllImport("kernel32.dll", CharSet = CharSet.Ansi)]
+		private static extern IntPtr GetProcAddress(IntPtr hModule, string procName);
+
+		/// <summary>
+		/// The current runtime platform is Windows but we're running on Wine
+		/// </summary>
+        public static bool IsRunningOnWine()
+		{
+			if (OperatingSystem.IsWindows())
+			{
+				IntPtr ntdllHandle = GetModuleHandle("ntdll.dll");
+				return ntdllHandle.ToInt64() != 0 && GetProcAddress(ntdllHandle, "wine_get_version") != IntPtr.Zero;
+			}
+			return false;
+		}
 	}
 }
diff --git a/Engine/Source/Programs/UnrealBuildTool/Platform/Windows/MicrosoftPlatformSDK.cs b/Engine/Source/Programs/UnrealBuildTool/Platform/Windows/MicrosoftPlatformSDK.cs
--- a/Engine/Source/Programs/UnrealBuildTool/Platform/Windows/MicrosoftPlatformSDK.cs
+++ b/Engine/Source/Programs/UnrealBuildTool/Platform/Windows/MicrosoftPlatformSDK.cs
@@ -652,7 +652,7 @@ namespace UnrealBuildTool
 				.FirstOrDefault();
 
 			// Allow x64 host toolchains on arm64 if no native arm64 toolchains are available, and not running via wine
-			if (!BuildHostPlatform.Current.IsRunningOnWine() && MicrosoftPlatformSDK.MSVCHostUnrealArch == UnrealArch.Arm64 && (ToolChain == null || ToolChain.Error != null))
+			if (!RuntimePlatform.IsRunningOnWine() && MicrosoftPlatformSDK.MSVCHostUnrealArch == UnrealArch.Arm64 && (ToolChain == null || ToolChain.Error != null))
 			{
 				ToolChain = Preference(ToolChains.Where(x => x.Architecture == Architecture && x.HostArchitecture != MicrosoftPlatformSDK.MSVCHostUnrealArch)
 					.OrderByDescending(x => x.Error == null))
diff --git a/Engine/Source/Programs/UnrealBuildTool/Platform/Windows/PVSToolChain.cs b/Engine/Source/Programs/UnrealBuildTool/Platform/Windows/PVSToolChain.cs
--- a/Engine/Source/Programs/UnrealBuildTool/Platform/Windows/PVSToolChain.cs
+++ b/Engine/Source/Programs/UnrealBuildTool/Platform/Windows/PVSToolChain.cs
@@ -605,7 +605,7 @@ namespace UnrealBuildTool
 				}
 			}
 
-			if (BuildHostPlatform.Current.IsRunningOnWine())
+			if (RuntimePlatform.IsRunningOnWine())
 			{
 				throw new BuildException("PVS-Studio is not supported with Wine.");
 			}
diff --git a/Engine/Source/Programs/UnrealBuildTool/System/BuildHostPlatform.cs b/Engine/Source/Programs/UnrealBuildTool/System/BuildHostPlatform.cs
--- a/Engine/Source/Programs/UnrealBuildTool/System/BuildHostPlatform.cs
+++ b/Engine/Source/Programs/UnrealBuildTool/System/BuildHostPlatform.cs
@@ -223,15 +223,6 @@ namespace UnrealBuildTool
 			return [.. Modules];
 		}
 
-		/// <summary>
-		/// Determines if the UBT process is running through WINE
-		/// </summary>
-		/// <returns>Sequence of project file formats</returns>
-		public virtual bool IsRunningOnWine()
-		{
-			return false;
-		}
-
 		/// <summary>
 		/// Determines the default project file formats for this platform
 		/// </summary>
@@ -249,18 +240,6 @@ namespace UnrealBuildTool
 
 		public override string BinarySuffix => ".exe";
 
-		[DllImport("kernel32.dll", CharSet = CharSet.Auto)]
-		private static extern IntPtr GetModuleHandle(string lpModuleName);
-
-		[DllImport("kernel32.dll", CharSet = CharSet.Auto)]
-		private static extern IntPtr GetProcAddress(IntPtr hModule, string procName);
-
-		public override bool IsRunningOnWine()
-		{
-			IntPtr NtdllHandle = GetModuleHandle("ntdll.dll");
-			return NtdllHandle.ToInt64() != 0 && GetProcAddress(NtdllHandle, "wine_get_version").ToInt64() != 0;
-		}
-
 		internal override IEnumerable<ProjectFileFormat> GetDefaultProjectFileFormats()
 		{
 			yield return ProjectFileFormat.VisualStudio;
