Replace all usages of BinaryFormatterUtils with System.Text.Json equivalents.

diff --git a/Engine/Source/Programs/Shared/EpicGames.Core/DirectoryReference.cs b/Engine/Source/Programs/Shared/EpicGames.Core/DirectoryReference.cs
--- a/Engine/Source/Programs/Shared/EpicGames.Core/DirectoryReference.cs
+++ b/Engine/Source/Programs/Shared/EpicGames.Core/DirectoryReference.cs
@@ -7,6 +7,8 @@ using System.Diagnostics.CodeAnalysis;
 using System.Globalization;
 using System.IO;
 using System.Linq;
+using System.Text.Json;
+using System.Text.Json.Serialization;
 
 namespace EpicGames.Core
 {
@@ -15,6 +17,7 @@ namespace EpicGames.Core
 	/// </summary>
 	[Serializable]
 	[TypeConverter(typeof(DirectoryReferenceTypeConverter))]
+	[JsonConverter(typeof(DirectoryReferenceJsonConverter))]
 	public class DirectoryReference : FileSystemReference, IEquatable<DirectoryReference>, IComparable<DirectoryReference>
 	{
 		/// <summary>
@@ -442,6 +445,34 @@ namespace EpicGames.Core
 		}
 	}
 
+	internal class DirectoryReferenceJsonConverter : JsonConverter<DirectoryReference>
+	{
+		public override DirectoryReference? Read(ref Utf8JsonReader reader, Type typeToConvert, JsonSerializerOptions options)
+		{
+			string? fullName = reader.GetString();
+			return fullName is null ? null : new(fullName);
+		}
+
+		public override DirectoryReference ReadAsPropertyName(ref Utf8JsonReader reader, Type typeToConvert, JsonSerializerOptions options)
+		{
+			if (reader.TokenType != JsonTokenType.PropertyName)
+			{
+				throw new JsonException("DirectoryReference keys must be represented as JSON strings");
+			}
+			return new(reader.GetString()!);
+		}
+
+		public override void Write(Utf8JsonWriter writer, DirectoryReference value, JsonSerializerOptions options)
+		{
+			writer.WriteStringValue(value.FullName);
+		}
+
+		public override void WriteAsPropertyName(Utf8JsonWriter writer, [DisallowNull] DirectoryReference value, JsonSerializerOptions options)
+		{
+			writer.WritePropertyName(value.FullName);
+		}
+	}
+
 	/// <summary>
 	/// Extension methods for passing DirectoryReference arguments
 	/// </summary>
diff --git a/Engine/Source/Programs/Shared/EpicGames.Core/FileReference.cs b/Engine/Source/Programs/Shared/EpicGames.Core/FileReference.cs
--- a/Engine/Source/Programs/Shared/EpicGames.Core/FileReference.cs
+++ b/Engine/Source/Programs/Shared/EpicGames.Core/FileReference.cs
@@ -8,6 +8,8 @@ using System.Globalization;
 using System.IO;
 using System.Linq;
 using System.Text;
+using System.Text.Json;
+using System.Text.Json.Serialization;
 using System.Threading;
 using System.Threading.Tasks;
 
@@ -18,6 +20,7 @@ namespace EpicGames.Core
 	/// </summary>
 	[Serializable]
 	[TypeConverter(typeof(FileReferenceTypeConverter))]
+	[JsonConverter(typeof(FileReferenceJsonConverter))]
 	public class FileReference : FileSystemReference, IEquatable<FileReference>, IComparable<FileReference>
 	{
 		/// <summary>
@@ -862,6 +865,34 @@ namespace EpicGames.Core
 		}
 	}
 
+	internal class FileReferenceJsonConverter : JsonConverter<FileReference>
+	{
+		public override FileReference? Read(ref Utf8JsonReader reader, Type typeToConvert, JsonSerializerOptions options)
+		{
+			string? fullName = reader.GetString();
+			return fullName is null ? null : new(fullName);
+		}
+
+		public override FileReference ReadAsPropertyName(ref Utf8JsonReader reader, Type typeToConvert, JsonSerializerOptions options)
+		{
+			if (reader.TokenType != JsonTokenType.PropertyName)
+			{
+				throw new JsonException("FileReference keys must be represented as JSON strings");
+			}
+			return new(reader.GetString()!);
+		}
+
+		public override void Write(Utf8JsonWriter writer, FileReference value, JsonSerializerOptions options)
+		{
+			writer.WriteStringValue(value.FullName);
+		}
+
+		public override void WriteAsPropertyName(Utf8JsonWriter writer, [DisallowNull] FileReference value, JsonSerializerOptions options)
+		{
+			writer.WritePropertyName(value.FullName);
+		}
+	}
+
 	/// <summary>
 	/// Extension methods for FileReference functionality
 	/// </summary>
diff --git a/Engine/Source/Programs/UnrealBuildTool/Configuration/ModuleRules.cs b/Engine/Source/Programs/UnrealBuildTool/Configuration/ModuleRules.cs
--- a/Engine/Source/Programs/UnrealBuildTool/Configuration/ModuleRules.cs
+++ b/Engine/Source/Programs/UnrealBuildTool/Configuration/ModuleRules.cs
@@ -7,6 +7,7 @@ using System.Diagnostics.CodeAnalysis;
 using System.IO;
 using System.Linq;
 using System.Reflection;
+using System.Text.Json.Serialization;
 using EpicGames.Core;
 using Microsoft.Extensions.Logging;
 using UnrealBuildBase;
@@ -317,40 +318,41 @@ namespace UnrealBuildTool
 			/// <summary>
 			/// The file that should be staged. Should use $(EngineDir) and $(ProjectDir) variables as a root, so that the target can be relocated to different machines.
 			/// </summary>
-			public string Path { get; init; }
+			public string Path { get; }
 
 			/// <summary>
 			/// The initial location for this file. It will be copied to Path at build time, ready for staging.
 			/// </summary>
-			public string? SourcePath { get; init; }
+			public string? SourcePath { get; }
 
 			/// <summary>
 			/// How to stage this file.
 			/// </summary>
-			public StagedFileType Type { get; init; }
+			public StagedFileType Type { get; }
 
 			/// <summary>
 			/// Constructor
 			/// </summary>
-			/// <param name="inPath">Path to the runtime dependency</param>
-			/// <param name="inType">How to stage the given path</param>
-			public RuntimeDependency(string inPath, StagedFileType inType = StagedFileType.NonUFS)
+			/// <param name="path">Path to the runtime dependency</param>
+			/// <param name="type">How to stage the given path</param>
+			public RuntimeDependency(string path, StagedFileType type = StagedFileType.NonUFS)
 			{
-				Path = inPath;
-				Type = inType;
+				Path = path;
+				Type = type;
 			}
 
 			/// <summary>
 			/// Constructor
 			/// </summary>
-			/// <param name="inPath">Path to the runtime dependency</param>
-			/// <param name="inSourcePath">Source path for the file in the working tree</param>
-			/// <param name="inType">How to stage the given path</param>
-			public RuntimeDependency(string inPath, string inSourcePath, StagedFileType inType = StagedFileType.NonUFS)
+			/// <param name="path">Path to the runtime dependency</param>
+			/// <param name="sourcePath">Source path for the file in the working tree</param>
+			/// <param name="type">How to stage the given path</param>
+			[JsonConstructor]
+			public RuntimeDependency(string path, string? sourcePath, StagedFileType type)
 			{
-				Path = inPath;
-				SourcePath = inSourcePath;
-				Type = inType;
+				Path = path;
+				SourcePath = sourcePath;
+				Type = type;
 			}
 		}
 
diff --git a/Engine/Source/Programs/UnrealBuildTool/Configuration/UEBuildPlatform.cs b/Engine/Source/Programs/UnrealBuildTool/Configuration/UEBuildPlatform.cs
--- a/Engine/Source/Programs/UnrealBuildTool/Configuration/UEBuildPlatform.cs
+++ b/Engine/Source/Programs/UnrealBuildTool/Configuration/UEBuildPlatform.cs
@@ -184,7 +184,7 @@ namespace UnrealBuildTool
 		protected UnrealArchitectureConfig(UnrealArchitectureMode Mode, IEnumerable<UnrealArch> SupportedArchitectures)
 		{
 			this.Mode = Mode;
-			AllSupportedArchitectures = new UnrealArchitectures(SupportedArchitectures);
+			AllSupportedArchitectures = new UnrealArchitectures(SupportedArchitectures.ToList());
 		}
 	}
 
diff --git a/Engine/Source/Programs/UnrealBuildTool/Configuration/UEBuildTarget.cs b/Engine/Source/Programs/UnrealBuildTool/Configuration/UEBuildTarget.cs
--- a/Engine/Source/Programs/UnrealBuildTool/Configuration/UEBuildTarget.cs
+++ b/Engine/Source/Programs/UnrealBuildTool/Configuration/UEBuildTarget.cs
@@ -1,5 +1,9 @@
 // Copyright Epic Games, Inc. All Rights Reserved.
 
+using EpicGames.Core;
+using EpicGames.UHT.Utils;
+using Microsoft.Extensions.Logging;
+using OpenTracing.Util;
 using System;
 using System.Collections.Concurrent;
 using System.Collections.Generic;
@@ -11,13 +15,10 @@ using System.Reflection;
 using System.Runtime.InteropServices;
 using System.Runtime.Serialization;
 using System.Text;
+using System.Text.Json;
 using System.Text.Json.Serialization;
 using System.Threading;
 using System.Threading.Tasks;
-using EpicGames.Core;
-using EpicGames.UHT.Utils;
-using Microsoft.Extensions.Logging;
-using OpenTracing.Util;
 using UnrealBuildBase;
 
 namespace UnrealBuildTool
@@ -26,6 +27,7 @@ namespace UnrealBuildTool
 	/// The platform we're building for
 	/// </summary>
 	[Serializable, TypeConverter(typeof(UnrealTargetPlatformTypeConverter))]
+	[JsonConverter(typeof(UnrealTargetPlatformJsonConverter))]
 	public partial struct UnrealTargetPlatform : ISerializable, IEquatable<UnrealTargetPlatform>
 	{
 		#region Private/boilerplate
@@ -331,6 +333,37 @@ namespace UnrealBuildTool
 		}
 	}
 
+	internal class UnrealTargetPlatformJsonConverter : JsonConverter<UnrealTargetPlatform>
+	{
+		public override UnrealTargetPlatform Read(ref Utf8JsonReader reader, Type typeToConvert, JsonSerializerOptions options)
+		{
+			if (reader.TokenType != JsonTokenType.String)
+			{
+				throw new JsonException("UnrealTargetPlatform values must be represented as JSON strings");
+			}
+			return UnrealTargetPlatform.Parse(reader.GetString()!);
+		}
+
+		public override UnrealTargetPlatform ReadAsPropertyName(ref Utf8JsonReader reader, Type typeToConvert, JsonSerializerOptions options)
+		{
+			if (reader.TokenType != JsonTokenType.PropertyName)
+			{
+				throw new JsonException("UnrealTargetPlatform keys must be represented as JSON strings");
+			}
+			return UnrealTargetPlatform.Parse(reader.GetString()!);
+		}
+
+		public override void Write(Utf8JsonWriter writer, UnrealTargetPlatform value, JsonSerializerOptions options)
+		{
+			writer.WriteStringValue(value.ToString());
+		}
+
+		public override void WriteAsPropertyName(Utf8JsonWriter writer, UnrealTargetPlatform value, JsonSerializerOptions options)
+		{
+			writer.WritePropertyName(value.ToString());
+		}
+	}
+
 	/// <summary>
 	/// Extension methods used for serializing UnrealTargetPlatform instances
 	/// </summary>
@@ -970,6 +1003,12 @@ namespace UnrealBuildTool
 		{
 		}
 
+		[JsonConstructor]
+		private UnrealArchitectures(List<UnrealArch> Architectures)
+		{
+			this.Architectures = [.. Architectures.Distinct()];
+		}
+
 		/// <summary>
 		/// Creates an UnrealArchitectures from a string (created via ToString() or similar)
 		/// </summary>
@@ -3199,7 +3238,7 @@ namespace UnrealBuildTool
 				if (VersionFile == null)
 				{
 					WriteMetadataTargetInfo TargetInfo = new WriteMetadataTargetInfo(ProjectFile, null, null, ReceiptFileName, Receipt, FileNameToModuleManifest, FileNameToLoadOrderManifest, Rules.bMergeModules);
-					FileReference TargetInfoFile = FileReference.Combine(ProjectIntermediateDirectory, "TargetMetadata.dat");
+					FileReference TargetInfoFile = FileReference.Combine(ProjectIntermediateDirectory, "TargetMetadata.json");
 
 					// Make a copy of this since OutputItems might be added byh post build actions
 					List<FileItem> TargetPrereqItems = new (Makefile.OutputItems);
@@ -3244,14 +3283,14 @@ namespace UnrealBuildTool
 					writeMetadataTasks.Add(Task.Run(() =>
 					{
 					WriteMetadataTargetInfo EngineInfo = new WriteMetadataTargetInfo(null, VersionFile, Receipt.Version, null, null, EngineFileToManifest, null, Rules.bMergeModules);
-						FileReference EngineInfoFile = FileReference.Combine(ProjectIntermediateDirectory, "EngineMetadata.dat");
+						FileReference EngineInfoFile = FileReference.Combine(ProjectIntermediateDirectory, "EngineMetadata.json");
 						string EngineInfoName = TargetName.Equals(VersionFile.GetFileNameWithoutExtension(), StringComparison.Ordinal) ? VersionFile.GetFileName() : $"{VersionFile.GetFileName()} ({TargetName})";
 						CreateWriteMetadataAction(MakefileBuilder, EngineInfoName, EngineInfoFile, EngineInfo, EnginePrereqItems);
 					}));
 					writeMetadataTasks.Add(Task.Run(() =>
 					{
 					WriteMetadataTargetInfo TargetInfo = new WriteMetadataTargetInfo(ProjectFile, VersionFile, null, ReceiptFileName, Receipt, TargetFileToManifest, FileNameToLoadOrderManifest, Rules.bMergeModules);
-						FileReference TargetInfoFile = FileReference.Combine(ProjectIntermediateDirectory, "TargetMetadata.dat");
+						FileReference TargetInfoFile = FileReference.Combine(ProjectIntermediateDirectory, "TargetMetadata.json");
 						CreateWriteMetadataAction(MakefileBuilder, ReceiptFileName.GetFileName(), TargetInfoFile, TargetInfo, TargetPrereqItems);
 					}));
 				}
@@ -3508,7 +3547,10 @@ namespace UnrealBuildTool
 				// is just to depend on the engine DLL timestamps. Instead, we add a makefile dependency on it, causing it to be
 				// regenerated if missing, and only add it as a prereq if we're updating the .target file.
 				FileItem InfoFileItem = FileItem.GetItemByFileReference(InfoFile);
-				BinaryFormatterUtils.SaveIfDifferent(InfoFile, Info);
+				byte[] InfoBytes = JsonSerializer.SerializeToUtf8Bytes(Info, WriteMetadataMode.SerializerOptions);
+				DirectoryReference.CreateDirectory(InfoFile.Directory);
+				FileReference.WriteAllBytesIfDifferent(InfoFile, InfoBytes);
+
 				InfoFileItem.ResetCachedInfo();
 				MakefileBuilder.AddInternalDependency(InfoFileItem);
 
diff --git a/Engine/Source/Programs/UnrealBuildTool/Modes/WriteMetadataMode.cs b/Engine/Source/Programs/UnrealBuildTool/Modes/WriteMetadataMode.cs
--- a/Engine/Source/Programs/UnrealBuildTool/Modes/WriteMetadataMode.cs
+++ b/Engine/Source/Programs/UnrealBuildTool/Modes/WriteMetadataMode.cs
@@ -4,6 +4,7 @@ using System;
 using System.Collections.Generic;
 using System.Diagnostics.CodeAnalysis;
 using System.IO;
+using System.Text.Json;
 using System.Threading.Tasks;
 using EpicGames.Core;
 using Microsoft.Extensions.Logging;
@@ -26,7 +27,7 @@ namespace UnrealBuildTool
 	record WriteMetadataTargetInfo(FileReference? ProjectFile, FileReference? VersionFile, BuildVersion? Version, FileReference? ReceiptFile, TargetReceipt? Receipt, Dictionary<FileReference, ModuleManifest> FileToManifest, Dictionary<FileReference, LoadOrderManifest>? FileToLoadOrderManifest, bool MergeModules)
 	{
 		/// <summary>
-		/// Map of load order manifest filenames to their locations on disk (generally, at most one load oder manifest is expected).
+		/// Map of load order manifest filenames to their locations on disk (generally, at most one load order manifest is expected).
 		/// </summary>
 		public Dictionary<FileReference, LoadOrderManifest> FileToLoadOrderManifest { get; init; } = FileToLoadOrderManifest ?? [];
 	}
@@ -43,6 +44,14 @@ namespace UnrealBuildTool
 		/// </summary>
 		public const int CurrentVersionNumber = 2;
 
+		/// <summary>
+		/// Serializer options for serializing <see cref="WriteMetadataTargetInfo"/> instances.
+		/// </summary>
+		internal static readonly JsonSerializerOptions SerializerOptions = new()
+		{
+			IncludeFields = true,
+		};
+
 		/// <summary>
 		/// Execute the command
 		/// </summary>
@@ -69,7 +78,14 @@ namespace UnrealBuildTool
 		private Task<int> ExecuteInternal(CommandLineArguments Arguments, ILogger Logger)
 		{
 			// Read the target info
-			WriteMetadataTargetInfo TargetInfo = BinaryFormatterUtils.Load<WriteMetadataTargetInfo>(Arguments.GetFileReference("-Input="));
+			FileReference FilePath = Arguments.GetFileReference("-Input=");
+			WriteMetadataTargetInfo TargetInfo;
+			using (FileStream FileStream = FileReference.Open(FilePath, FileMode.Open))
+			{
+				WriteMetadataTargetInfo? MaybeTargetInfo = JsonSerializer.Deserialize<WriteMetadataTargetInfo>(FileStream, SerializerOptions);
+				TargetInfo = MaybeTargetInfo ?? throw new BuildLogEventException("Target info file '{FilePath}' was invalid", FilePath);
+			}
+
 			bool bNoManifestChanges = Arguments.HasOption("-NoManifestChanges");
 			int VersionNumber = Arguments.GetInteger("-Version=");
 			Arguments.CheckAllArgumentsUsed();
diff --git a/Engine/Source/Programs/UnrealBuildTool/Platform/IOS/IOSPostBuildSyncMode.cs b/Engine/Source/Programs/UnrealBuildTool/Platform/IOS/IOSPostBuildSyncMode.cs
--- a/Engine/Source/Programs/UnrealBuildTool/Platform/IOS/IOSPostBuildSyncMode.cs
+++ b/Engine/Source/Programs/UnrealBuildTool/Platform/IOS/IOSPostBuildSyncMode.cs
@@ -1,37 +1,40 @@
 // Copyright Epic Games, Inc. All Rights Reserved.
 
-using System;
-using System.Collections.Generic;
-using System.Threading.Tasks;
 using EpicGames.Core;
 using Microsoft.Extensions.Logging;
+using System.Collections.Generic;
+using System.Diagnostics.CodeAnalysis;
+using System.IO;
+using System.Text.Json;
+using System.Text.Json.Serialization;
+using System.Threading.Tasks;
 using UnrealBuildBase;
 
 namespace UnrealBuildTool
 {
-	[Serializable]
 	class IOSPostBuildSyncTarget
 	{
 		public UnrealTargetPlatform Platform;
-		public UnrealArchitectures Architectures;
+		public required UnrealArchitectures Architectures;
 		public UnrealTargetConfiguration Configuration;
 		public FileReference? ProjectFile;
-		public string TargetName;
+		public required string TargetName;
 		public TargetType TargetType;
-		public FileReference OutputPath;
-		public List<string> UPLScripts;
-		public VersionNumber SdkVersion;
+		public required FileReference OutputPath;
+		public required List<string> UPLScripts;
+		public required VersionNumber SdkVersion;
 		public bool bCreateStubIPA;
 		public bool bSkipCrashlytics;
-		public DirectoryReference ProjectDirectory;
+		public required DirectoryReference ProjectDirectory;
 		public DirectoryReference? ProjectIntermediateDirectory;
 		public string? ImportProvision;
 		public string? ImportCertificate;
 		public string? ImportCertificatePassword;
-		public Dictionary<string, DirectoryReference> FrameworkNameToSourceDir;
+		public required Dictionary<string, DirectoryReference> FrameworkNameToSourceDir;
 		public bool bForDistribution = false;
 		public bool bBuildAsFramework = false;
 
+		[SetsRequiredMembers]
 		public IOSPostBuildSyncTarget(ReadOnlyTargetRules Target, FileReference OutputPath, DirectoryReference? ProjectIntermediateDirectory, List<string> UPLScripts, VersionNumber SdkVersion, Dictionary<string, DirectoryReference> FrameworkNameToSourceDir)
 		{
 			Platform = Target.Platform;
@@ -54,6 +57,14 @@ namespace UnrealBuildTool
 			bForDistribution = Target.IOSPlatform.bForDistribution;
 			bBuildAsFramework = Target.bShouldCompileAsDLL;
 		}
+
+		/// <summary>
+		/// Constructor required for JSON serialization.
+		/// </summary>
+		[JsonConstructor]
+		private IOSPostBuildSyncTarget()
+		{
+		}
 	}
 
 	[ToolMode("IOSPostBuildSync", ToolModeOptions.XmlConfig | ToolModeOptions.BuildPlatforms)]
@@ -76,7 +87,13 @@ namespace UnrealBuildTool
 			Arguments.CheckAllArgumentsUsed();
 
 			// Run the PostBuildSync command
-			IOSPostBuildSyncTarget Target = BinaryFormatterUtils.Load<IOSPostBuildSyncTarget>(InputFile!);
+			IOSPostBuildSyncTarget Target;
+			using (FileStream FileStream = FileReference.Open(InputFile!, FileMode.Open))
+			{
+				IOSPostBuildSyncTarget? MaybeTarget = JsonSerializer.Deserialize<IOSPostBuildSyncTarget>(FileStream, AppleToolChain.SerializerOptions);
+				Target = MaybeTarget ?? throw new BuildLogEventException("Post-build sync target file '{FilePath}' was invalid", InputFile!);
+			}
+
 			IOSToolChain.PostBuildSync(Target, Logger);
 
 			return Task.FromResult(0);
diff --git a/Engine/Source/Programs/UnrealBuildTool/Platform/IOS/IOSToolChain.cs b/Engine/Source/Programs/UnrealBuildTool/Platform/IOS/IOSToolChain.cs
--- a/Engine/Source/Programs/UnrealBuildTool/Platform/IOS/IOSToolChain.cs
+++ b/Engine/Source/Programs/UnrealBuildTool/Platform/IOS/IOSToolChain.cs
@@ -11,6 +11,7 @@ using System.IO;
 using System.Linq;
 using System.Security.Cryptography.X509Certificates;
 using System.Text;
+using System.Text.Json;
 using UnrealBuildBase;
 
 namespace UnrealBuildTool
@@ -1128,8 +1129,12 @@ namespace UnrealBuildTool
 				}
 
 				IOSPostBuildSyncTarget PostBuildSyncTarget = new IOSPostBuildSyncTarget(Target, Executable.Location, BinaryLinkEnvironment.IntermediateDirectory, UPLScripts, SdkVersion, FrameworkNameToSourceDir);
-				FileReference PostBuildSyncFile = FileReference.Combine(BinaryLinkEnvironment.IntermediateDirectory!, "PostBuildSync.dat");
-				BinaryFormatterUtils.Save(PostBuildSyncFile, PostBuildSyncTarget);
+				FileReference PostBuildSyncFile = FileReference.Combine(BinaryLinkEnvironment.IntermediateDirectory!, "PostBuildSync.json");
+				DirectoryReference.CreateDirectory(PostBuildSyncFile.Directory);
+				using (FileStream stream = new(PostBuildSyncFile.FullName, FileMode.Create, FileAccess.Write))
+				{
+					JsonSerializer.Serialize(stream, PostBuildSyncTarget, SerializerOptions);
+				}
 
 				string PostBuildSyncArguments = String.Format("-Input=\"{0}\" -XmlConfigCache=\"{1}\" -remoteini=\"{2}\" -legacyxcode", PostBuildSyncFile, XmlConfig.CacheFile, UnrealBuildTool.GetRemoteIniPath());
 
diff --git a/Engine/Source/Programs/UnrealBuildTool/System/HotReload.cs b/Engine/Source/Programs/UnrealBuildTool/System/HotReload.cs
--- a/Engine/Source/Programs/UnrealBuildTool/System/HotReload.cs
+++ b/Engine/Source/Programs/UnrealBuildTool/System/HotReload.cs
@@ -41,17 +41,17 @@ namespace UnrealBuildTool
 		/// <summary>
 		/// Suffix to use for the next hot reload invocation
 		/// </summary>
-		public int NextSuffix = 1;
+		public int NextSuffix { get; set; } = 1;
 
 		/// <summary>
 		/// Map from original filename in the action graph to hot reload file
 		/// </summary>
-		public Dictionary<FileReference, FileReference> OriginalFileToHotReloadFile = new Dictionary<FileReference, FileReference>();
+		public Dictionary<FileReference, FileReference> OriginalFileToHotReloadFile { get; set; } = [];
 
 		/// <summary>
 		/// Set of all temporary files created for hot reload
 		/// </summary>
-		public HashSet<FileReference> TemporaryFiles = new HashSet<FileReference>();
+		public HashSet<FileReference> TemporaryFiles { get; set; } = [];
 
 		/// <summary>
 		/// Public constructor
@@ -60,21 +60,6 @@ namespace UnrealBuildTool
 		{
 		}
 
-		/// <summary>
-		/// Private constructor
-		/// </summary>
-		/// <param name="Reader"></param>
-		HotReloadState(BinaryArchiveReader Reader)
-		{
-			int ArchiveVersion = Reader.ReadInt();
-			if (CurrentVersion == ArchiveVersion)
-			{
-				NextSuffix = Reader.ReadInt();
-				OriginalFileToHotReloadFile = Reader.ReadDictionary(Reader.ReadFileReference, Reader.ReadFileReference) ?? [];
-				TemporaryFiles = Reader.ReadHashSet(Reader.ReadFileReference) ?? [];
-			}
-		}
-
 		/// <summary>
 		/// Adds all the actions into the hot reload state, so we can restore the action graph on next iteration
 		/// </summary>
@@ -149,7 +134,7 @@ namespace UnrealBuildTool
 		public static FileReference GetLocation(FileReference? ProjectFile, string TargetName, UnrealTargetPlatform Platform, UnrealTargetConfiguration Configuration, UnrealArchitectures Architectures)
 		{
 			DirectoryReference BaseDir = DirectoryReference.FromFile(ProjectFile) ?? Unreal.EngineDirectory;
-			return FileReference.Combine(BaseDir, UEBuildTarget.GetPlatformIntermediateFolder(Platform, Architectures, false), TargetName, Configuration.ToString(), "HotReloadState.bin");
+			return FileReference.Combine(BaseDir, UEBuildTarget.GetPlatformIntermediateFolder(Platform, Architectures, false), TargetName, Configuration.ToString(), "HotReloadState.json");
 		}
 
 		/// <summary>
@@ -159,8 +144,9 @@ namespace UnrealBuildTool
 		/// <returns>New hot reload state instance</returns>
 		public static HotReloadState Load(FileReference Location)
 		{
-			using BinaryArchiveReader Reader = new(Location);
-			return new(Reader);
+			using FileStream FileStream = FileReference.Open(Location, FileMode.Open);
+			HotReloadState? MaybeState = JsonSerializer.Deserialize<HotReloadState>(FileStream);
+			return MaybeState ?? throw new BuildLogEventException("Hot reload state file '{FilePath}' was invalid", Location);
 		}
 
 		/// <summary>
@@ -170,11 +156,8 @@ namespace UnrealBuildTool
 		public void Save(FileReference Location)
 		{
 			DirectoryReference.CreateDirectory(Location.Directory);
-			using BinaryArchiveWriter Writer = new(Location);
-			Writer.WriteInt(CurrentVersion);
-			Writer.WriteInt(NextSuffix);
-			Writer.WriteDictionary(OriginalFileToHotReloadFile, Writer.WriteFileReference, Writer.WriteFileReference);
-			Writer.WriteHashSet(TemporaryFiles, Writer.WriteFileReference);
+			using FileStream stream = new(Location.FullName, FileMode.Create, FileAccess.Write);
+			JsonSerializer.Serialize(stream, this);
 		}
 	}
 
@@ -1294,7 +1277,12 @@ namespace UnrealBuildTool
 					{
 						throw new Exception(String.Format("Unable to find metadata file to patch action graph ({0})", TargetInfoFile));
 					}
-					WriteMetadataTargetInfo TargetInfo = BinaryFormatterUtils.Load<WriteMetadataTargetInfo>(TargetInfoFile);
+					WriteMetadataTargetInfo TargetInfo;
+					using (FileStream FileStream = FileReference.Open(TargetInfoFile, FileMode.Open))
+					{
+						WriteMetadataTargetInfo? MaybeTargetInfo = JsonSerializer.Deserialize<WriteMetadataTargetInfo>(FileStream, WriteMetadataMode.SerializerOptions);
+						TargetInfo = MaybeTargetInfo ?? throw new BuildLogEventException("Target info file '{FilePath}' was invalid", TargetInfoFile);
+					}
 
 					// Update the module names
 					bool bHasUpdatedModuleNames = false;
@@ -1317,8 +1305,10 @@ namespace UnrealBuildTool
 					// Write the hot-reload metadata file and update the argument list
 					if (bHasUpdatedModuleNames)
 					{
-						FileReference HotReloadTargetInfoFile = FileReference.Combine(TargetInfoFile.Directory, "Metadata-HotReload.dat");
-						BinaryFormatterUtils.SaveIfDifferent(HotReloadTargetInfoFile, TargetInfo);
+						FileReference HotReloadTargetInfoFile = FileReference.Combine(TargetInfoFile.Directory, "Metadata-HotReload.json");
+						byte[] InfoBytes = JsonSerializer.SerializeToUtf8Bytes(TargetInfo, WriteMetadataMode.SerializerOptions);
+						DirectoryReference.CreateDirectory(HotReloadTargetInfoFile.Directory);
+						FileReference.WriteAllBytesIfDifferent(HotReloadTargetInfoFile, InfoBytes);
 
 						Action NewAction = new Action(Action.Inner);
 
diff --git a/Engine/Source/Programs/UnrealBuildTool/System/LoadOrderManifest.cs b/Engine/Source/Programs/UnrealBuildTool/System/LoadOrderManifest.cs
--- a/Engine/Source/Programs/UnrealBuildTool/System/LoadOrderManifest.cs
+++ b/Engine/Source/Programs/UnrealBuildTool/System/LoadOrderManifest.cs
@@ -24,11 +24,7 @@ namespace UnrealBuildTool
 		/// List of all project's dynamic modules in the order in which they should be loaded.
 		/// All paths are relative to the main binary's directory.
 		/// </summary>
-		public List<string> Libraries = new List<string>();
-
-		public LoadOrderManifest()
-		{
-		}
+		public readonly List<string> Libraries = [];
 
 		/// <summary>
 		/// Gets a filename for a load order manifest for a build with the given platform/configuration.
diff --git a/Engine/Source/Programs/UnrealBuildTool/System/ModuleManifest.cs b/Engine/Source/Programs/UnrealBuildTool/System/ModuleManifest.cs
--- a/Engine/Source/Programs/UnrealBuildTool/System/ModuleManifest.cs
+++ b/Engine/Source/Programs/UnrealBuildTool/System/ModuleManifest.cs
@@ -5,6 +5,7 @@ using System.Collections.Generic;
 using System.Diagnostics.CodeAnalysis;
 using System.IO;
 using System.Linq;
+using System.Text.Json.Serialization;
 using EpicGames.Core;
 
 namespace UnrealBuildTool
@@ -22,7 +23,7 @@ namespace UnrealBuildTool
 		/// <summary>
 		/// Unique build id for the modules described in the manifest.
 		/// </summary>
-		public string BuildId;
+		public required string BuildId;
 
 		/// <summary>
 		/// Map of module name to file name.
@@ -38,11 +39,17 @@ namespace UnrealBuildTool
 		/// Constructs the module map with the given changelist
 		/// </summary>
 		/// <param name="InBuildId">The unique build id</param>
+		[SetsRequiredMembers]
 		public ModuleManifest(string InBuildId)
 		{
 			BuildId = InBuildId;
 		}
 
+		[JsonConstructor]
+		private ModuleManifest()
+		{
+		}
+
 		/// <summary>
 		/// Merge another manifest into this one
 		/// </summary>
diff --git a/Engine/Source/Programs/UnrealBuildTool/System/TargetReceipt.cs b/Engine/Source/Programs/UnrealBuildTool/System/TargetReceipt.cs
--- a/Engine/Source/Programs/UnrealBuildTool/System/TargetReceipt.cs
+++ b/Engine/Source/Programs/UnrealBuildTool/System/TargetReceipt.cs
@@ -6,6 +6,7 @@ using System.Diagnostics.CodeAnalysis;
 using System.IO;
 using System.Linq;
 using System.Text.Json;
+using System.Text.Json.Serialization;
 using System.Xml.Serialization;
 using EpicGames.Core;
 using UnrealBuildBase;
@@ -80,12 +81,13 @@ namespace UnrealBuildTool
 		/// <summary>
 		/// Constructor.
 		/// </summary>
-		/// <param name="InPath">Path to the build product</param>
-		/// <param name="InType">Type of the build product</param>
-		public BuildProduct(FileReference InPath, BuildProductType InType)
+		/// <param name="Path">Path to the build product</param>
+		/// <param name="Type">Type of the build product</param>
+		[JsonConstructor]
+		public BuildProduct(FileReference Path, BuildProductType Type)
 		{
-			Path = InPath;
-			Type = InType;
+			this.Path = Path;
+			this.Type = Type;
 		}
 
 		/// <summary>
@@ -161,22 +163,23 @@ namespace UnrealBuildTool
 		/// <summary>
 		/// Constructor
 		/// </summary>
-		/// <param name="InPath">Path to the runtime dependency</param>
-		/// <param name="InType">How to stage the given path</param>
-		public RuntimeDependency(FileReference InPath, StagedFileType InType = StagedFileType.NonUFS)
+		/// <param name="Path">Path to the runtime dependency</param>
+		/// <param name="Type">How to stage the given path</param>
+		[JsonConstructor]
+		public RuntimeDependency(FileReference Path, StagedFileType Type = StagedFileType.NonUFS)
 		{
-			Path = InPath;
-			Type = InType;
+			this.Path = Path;
+			this.Type = Type;
 		}
 
 		/// <summary>
 		/// Copy constructor
 		/// </summary>
-		/// <param name="InOther">Runtime dependency to copy settings from</param>
-		public RuntimeDependency(RuntimeDependency InOther)
+		/// <param name="Other">Runtime dependency to copy settings from</param>
+		public RuntimeDependency(RuntimeDependency Other)
 		{
-			Path = InOther.Path;
-			Type = InOther.Type;
+			Path = Other.Path;
+			Type = Other.Type;
 		}
 
 		/// <summary>
@@ -242,12 +245,12 @@ namespace UnrealBuildTool
 		/// <summary>
 		/// Construct a property with the given name and value
 		/// </summary>
-		/// <param name="InName">Name of the property</param>
-		/// <param name="InValue">Value of the property</param>
-		public ReceiptProperty(string InName, string InValue)
+		/// <param name="Name">Name of the property</param>
+		/// <param name="Value">Value of the property</param>
+		public ReceiptProperty(string Name, string Value)
 		{
-			Name = InName;
-			Value = InValue;
+			this.Name = Name;
+			this.Value = Value;
 		}
 	}
 
@@ -270,7 +273,7 @@ namespace UnrealBuildTool
 		/// <summary>
 		/// The name of this target
 		/// </summary>
-		public string TargetName;
+		public required string TargetName;
 
 		/// <summary>
 		/// The build settings version used for this target.
@@ -290,7 +293,7 @@ namespace UnrealBuildTool
 		/// <summary>
 		/// Which platform the target is compiled for
 		/// </summary>
-		public UnrealArchitectures Architectures;
+		public required UnrealArchitectures Architectures;
 
 		/// <summary>
 		/// Which configuration this target is compiled in
@@ -310,7 +313,7 @@ namespace UnrealBuildTool
 		/// <summary>
 		/// Version information for this target.
 		/// </summary>
-		public BuildVersion Version;
+		public required BuildVersion Version;
 
 		/// <summary>
 		/// The exectuable to launch for this target
@@ -360,6 +363,7 @@ namespace UnrealBuildTool
 		/// <param name="InIsTestTarget">Whether it's a target for low level tests</param>
 		/// <param name="InBuildSettingsVersion">Build settings version of the target.</param>
 		/// <param name="InTargetBuildEnvironment">Build environment used to produce the target.</param>
+		[SetsRequiredMembers]
 		public TargetReceipt(FileReference? InProjectFile, string InTargetName, TargetType InTargetType, UnrealTargetPlatform InPlatform, UnrealTargetConfiguration InConfiguration, BuildVersion InVersion, UnrealArchitectures InArchitectures, bool InIsTestTarget, BuildSettingsVersion InBuildSettingsVersion, TargetBuildEnvironment InTargetBuildEnvironment)
 		{
 			ProjectFile = InProjectFile;
@@ -375,6 +379,14 @@ namespace UnrealBuildTool
 			IsTestTarget = InIsTestTarget;
 		}
 
+		/// <summary>
+		/// Constructor required for JSON serialization.
+		/// </summary>
+		[JsonConstructor]
+		private TargetReceipt()
+		{
+		}
+
 		/// <summary>
 		/// Adds a build product to the receipt. Does not check whether it already exists.
 		/// </summary>
diff --git a/Engine/Source/Programs/UnrealBuildTool/ToolChain/AppleToolChain.cs b/Engine/Source/Programs/UnrealBuildTool/ToolChain/AppleToolChain.cs
--- a/Engine/Source/Programs/UnrealBuildTool/ToolChain/AppleToolChain.cs
+++ b/Engine/Source/Programs/UnrealBuildTool/ToolChain/AppleToolChain.cs
@@ -1,16 +1,19 @@
 // Copyright Epic Games, Inc. All Rights Reserved.
 
+using EpicGames.Core;
+using Microsoft.Extensions.Logging;
 using System;
 using System.Collections.Generic;
 using System.Diagnostics;
+using System.Diagnostics.CodeAnalysis;
 using System.IO;
 using System.Linq;
 using System.Security.Cryptography.X509Certificates;
 using System.Text;
+using System.Text.Json;
+using System.Text.Json.Serialization;
 using System.Text.RegularExpressions;
 using System.Threading.Tasks;
-using EpicGames.Core;
-using Microsoft.Extensions.Logging;
 using UnrealBuildBase;
 
 namespace UnrealBuildTool
@@ -252,6 +255,14 @@ namespace UnrealBuildTool
 		readonly bool bUseSwiftUIMain;
 		readonly bool bCreateSwiftBridgingHeader;
 
+		/// <summary>
+		/// Serializer options for serializing <see cref="IOSPostBuildSyncTarget"/> and <see cref="ApplePostBuildSyncTarget"/> instances.
+		/// </summary>
+		internal static readonly JsonSerializerOptions SerializerOptions = new()
+		{
+			IncludeFields = true,
+		};
+
 		protected bool bUseModernXcode => AppleExports.UseModernXcode(ProjectFile);
 
 		public AppleToolChain(ReadOnlyTargetRules? Target, Func<AppleToolChainSettings> InCreateSettings, ClangToolChainOptions InOptions, ILogger InLogger) : base(InOptions, InLogger)
@@ -1122,23 +1133,23 @@ namespace UnrealBuildTool
 		#endregion
 	};
 
-	[Serializable]
 	class ApplePostBuildSyncTarget
 	{
 		public FileReference? ProjectFile;
 		public UnrealTargetPlatform Platform;
-		public UnrealArchitectures Architectures;
+		public required UnrealArchitectures Architectures;
 		public UnrealTargetConfiguration Configuration;
-		public string TargetName;
+		public required string TargetName;
 
 		// For iOS/TVOS
 		public bool bCreateStubIPA;
 		public string? RemoteImportProvision;
 		public string? RemoteImportCertificate;
 		public string? RemoteImportCertificatePassword;
-		public DirectoryReference ProjectIntermediateDirectory;
-		public FileReference StubOutputPath;
+		public required DirectoryReference ProjectIntermediateDirectory;
+		public required FileReference StubOutputPath;
 
+		[SetsRequiredMembers]
 		public ApplePostBuildSyncTarget(ReadOnlyTargetRules Target, FileItem Executable, DirectoryReference IntermediateDir)
 		{
 			Platform = Target.Platform;
@@ -1155,6 +1166,14 @@ namespace UnrealBuildTool
 
 			StubOutputPath = Executable.Location;
 		}
+
+		/// <summary>
+		/// Constructor required for JSON serialization.
+		/// </summary>
+		[JsonConstructor]
+		private ApplePostBuildSyncTarget()
+		{
+		}
 	}
 
 	[ToolMode("ApplePostBuildSync", ToolModeOptions.XmlConfig | ToolModeOptions.BuildPlatforms)]
@@ -1180,7 +1199,13 @@ namespace UnrealBuildTool
 			Arguments.CheckAllArgumentsUsed();
 
 			// Run the PostBuildSync command
-			ApplePostBuildSyncTarget Target = BinaryFormatterUtils.Load<ApplePostBuildSyncTarget>(InputFile!);
+			ApplePostBuildSyncTarget Target;
+			using (FileStream FileStream = FileReference.Open(InputFile!, FileMode.Open))
+			{
+				ApplePostBuildSyncTarget? MaybeTarget = JsonSerializer.Deserialize<ApplePostBuildSyncTarget>(FileStream, AppleToolChain.SerializerOptions);
+				Target = MaybeTarget ?? throw new BuildLogEventException("Post-build sync target file '{FilePath}' was invalid", InputFile!);
+			}
+
 			int ExitCode = PostBuildSync(Target, Logger);
 
 			return Task.FromResult(ExitCode);
@@ -1341,8 +1366,13 @@ namespace UnrealBuildTool
 		public static Action CreatePostBuildSyncAction(ReadOnlyTargetRules Target, FileItem Executable, DirectoryReference IntermediateDir, IActionGraphBuilder Graph, bool bNoEntitlements = false)
 		{
 			ApplePostBuildSyncTarget PostBuildSync = new(Target, Executable, IntermediateDir);
-			FileReference PostBuildSyncFile = FileReference.Combine(IntermediateDir!, "PostBuildSync.dat");
-			BinaryFormatterUtils.Save(PostBuildSyncFile, PostBuildSync);			
+			FileReference PostBuildSyncFile = FileReference.Combine(IntermediateDir!, "PostBuildSync.json");
+			DirectoryReference.CreateDirectory(PostBuildSyncFile.Directory);
+			using (FileStream stream = new(PostBuildSyncFile.FullName, FileMode.Create, FileAccess.Write))
+			{
+				JsonSerializer.Serialize(stream, PostBuildSync, AppleToolChain.SerializerOptions);
+			}
+
 			string PostBuildSyncArguments = String.Format("-modernxcode -Input=\"{0}\" -XmlConfigCache=\"{1}\" -remoteini=\"{2}\" {3}", PostBuildSyncFile, XmlConfig.CacheFile, UnrealBuildTool.GetRemoteIniPath(), bNoEntitlements ? "-noEntitlements" : "");
 			Action PostBuildSyncAction = Graph.CreateRecursiveAction<ApplePostBuildSyncMode>(ActionType.CreateAppBundle, PostBuildSyncArguments);
 			PostBuildSyncAction.PrerequisiteItems.Add(Executable);
