Change project discovery by target name in UBT to more closely match the target discovery in the rules compiler by searching recursively under candidate directories.

diff --git a/Engine/Source/Programs/Shared/EpicGames.Build/System/NativeProjectsBase.cs b/Engine/Source/Programs/Shared/EpicGames.Build/System/NativeProjectsBase.cs
--- a/Engine/Source/Programs/Shared/EpicGames.Build/System/NativeProjectsBase.cs
+++ b/Engine/Source/Programs/Shared/EpicGames.Build/System/NativeProjectsBase.cs
@@ -179,36 +179,6 @@ namespace UnrealBuildBase
 				FindRelativeFileReference(Path.Combine(ProjectName, $"{ProjectName}.uproject"), Logger);
 		}
 
-		/// <summary>
-		/// Finds all target files under a given folder, and add them to the target name to project file map
-		/// </summary>
-		/// <param name="Directory">Directory to search</param>
-		/// <param name="TargetNameToProjectFile">Map from target name to project file</param>
-		/// <param name="ProjectFile">The project file for this directory</param>
-		protected static void FindTargetFiles(DirectoryReference Directory, Dictionary<string, FileReference> TargetNameToProjectFile, FileReference ProjectFile)
-		{
-			// Search for all target files within this directory
-			bool bSearchSubFolders = true;
-			foreach (FileReference File in DirectoryLookupCache.EnumerateFiles(Directory))
-			{
-				if (File.HasExtension(".target.cs"))
-				{
-					string TargetName = Path.GetFileNameWithoutExtension(File.GetFileNameWithoutExtension());
-					TargetNameToProjectFile[TargetName] = ProjectFile;
-					bSearchSubFolders = false;
-				}
-			}
-
-			// If we didn't find anything, recurse through the subfolders
-			if(bSearchSubFolders)
-			{
-				foreach(DirectoryReference SubDirectory in DirectoryLookupCache.EnumerateDirectories(Directory))
-				{
-					FindTargetFiles(SubDirectory, TargetNameToProjectFile, ProjectFile);
-				}
-			}
-		}
-
 		/// <summary>
 		/// Checks if a given project is a native project
 		/// </summary>
diff --git a/Engine/Source/Programs/UnrealBuildTool/System/NativeProjects.cs b/Engine/Source/Programs/UnrealBuildTool/System/NativeProjects.cs
--- a/Engine/Source/Programs/UnrealBuildTool/System/NativeProjects.cs
+++ b/Engine/Source/Programs/UnrealBuildTool/System/NativeProjects.cs
@@ -46,22 +46,21 @@ namespace UnrealBuildTool
 					if (CachedTargetNameToProjectFile == null)
 					{
 						using ITimelineEvent _ = Timeline.ScopeEvent("Find target files");
-						Dictionary<string, FileReference> TargetNameToProjectFile = new Dictionary<string, FileReference>();
+						Dictionary<string, FileReference> TargetNameToProjectFile = [];
+
+						IEnumerable<DirectoryReference> AllPotentialTargetDirectories = EnumerateProjectFiles(Logger)
+							.SelectMany(ProjectFile => Unreal.GetExtensionDirs(ProjectFile.Directory))
+							.SelectMany(ExpandSource);
+						Rules.Populate(AllPotentialTargetDirectories);
+
 						foreach (FileReference ProjectFile in EnumerateProjectFiles(Logger))
 						{
-							foreach (DirectoryReference ExtensionDir in Unreal.GetExtensionDirs(ProjectFile.Directory))
+							IEnumerable<DirectoryReference> ProjectTargetDirectories = Unreal.GetExtensionDirs(ProjectFile.Directory).SelectMany(ExpandSource);
+							List<FileReference> TargetFiles = Rules.Get(ProjectTargetDirectories, Rules.RulesFileType.Target);
+							foreach (FileReference TargetFile in TargetFiles)
 							{
-								DirectoryReference SourceDirectory = DirectoryReference.Combine(ExtensionDir, "Source");
-								if (DirectoryLookupCache.DirectoryExists(SourceDirectory))
-								{
-									FindTargetFiles(SourceDirectory, TargetNameToProjectFile, ProjectFile);
-								}
-
-								DirectoryReference IntermediateSourceDirectory = DirectoryReference.Combine(ExtensionDir, "Intermediate", "Source");
-								if (DirectoryLookupCache.DirectoryExists(IntermediateSourceDirectory))
-								{
-									FindTargetFiles(IntermediateSourceDirectory, TargetNameToProjectFile, ProjectFile);
-								}
+								string TargetName = TargetFile.GetFileNameWithoutAnyExtensions();
+								TargetNameToProjectFile[TargetName] = ProjectFile;
 							}
 
 							// Programs are a special case where the .uproject files are separated from the main project source code- in this case,
@@ -232,6 +231,12 @@ namespace UnrealBuildTool
 			return false;
 		}
 
+		private static IEnumerable<DirectoryReference> ExpandSource(DirectoryReference Item)
+		{
+			yield return DirectoryReference.Combine(Item, "Source");
+			yield return DirectoryReference.Combine(Item, "Intermediate", "Source");
+		}
+
 		private static bool RequiresTempTarget(FileReference UProjectFile, List<UnrealTargetConfiguration> Configurations, [NotNullWhen(true)] out string? Reason, ILogger Logger)
 		{
 			// no reason by default
