Fix an issue in UBT's IWYU mode where generated header includes were not being tracked.

diff --git a/Engine/Source/Programs/UnrealBuildTool/Modes/IWYUMode.cs b/Engine/Source/Programs/UnrealBuildTool/Modes/IWYUMode.cs
--- a/Engine/Source/Programs/UnrealBuildTool/Modes/IWYUMode.cs
+++ b/Engine/Source/Programs/UnrealBuildTool/Modes/IWYUMode.cs
@@ -816,7 +816,9 @@ namespace UnrealBuildTool
 					List<IWYUInfo> generatedHeaders = entry.generatedHeaders;
 					var exists = new HashSet<IWYUIncludeEntry>(mainInfo.Includes, new IWYUIncludeEntryFullPathComparer());
 					foreach (var generatedHeader in generatedHeaders)
-						exists.Union(generatedHeader.IncludesSeenInFile);
+					{
+						exists.UnionWith(generatedHeader.IncludesSeenInFile);
+					}
 
 					foreach (var generatedHeader in generatedHeaders)
 					{
