Fix issue resulting from making abstract type constructors protected.
Minor change to wording of summary message.

#rnx

diff --git a/Engine/Source/Programs/Shared/EpicGames.UHT/Utils/UhtTables.cs b/Engine/Source/Programs/Shared/EpicGames.UHT/Utils/UhtTables.cs
--- a/Engine/Source/Programs/Shared/EpicGames.UHT/Utils/UhtTables.cs
+++ b/Engine/Source/Programs/Shared/EpicGames.UHT/Utils/UhtTables.cs
@@ -623,12 +623,12 @@ namespace EpicGames.UHT.Utils
 			bool hasBadType = false;
 			foreach (Type type in _knownTypes)
 			{
-				if (type != typeof(UhtPackage))
+				if (type != typeof(UhtPackage) && !type.IsAbstract)
 				{
 					ConstructorInfo? constructor = type.GetConstructor([typeof(UhtInputCacheReader), typeof(UhtType)]);
 					if (constructor == null)
 					{
-						logger.LogError("The UhtType '{TypeName}' must defined the constructor '{TypeName}(UhtInputCacheReader reader, UhtType outer)'", type.Name, type.Name);
+						logger.LogError("The UhtType '{TypeName}' must defined the constructor '{ConstructorName}(UhtInputCacheReader reader, UhtType outer)'", type.Name, type.Name);
 						hasBadType = true;
 					}
 					MethodInfo? method = type.GetMethod("Write", [typeof(UhtInputCacheWriter)]);
diff --git a/Engine/Source/Programs/UnrealBuildTool/Modes/UnrealHeaderToolMode.cs b/Engine/Source/Programs/UnrealBuildTool/Modes/UnrealHeaderToolMode.cs
--- a/Engine/Source/Programs/UnrealBuildTool/Modes/UnrealHeaderToolMode.cs
+++ b/Engine/Source/Programs/UnrealBuildTool/Modes/UnrealHeaderToolMode.cs
@@ -864,7 +864,7 @@ namespace UnrealBuildTool.Modes
 				if (!Session.HasErrors || Session.TotalFilesWritten != 0)
 				{
 					string ActualTargetName = Session.Manifest != null ? Session.Manifest.TargetName : "UnknownTarget";
-					Logger.Log(level, "UHT processed {TargetName} in {Time} seconds ({TotalFilesWritten} files written, {TotalFilesRead} files read, {TotalFilesCacheRead} files cache read)",
+					Logger.Log(level, "UHT processed {TargetName} in {Time} seconds ({TotalFilesWritten} generated files written, {TotalFilesRead} H files read, {TotalFilesCacheRead} H files read from cache)",
 						ActualTargetName, stopwatch.Elapsed.TotalSeconds, Session.TotalFilesWritten, Session.TotalFilesRead, Session.TotalFilesCacheRead);
 				}
 				if (Session.TotalFilesCacheReadErrors != 0 || Session.TotalFilesCacheWriteErrors != 0)
