UnrealBuildTool: Move Apple platform specific files to a Platform folder and reorganize some classes. No functional changes

#rnx

diff --git a/Engine/Platforms/VisionOS/Source/Programs/UnrealBuildTool/VisionOSToolChain.cs b/Engine/Platforms/VisionOS/Source/Programs/UnrealBuildTool/VisionOSToolChain.cs
--- a/Engine/Platforms/VisionOS/Source/Programs/UnrealBuildTool/VisionOSToolChain.cs
+++ b/Engine/Platforms/VisionOS/Source/Programs/UnrealBuildTool/VisionOSToolChain.cs
@@ -1,26 +1,9 @@
 // Copyright Epic Games, Inc. All Rights Reserved.
 
 using Microsoft.Extensions.Logging;
-using EpicGames.Core;
 
 namespace UnrealBuildTool
 {
-	class VisionOSToolChainSettings : IOSToolChainSettings
-	{
-		public VisionOSToolChainSettings(IOSProjectSettings ProjectSettings, ILogger Logger) 
-			: base("XROS", "XRSimulator", "xros", ProjectSettings, Logger)
-		{
-		}
-
-		public virtual string RuntimeVersion
-		{
-			get
-			{
-				return UEBuildPlatformSDK.GetSDKForPlatform("VisionOS")!.GetSoftwareInfo()!.Min!;
-			}
-		}
-	}
-
 	class VisionOSToolChain : IOSToolChain
 	{
 		public VisionOSToolChain(ReadOnlyTargetRules InTarget, VisionOSProjectSettings InProjectSettings, ILogger InLogger)
diff --git a/Engine/Platforms/VisionOS/Source/Programs/UnrealBuildTool/VisionOSToolChainSettings.cs b/Engine/Platforms/VisionOS/Source/Programs/UnrealBuildTool/VisionOSToolChainSettings.cs
new file mode 100644
--- /dev/null
+++ b/Engine/Platforms/VisionOS/Source/Programs/UnrealBuildTool/VisionOSToolChainSettings.cs
@@ -0,0 +1,15 @@
+// Copyright Epic Games, Inc. All Rights Reserved.
+
+using Microsoft.Extensions.Logging;
+using EpicGames.Core;
+
+namespace UnrealBuildTool
+{
+	class VisionOSToolChainSettings : IOSToolChainSettings
+	{
+		public VisionOSToolChainSettings(IOSProjectSettings ProjectSettings, ILogger Logger)
+			: base("XROS", "XRSimulator", "xros", ProjectSettings, Logger)
+		{
+		}
+	}
+}
diff --git a/Engine/Source/Programs/UnrealBuildTool/Platform/Apple/AppleToolChain.cs b/Engine/Source/Programs/UnrealBuildTool/Platform/Apple/AppleToolChain.cs
new file mode 100644
--- /dev/null
+++ b/Engine/Source/Programs/UnrealBuildTool/Platform/Apple/AppleToolChain.cs
@@ -0,0 +1,1191 @@
+// Copyright Epic Games, Inc. All Rights Reserved.
+
+using System;
+using System.Collections.Generic;
+using System.Diagnostics;
+using System.IO;
+using System.Linq;
+using System.Security.Cryptography.X509Certificates;
+using System.Text;
+using System.Text.RegularExpressions;
+using System.Threading.Tasks;
+using EpicGames.Core;
+using Microsoft.Extensions.Logging;
+using UnrealBuildBase;
+
+namespace UnrealBuildTool
+{
+	abstract class AppleToolChain : ClangToolChain
+	{
+		protected class AppleToolChainInfo : ClangToolChainInfo
+		{
+			Tuple<Version, Version>[] AppleVersionToLLVMVersion;
+
+			public AppleToolChainInfo(UnrealTargetPlatform Platform, DirectoryReference DeveloperDir, FileReference Clang, FileReference Archiver, ILogger Logger)
+				: base(DeveloperDir, Clang, Archiver, Logger)
+			{
+				// get the mapping from Apple_SDK.json (turning "version ranges" into a mapping list)
+				UEBuildPlatformSDK SDK = UEBuildPlatformSDK.GetSDKForPlatform(Platform.ToString())!;
+				
+				AppleVersionToLLVMVersion = SDK.GetVersionNumberRangeArrayFromConfig("AppleVersionToLLVMVersions").
+					Select(x => new Tuple<Version, Version>(new Version(x.Min.ToString()), new Version(x.Max.ToString()))).ToArray();
+			}
+
+			// libtool doesn't provide version, just use clang's version string
+			/// <inheritdoc/>
+			protected override string QueryArchiverVersionString() => ClangVersionString;
+
+			protected override Version QueryClangVersion()
+			{
+				// get the clang version from the clang -v
+				Version AppleVersion = base.QueryClangVersion();
+
+					// now look up this version in mappings
+				for (int MappingIndex = AppleVersionToLLVMVersion.Length - 1; MappingIndex >= 0; MappingIndex--)
+				{
+					if (AppleVersion >= AppleVersionToLLVMVersion[MappingIndex].Item1)
+					{
+						Version LLVMVersion = AppleVersionToLLVMVersion[MappingIndex].Item2;
+						Logger.LogDebug("Converted Apple version {AppleVersion} to LLVM version {LLVMVersion}", AppleVersion, LLVMVersion);
+						return LLVMVersion;
+					}
+				}
+
+				throw new BuildException($"Failed to find mapping of Apple clang version {AppleVersion} in Apple_SDK.json");
+			}
+
+			// get the actual Apple version, not the LLVM version (only Apple platform code should use this)
+			public VersionNumber AppleClangVersion => VersionNumber.Parse(base.QueryClangVersion().ToString());
+		}
+
+		private Lazy<AppleToolChainSettings> ToolChainSettings;
+		
+		protected AppleToolChainSettings Settings => ToolChainSettings.Value;
+
+		public abstract ReadOnlyAppleTargetRules GetAppleTargetRules(ReadOnlyTargetRules Target);
+
+		/// <summary>
+		/// The SDK version being used to compile with the toolchain
+		/// </summary>
+		public float SDKVersionFloat => Settings.SDKVersionFloat;
+
+		public readonly ReadOnlyTargetRules? Target;
+
+		// cache some ini settings
+		readonly bool bUseSwiftUIMain;
+		readonly bool bCreateSwiftBridgingHeader;
+
+		protected bool bUseModernXcode => AppleExports.UseModernXcode(ProjectFile);
+
+		public AppleToolChain(ReadOnlyTargetRules? Target, Func<AppleToolChainSettings> InCreateSettings, ClangToolChainOptions InOptions, ILogger InLogger) : base(InOptions, InLogger)
+		{
+			this.Target = Target;
+			ProjectFile = Target?.ProjectFile;
+			ToolChainSettings = new Lazy<AppleToolChainSettings>(InCreateSettings);
+
+			AppleExports.GetSwiftIntegrationSettings(ProjectFile, Target == null ? TargetType.Game : Target.Type, Target == null ? UnrealTargetPlatform.Mac : Target.Platform, out bUseSwiftUIMain, out bCreateSwiftBridgingHeader);
+		}
+
+		/// <inheritdoc/>
+		public override void SetUpGlobalEnvironment(ReadOnlyTargetRules Target, CppCompileEnvironment GlobalCompileEnvironment, LinkEnvironment GlobalLinkEnvironment)
+		{
+			base.SetUpGlobalEnvironment(Target, GlobalCompileEnvironment, GlobalLinkEnvironment);
+
+			DirectoryReference XcodeDeveloperDir = ApplePlatformSDK.DeveloperDir;
+			DirectoryReference? XcodeAppDir = XcodeDeveloperDir.ParentDirectory?.ParentDirectory;
+
+			if (XcodeAppDir != null)
+			{
+				GlobalCompileEnvironment.RootPaths.AddFolderName(CppRootPathFolder.PlatformSDK, "Xcode.app");
+				GlobalLinkEnvironment.RootPaths.AddFolderName(CppRootPathFolder.PlatformSDK, "Xcode.app");
+				GlobalCompileEnvironment.RootPaths[CppRootPathFolder.PlatformSDK] = XcodeAppDir;
+				GlobalLinkEnvironment.RootPaths[CppRootPathFolder.PlatformSDK] = XcodeAppDir;
+
+				EpicGames.UBA.Utils.RegisterPathHash(XcodeAppDir.FullName, ApplePlatformSDK.InstalledXcodeVersion.Value ?? "0");
+
+				// Xcode sdk files are not allowed to be transferred over the network, ensure UBA will not do so
+				EpicGames.UBA.Utils.RegisterDisallowedPaths(XcodeAppDir.FullName);
+			}
+		}
+
+		/// <summary>
+		/// Takes an architecture string as provided by UBT for the target and formats it for Clang. Supports
+		/// multiple architectures joined with '+'
+		/// </summary>
+		/// <param name="InArchitectures"></param>
+		/// <returns></returns>
+		protected string FormatArchitectureArg(UnrealArchitectures InArchitectures)
+		{
+			return "-arch " + String.Join('+', InArchitectures.Architectures.Select(x => x.AppleName));
+		}
+
+		public override string GetSDKVersion()
+		{
+			return Settings.SDKVersion;
+		}
+
+		protected DirectoryReference GetMacDevSrcRoot()
+		{
+			return Unreal.EngineSourceDirectory;
+		}
+
+		protected override CPPOutput CompileCPPFiles(CppCompileEnvironment CompileEnvironment, IEnumerable<FileItem> InputFiles, DirectoryReference OutputDir, string ModuleName, IActionGraphBuilder Graph)
+		{
+			if (ShouldSkipCompile(CompileEnvironment))
+			{
+				return new CPPOutput();
+			}
+
+			// log some iseful info
+			Log.TraceInformationOnce($"Compiling with {CompileEnvironment.Platform} SDK {Settings.SDKVersion}, minimum Target OS version {Settings.GetTargetVersionForTargetType(Target!.Type)}");
+
+			List<string> GlobalArguments = new();
+
+			GetCompileArguments_Global(CompileEnvironment, GlobalArguments);
+
+			List<FileItem> FrameworkTokenFiles = new List<FileItem>();
+			foreach (UEBuildFramework Framework in CompileEnvironment.AdditionalFrameworks)
+			{
+				if (Framework.ZipFile != null)
+				{
+					// even in modern, we have to try to unzip in UBT, so that commandline builds extract the frameworks _before_ building (the Xcode finalize
+					// happens after building). We still need the unzip in Xcode pre-build phase, because Xcode wants the frameworks around for its build graph
+					ExtractFramework(Framework, Graph, OutputDir, Logger);
+					FrameworkTokenFiles.Add(Framework.ExtractedTokenFile!);
+				}
+			}
+
+			CPPOutput Result = new CPPOutput();
+			// Create a compile action for each source file.
+			List<FileItem> SwiftFiles = new List<FileItem>();
+
+			foreach (FileItem SourceFile in InputFiles)
+			{
+				if (SourceFile.HasExtension(".swift"))
+				{
+					SwiftFiles.Add(SourceFile);
+				}
+				else
+				{
+					Action CompileAction = CompileCPPFile(CompileEnvironment, SourceFile, OutputDir, ModuleName, Graph, GlobalArguments, Result);
+					CompileAction.PrerequisiteItems.UnionWith(FrameworkTokenFiles);
+				}
+			}
+			if (SwiftFiles.Any())
+			{
+				Action CompileAction = CompileSwiftFiles(CompileEnvironment, SwiftFiles, OutputDir, ModuleName, Graph, Result);
+				CompileAction.PrerequisiteItems.UnionWith(FrameworkTokenFiles);
+			}
+	
+			return Result;
+		}
+
+		protected void StripSymbolsWithXcode(FileReference SourceFile, FileReference TargetFile, DirectoryReference ToolchainDir)
+		{
+			if (SourceFile != TargetFile)
+			{
+				// Strip command only works in place so we need to copy original if target is different
+				File.Copy(SourceFile.FullName, TargetFile.FullName, true);
+			}
+
+			ProcessStartInfo StartInfo = CreateStripSymbolsProcessInfo(TargetFile, ToolchainDir);
+			Utils.RunLocalProcessAndLogOutput(StartInfo, Logger);
+		}
+
+		protected static ProcessStartInfo CreateStripSymbolsProcessInfo(FileReference File, DirectoryReference ToolchainDir) => new()
+		{
+			FileName = Path.Combine(ToolchainDir.FullName, "strip"),
+			Arguments = $"\"{File.FullName}\" -S",
+			UseShellExecute = false,
+			CreateNoWindow = true
+		};
+
+		/// <summary>
+		/// Writes a versions.xcconfig file for xcode to pull in when making an app plist
+		/// </summary>
+		/// <param name="LinkEnvironment"></param>
+		/// <param name="Prerequisite">FileItem describing the Prerequisite that this will this depends on (executable or similar) </param>
+		/// <param name="Graph">List of actions to be executed. Additional actions will be added to this list.</param>
+		protected FileItem UpdateVersionFile(LinkEnvironment LinkEnvironment, FileItem Prerequisite, IActionGraphBuilder Graph)
+		{
+			FileItem DestFile;
+
+			// Make the compile action
+			Action UpdateVersionAction = Graph.CreateAction(ActionType.CreateAppBundle);
+			UpdateVersionAction.WorkingDirectory = GetMacDevSrcRoot();
+			UpdateVersionAction.CommandPath = BuildHostPlatform.Current.Shell;
+			UpdateVersionAction.CommandDescription = "";
+
+			// @todo programs right nhow are sharing the Engine build version - one reason for this is that we can't get to the Engine/Programs directory from here
+			// (we can't even get to the Engine/Source/Programs directory without searching on disk), and if we did, we would create a _lot_ of Engine/Programs directories
+			// on disk that don't exist in p4. So, we just re-use Engine version, not Project version
+			//				DirectoryReference ProductDirectory = FindProductDirectory(ProjectFile, LinkEnvironment.OutputDirectory!, Graph.Makefile.TargetType);
+			DirectoryReference ProductDirectory = (ProjectFile?.Directory) ?? Unreal.EngineDirectory;
+			FileReference OutputVersionFile = FileReference.Combine(ProductDirectory, "Intermediate/Build/Versions.xcconfig");
+			DestFile = FileItem.GetItemByFileReference(OutputVersionFile);
+
+			// grab a changlist version if we have it to pass to the script to use if desired
+			int Changelist = 0;
+			BuildVersion? Version;
+			if (BuildVersion.TryRead(BuildVersion.GetDefaultFileName(), out Version))
+			{
+				Changelist = Version.Changelist;
+			}
+
+			// make path to the script
+			FileReference VersionScript = FileReference.Combine(ProductDirectory, "Build/BatchFiles/Mac/UpdateVersionAfterBuild.sh");
+			if (!FileReference.Exists(VersionScript))
+			{
+				VersionScript = FileReference.Combine(Unreal.EngineDirectory, "Build/BatchFiles/Mac/UpdateVersionAfterBuild.sh");
+			}
+			FileItem BundleScript = FileItem.GetItemByFileReference(VersionScript);
+			UpdateVersionAction.CommandArguments = $"\"{BundleScript.AbsolutePath}\" \"{ProductDirectory}\" {LinkEnvironment.Platform} {Changelist}";
+			UpdateVersionAction.PrerequisiteItems.Add(Prerequisite);
+			UpdateVersionAction.PrerequisiteItems.Add(BundleScript);
+			UpdateVersionAction.ProducedItems.Add(DestFile);
+			UpdateVersionAction.StatusDescription = $"Updating version file: {OutputVersionFile}";
+
+			return DestFile;
+		}
+
+		protected FileItem ExtractFramework(UEBuildFramework Framework, IActionGraphBuilder Graph, DirectoryReference IntermediateDirectory, ILogger Logger)
+		{
+			if (Framework.ZipFile == null)
+			{
+				throw new BuildException("Unable to extract framework '{0}' - no zip file specified", Framework.Name);
+			}
+			if (!Framework.bHasMadeUnzipAction)
+			{
+				Framework.bHasMadeUnzipAction = true;
+
+				FileItem InputFile = FileItem.GetItemByFileReference(Framework.ZipFile);
+
+				StringBuilder ExtractScript = new StringBuilder();
+				ExtractScript.AppendLine("#!/bin/sh");
+				ExtractScript.AppendLine("set -e");
+				// ExtractScript.AppendLine("set -x"); // For debugging
+				ExtractScript.AppendLine(String.Format("[ -d {0} ] && rm -rf {0}", Utils.MakePathSafeToUseWithCommandLine(Framework.ZipOutputDirectory!.FullName)));
+				ExtractScript.AppendLine(String.Format("unzip -q -o {0} -d {1}", Utils.MakePathSafeToUseWithCommandLine(Framework.ZipFile.FullName), Utils.MakePathSafeToUseWithCommandLine(Framework.ZipOutputDirectory.ParentDirectory!.FullName))); // Zip contains folder with the same name, hence ParentDirectory
+				ExtractScript.AppendLine(String.Format("touch {0}", Utils.MakePathSafeToUseWithCommandLine(Framework.ExtractedTokenFile!.AbsolutePath)));
+
+				// This is writing to folder shared by all configs.
+				//FileReference ScriptFile = FileReference.Combine(IntermediateDirectory, Framework.ZipOutputDirectory.GetDirectoryName() + ".sh");
+				//FileItem ExtractScriptFileItem = Graph.CreateIntermediateTextFile(ScriptFile, ExtractScript.ToString());
+				FileItem ExtractScriptFileItem = Graph.CreateIntermediateTextFile(new FileReference(Framework.ZipOutputDirectory.FullName + ".sh"), ExtractScript.ToString());
+
+				Action UnzipAction = Graph.CreateAction(ActionType.BuildProject);
+				UnzipAction.CommandPath = new FileReference("/bin/sh");
+				UnzipAction.CommandArguments = Utils.MakePathSafeToUseWithCommandLine(ExtractScriptFileItem.AbsolutePath);
+				UnzipAction.WorkingDirectory = Unreal.EngineDirectory;
+				UnzipAction.PrerequisiteItems.Add(InputFile);
+				UnzipAction.PrerequisiteItems.Add(ExtractScriptFileItem);
+				UnzipAction.ProducedItems.Add(Framework.ExtractedTokenFile);
+				UnzipAction.DeleteItems.Add(Framework.ExtractedTokenFile);
+				UnzipAction.StatusDescription = String.Format("Unzipping : {0} -> {1}", Framework.ZipFile, Framework.ZipOutputDirectory);
+				UnzipAction.bCanExecuteRemotely = false;
+			}
+			return Framework.ExtractedTokenFile!;
+		}
+
+		/// <summary>
+		/// If the project is a UnrealGame project, Target.ProjectDirectory refers to the engine dir, not the actual dir of the project. So this method gets the 
+		/// actual directory of the project whether it is a UnrealGame project or not.
+		/// </summary>
+		/// <returns>The actual project directory.</returns>
+		/// <param name="ProjectFile">The path to the project file</param>
+		internal static DirectoryReference GetActualProjectDirectory(FileReference? ProjectFile)
+		{
+			DirectoryReference ProjectDirectory = (ProjectFile == null ? Unreal.EngineDirectory : DirectoryReference.FromFile(ProjectFile)!);
+			return ProjectDirectory;
+		}
+
+		/// <inheritdoc/>
+		protected override string EscapePreprocessorDefinition(string Definition)
+		{
+			return Definition.Contains('"') ? Definition.Replace("\"", "\\\"") : Definition;
+		}
+
+		protected override void GetCppStandardCompileArgument(CppCompileEnvironment CompileEnvironment, List<string> Arguments)
+		{
+			if (CompileEnvironment.bEnableObjCAutomaticReferenceCounting)
+			{
+				Arguments.Add("-fobjc-arc");
+			}
+
+			base.GetCppStandardCompileArgument(CompileEnvironment, Arguments);
+		}
+
+		protected override void GetCompileArguments_CPP(CppCompileEnvironment CompileEnvironment, List<string> Arguments)
+		{
+			Arguments.Add("-x objective-c++");
+			GetCppStandardCompileArgument(CompileEnvironment, Arguments);
+			Arguments.Add("-stdlib=libc++");
+			Arguments.Add("-fmodules");
+			Arguments.Add("-fno-implicit-modules");
+			Arguments.Add("-fimplicit-module-maps");
+			Arguments.Add("-Wno-module-import-in-extern-c");
+			//Arguments.Add("-fcxx-modules"); // For some reason this does not work
+		}
+
+		protected override void GetCompileArguments_MM(CppCompileEnvironment CompileEnvironment, List<string> Arguments)
+		{
+			Arguments.Add("-x objective-c++");
+			GetCppStandardCompileArgument(CompileEnvironment, Arguments);
+			Arguments.Add("-stdlib=libc++");
+		}
+
+		protected override void GetCompileArguments_M(CppCompileEnvironment CompileEnvironment, List<string> Arguments)
+		{
+			Arguments.Add("-x objective-c");
+			GetCppStandardCompileArgument(CompileEnvironment, Arguments);
+			Arguments.Add("-stdlib=libc++");
+		}
+
+		protected override void GetCompileArguments_PCH(CppCompileEnvironment CompileEnvironment, List<string> Arguments)
+		{
+			Arguments.Add("-x objective-c++-header");
+			GetCppStandardCompileArgument(CompileEnvironment, Arguments);
+			Arguments.Add("-stdlib=libc++");
+			Arguments.Add("-fpch-instantiate-templates");
+		}
+
+		/// <inheritdoc/>
+		protected override void GetCompileArguments_WarningsAndErrors(CppCompileEnvironment CompileEnvironment, List<string> Arguments)
+		{
+			base.GetCompileArguments_WarningsAndErrors(CompileEnvironment, Arguments);
+
+			Arguments.Add("-Wno-unknown-warning-option");
+			Arguments.Add("-Wno-range-loop-analysis");
+			Arguments.Add("-Wno-single-bit-bitfield-constant-conversion");
+
+			// Various new warning from Xcode16.3. Disable for now.
+			VersionNumber AppleClangVersion = ((AppleToolChainInfo)Info).AppleClangVersion;
+			if (AppleClangVersion >= new VersionNumber(17))
+			{
+				// "implicit conversion loses integer precision:" in many files
+				Arguments.Add("-Wno-shorten-64-to-32");
+				
+				// clang17 is detecting some uses of cxx-extensions,
+				// such as: "error: variable length arrays in C++ are a Clang extension" 
+				Arguments.Add("-Wno-vla-extension");
+				
+				// From Engine/Plugins/Runtime/Metasound/Source/MetasoundFrontend/Public/MetasoundPrimitives/Time/AudioBuffers.h
+				Arguments.Add("-Wno-extra-qualification");
+			}
+		}
+
+		/// <inheritdoc/>
+		protected override void GetCompileArguments_Optimizations(CppCompileEnvironment CompileEnvironment, List<string> Arguments)
+		{
+			base.GetCompileArguments_Optimizations(CompileEnvironment, Arguments);
+
+			bool bStaticAnalysis = false;
+			string? StaticAnalysisMode = Environment.GetEnvironmentVariable("CLANG_STATIC_ANALYZER_MODE");
+			if (!String.IsNullOrEmpty(StaticAnalysisMode))
+			{
+				bStaticAnalysis = true;
+			}
+
+			// Optimize non- debug builds.
+			if (CompileEnvironment.bOptimizeCode && !bStaticAnalysis)
+			{
+				// Don't over optimise if using AddressSanitizer or you'll get false positive errors due to erroneous optimisation of necessary AddressSanitizer instrumentation.
+				if (Options.HasFlag(ClangToolChainOptions.EnableAddressSanitizer))
+				{
+					Arguments.Add("-O1");
+					Arguments.Add("-g");
+					Arguments.Add("-fno-optimize-sibling-calls");
+					Arguments.Add("-fno-omit-frame-pointer");
+				}
+				else if (Options.HasFlag(ClangToolChainOptions.EnableThreadSanitizer))
+				{
+					Arguments.Add("-O1");
+					Arguments.Add("-g");
+				}
+				else if (CompileEnvironment.OptimizationLevel == OptimizationMode.Size)
+				{
+					Arguments.Add("-Oz");
+				}
+				else if (CompileEnvironment.OptimizationLevel == OptimizationMode.SizeAndSpeed)
+				{
+					Arguments.Add("-Os");
+
+					if (CompileEnvironment.Architecture == UnrealArch.Arm64)
+					{
+						Arguments.Add("-moutline");
+					}
+				}
+				else
+				{
+					Arguments.Add("-O3");
+				}
+			}
+			else
+			{
+				Arguments.Add("-O0");
+			}
+		}
+
+		/// <inheritdoc/>
+		protected override void GetCompileArguments_Debugging(CppCompileEnvironment CompileEnvironment, List<string> Arguments)
+		{
+			base.GetCompileArguments_Debugging(CompileEnvironment, Arguments);
+
+			// Create DWARF format debug info if wanted,
+			if (CompileEnvironment.bCreateDebugInfo)
+			{
+				Arguments.Add("-gdwarf-4");
+
+				if (CompileEnvironment.bDebugLineTablesOnly)
+				{
+					Arguments.Add("-gline-tables-only");
+				}
+			}
+		}
+
+		/// <inheritdoc/>
+		protected override void GetCompileArguments_Analyze(CppCompileEnvironment CompileEnvironment, List<string> Arguments)
+		{
+			base.GetCompileArguments_Analyze(CompileEnvironment, Arguments);
+
+			// Disable all clang tidy checks
+			Arguments.Add($"-Xclang -analyzer-tidy-checker=-*");
+		}
+
+		/// <inheritdoc/>
+		protected override void GetCompileArguments_Global(CppCompileEnvironment CompileEnvironment, List<string> Arguments)
+		{
+			base.GetCompileArguments_Global(CompileEnvironment, Arguments);
+
+			Arguments.Add(GetRTTIFlag(CompileEnvironment));
+
+			Arguments.Add("-fmessage-length=0");
+			Arguments.Add("-fpascal-strings");
+
+			string? Override = (CompileEnvironment.bEnableOSX109Support && CompileEnvironment.Platform == UnrealTargetPlatform.Mac) ? "10.9" : null;
+			Arguments.Add($"-target {ToolChainSettings.Value.GetTargetTuple(CompileEnvironment.Architecture, Target!.Platform, Target!.Type, Override)}");
+		}
+
+		public override CppCompileEnvironment CreateSharedResponseFile(CppCompileEnvironment CompileEnvironment, FileReference OutResponseFile, IActionGraphBuilder Graph)
+		{
+			// Temporarily turn of shared response files for apple toolchains
+			return CompileEnvironment;
+		}
+
+		private FileItem GetBridgingHeader(string ModuleName, DirectoryReference OutputDir)
+		{
+			string Filename = ModuleName + "-Swift.h";
+			return FileItem.GetItemByFileReference(FileReference.Combine(OutputDir, "Bridging", Filename));
+		}
+
+		private void GetSharedSwiftArguments(CppCompileEnvironment CompileEnvironment, List<string> Arguments, string ModuleName)
+		{
+			// TODO: Fix so swiftc path supports vfs
+			CppRootPaths rootPaths = new(CompileEnvironment.RootPaths);
+			rootPaths.bUseVfs = false;
+
+			Arguments.Add($"-target {ToolChainSettings.Value.GetTargetTuple(CompileEnvironment.Architecture, Target!.Platform, Target!.Type)}");
+			Arguments.Add($"-sdk {ToolChainSettings.Value.GetSDKPath(CompileEnvironment.Architecture)}");
+
+			Arguments.Add($"-swift-version 5");
+			if (CompileEnvironment.bCreateDebugInfo)
+			{
+				Arguments.Add($"-g");
+			}
+
+			Arguments.Add($"-module-name {ModuleName}Swift");
+
+			if (CompileEnvironment.ForceIncludeFiles.Count > 0)
+			{
+				Arguments.Add($"-import-objc-header \"{NormalizeCommandLinePath(CompileEnvironment.ForceIncludeFiles.First(), rootPaths)}\"");
+			}
+
+			// tell swift to define PLATFORM_FOO, and tell C++ to define PLATFORM_FOO=1
+			// swift doesn't pass it's defines along to C++ when importing a C header, and it also doesn't have values for defines
+			Arguments.Add($"-DPLATFORM_{Target.Platform.ToString().ToUpper()}");
+			Arguments.Add($"-Xcc -DPLATFORM_{Target.Platform.ToString().ToUpper()}=1");
+			if (bUseSwiftUIMain)
+			{
+				Arguments.Add("-DUE_USE_SWIFT_UI_MAIN");
+				if (ToolChainSettings.Value.SDKVersionFloat < 2.0)
+				{
+					Arguments.Add("-DUE_SDK_VERSION_1");
+				}
+			}
+
+			// pass the UE definitions to C++ land for the imported obj-c header
+			Arguments.AddRange(CompileEnvironment.Definitions.Select(x => $"-Xcc -D{x}"));
+
+			Console.WriteLine("Arguments: {0}", String.Join(", ", Arguments));
+		}
+
+		private Action CompileSwiftFiles(CppCompileEnvironment CompileEnvironment, IEnumerable<FileItem> InputFiles, DirectoryReference OutputDir, string ModuleName, IActionGraphBuilder Graph, CPPOutput CompileResult)
+		{
+			FileItem OutputFile = FileItem.GetItemByFileReference(FileReference.Combine(OutputDir, GetFileNameFromExtension(ModuleName + ".swift", ".o")));
+
+			List<string> Arguments = new();
+			foreach (FileItem SourceFile in InputFiles)
+			{
+				Arguments.Add($"\"{SourceFile}\"");
+			}
+
+			GetSharedSwiftArguments(CompileEnvironment, Arguments, ModuleName);
+			// output file settings
+			Arguments.Add("-emit-object"); // same as -c
+			Arguments.Add($"-parse-as-library"); // don't create a main function (even if we use the SwiftUI @main, this still seems to be correctly working)
+			Arguments.Add($"-o \"{OutputFile}\"");
+
+			// platform settings
+			Arguments.Add($"-target {Settings.GetTargetTuple(CompileEnvironment.Architecture, Target!.Platform, Target!.Type)}");
+			Arguments.Add($"-sdk {Settings.GetSDKPath(CompileEnvironment.Architecture)}");
+
+			// misc settings copied from Xcode
+			Arguments.Add("-Xllvm");
+			Arguments.Add("-aarch64-use-tbi");
+			Arguments.Add("-stack-check");
+
+			Arguments.Add("-enable-objc-interop");
+			Arguments.Add("-cxx-interoperability-mode=default");
+
+			Action CompileAction = Graph.CreateAction(ActionType.Compile);
+			CompileAction.Weight = CompileActionWeight;
+			CompileAction.CommandArguments = String.Join(" ", Arguments);
+			CompileAction.CommandPath = FileReference.Combine(ToolChainSettings.Value.ToolchainDir, "swift-frontend");
+			foreach (FileItem SourceFile in InputFiles)
+			{
+				CompileAction.PrerequisiteItems.Add(SourceFile);
+			}
+			CompileAction.ProducedItems.Add(OutputFile);
+			CompileResult.ObjectFiles.Add(OutputFile);
+			CompileAction.WorkingDirectory = Unreal.EngineSourceDirectory;
+			CompileAction.CommandDescription = "Compile";
+			UnrealArchitectureConfig ArchConfig = UnrealArchitectureConfig.ForPlatform(CompileEnvironment.Platform);
+			if (ArchConfig.Mode != UnrealArchitectureMode.SingleArchitecture)
+			{
+				string ReadableArch = ArchConfig.ConvertToReadableArchitecture(CompileEnvironment.Architecture);
+				CompileAction.CommandDescription += $" [{ReadableArch}]";
+			}
+			CompileAction.StatusDescription = String.Join(" ", InputFiles.Select(x => x.Name));
+			CompileAction.bIsClangCompiler = true;
+			CompileAction.bCanExecuteRemotely = true;
+			CompileAction.bCanExecuteInUBA = false;
+
+
+			if (CompileEnvironment.ForceIncludeFiles.Count > 0)
+			{
+				FileItem OutputInteropHeader = GetBridgingHeader(ModuleName, OutputDir);
+
+				// obj-c bridging header settings
+				Arguments.Clear();
+				foreach (FileItem SourceFile in InputFiles)
+				{
+					Arguments.Add($"\"{SourceFile}\"");
+				}
+				GetSharedSwiftArguments(CompileEnvironment, Arguments, ModuleName);
+
+				Arguments.Add("-parse"); // this will allow it to generate the header without writing out any .o/executable
+				Arguments.Add("-emit-objc-header");
+				Arguments.Add($"-emit-objc-header-path \"{OutputInteropHeader}\"");
+				Arguments.Add("-parse-as-library");
+
+				//Arguments.Add("-enable-objc-interop");
+				//Arguments.Add("-cxx-interoperability-mode=default");
+				
+				// now make an action to export the swift code as a header Obj-C bridging
+				Action HeaderAction = Graph.CreateAction(ActionType.Compile);
+				HeaderAction.CommandPath = FileReference.Combine(ToolChainSettings.Value.ToolchainDir, "swiftc");
+				HeaderAction.CommandArguments = String.Join(" ", Arguments);
+				foreach (FileItem SourceFile in InputFiles)
+				{
+					CompileAction.PrerequisiteItems.Add(SourceFile);
+				}
+				HeaderAction.ProducedItems.Add(OutputInteropHeader);
+				// swiftc won't touch the output header if the generated content doesn't change, so always delete the output so that 
+				// the timestamp is correct (it must be newer than the input .swift file, or it will build over and over)
+				HeaderAction.DeleteItems.Add(OutputInteropHeader);
+				CompileResult.GeneratedHeaderFiles.Add(OutputInteropHeader);
+				HeaderAction.CommandDescription = "Generate SwiftToCPP Header";
+				HeaderAction.StatusDescription = Path.GetFileName(OutputInteropHeader.AbsolutePath);
+				if (ArchConfig.Mode != UnrealArchitectureMode.SingleArchitecture)
+				{
+					string ReadableArch = ArchConfig.ConvertToReadableArchitecture(CompileEnvironment.Architecture);
+					HeaderAction.CommandDescription += $" [{ReadableArch}]";
+				}
+
+				HeaderAction.WorkingDirectory = CompileAction.WorkingDirectory;
+				HeaderAction.bIsClangCompiler = CompileAction.bIsClangCompiler;
+				HeaderAction.bCanExecuteRemotely = CompileAction.bCanExecuteRemotely;
+				HeaderAction.bCanExecuteInUBA = CompileAction.bCanExecuteInUBA;
+			}
+			// this is likely ignored, but the Compile action is the important one
+			return CompileAction;
+		}
+
+		public override ICollection<FileItem> PostBuild(ReadOnlyTargetRules Target, FileItem Executable, LinkEnvironment BinaryLinkEnvironment, IActionGraphBuilder Graph)
+		{
+			List<FileItem> OutputFiles = new List<FileItem>(base.PostBuild(Target, Executable, BinaryLinkEnvironment, Graph));
+
+			bool bIsBuildingAppBundle = !BinaryLinkEnvironment.bIsBuildingDLL && !BinaryLinkEnvironment.bIsBuildingLibrary && !BinaryLinkEnvironment.bIsBuildingConsoleApplication;
+			if (AppleExports.UseModernXcode(Target.ProjectFile) && bIsBuildingAppBundle)
+			{
+				Action PostBuildAction = ApplePostBuildSyncMode.CreatePostBuildSyncAction(Target, Executable, BinaryLinkEnvironment.IntermediateDirectory!, Graph, AppleExports.ForceNoEntitlements());
+
+				PostBuildAction.PrerequisiteItems.UnionWith(OutputFiles);
+				OutputFiles.AddRange(PostBuildAction.ProducedItems);
+
+				OutputFiles.Add(UpdateVersionFile(BinaryLinkEnvironment, FileItem.GetItemByFileReference(BinaryLinkEnvironment.OutputFilePath), Graph));
+			}
+			
+			ReadOnlyAppleTargetRules AppleTargetRules = GetAppleTargetRules(Target);
+			if (!BinaryLinkEnvironment.bIsBuildingLibrary && (AppleTargetRules.bStripSymbols || Target.Configuration == UnrealTargetConfiguration.Shipping))
+			{
+				// Generate the standard strip command. Should our default strip command ever change, we have a centralized place to update it.
+				ProcessStartInfo StripProcessInfo = CreateStripSymbolsProcessInfo(Executable.Location, Settings.ToolchainDir);
+
+				// If building a framework, we can only strip local symbols and we need to leave global in place (-x argument)
+				// We also preserve global in binaries on IOS/TVOS as we expose operators new/delete.
+				bool bPreserveGlobalSymbols = BinaryLinkEnvironment.bIsBuildingDLL || (Target.Platform == UnrealTargetPlatform.IOS || Target.Platform == UnrealTargetPlatform.TVOS);
+				string AddedArguments = bPreserveGlobalSymbols ? " -x" : "";
+				// Determine all of the arguments we're passing to the strip command (the process info plus arguments above)
+				string StripCommandArguments = $"{StripProcessInfo.Arguments}{AddedArguments}";
+
+				// Convert the process info to an action graph entry
+				Action StripAction = Graph.CreateAction(ActionType.CreateAppBundle);
+				StripAction.WorkingDirectory = GetMacDevSrcRoot();
+				StripAction.CommandPath = BuildHostPlatform.Current.Shell;
+				StripAction.PrerequisiteItems.Add(Executable);
+				StripAction.PrerequisiteItems.UnionWith(OutputFiles);
+				StripAction.StatusDescription = $"Stripping symbols from {Executable.AbsolutePath}";
+				StripAction.bCanExecuteRemotely = false;
+				
+				// Construct the strip command we want to execute
+				string Command = $"\"{StripProcessInfo.FileName}\" {StripCommandArguments}";
+				
+				// Determine if we are creating a strip flag file.
+				// This will essentially be the same location & file name, but with a ".stripped" extension.
+				if (AppleTargetRules.bCreateStripFlagFile)
+				{
+					// Determine the name & location of the strip flag file
+					FileItem StripCompleteFile = FileItem.GetItemByFileReference(FileReference.Combine(BinaryLinkEnvironment.IntermediateDirectory!, Executable.Location.GetFileName() + ".stripped"));
+					
+					// Create this strip flag file as part of the command we are executing
+					Command += $" && touch \"{StripCompleteFile.FullName}\"";
+					
+					// Ensure our action identifies this as a file it is producing and it's flagged as an output file from this process
+					StripAction.ProducedItems.Add(StripCompleteFile);
+					OutputFiles.Add(StripCompleteFile);
+				}
+				else
+				{
+					// If we're not creating a stripped flag file, we're stripping the actual
+					// target file, so we need to add that as an output from this process
+					OutputFiles.Add(Executable);
+				}
+				
+				// Additional escaping for backslashes and quotes since we are executing a shell and passing in our command as a quoted string
+				string ShellCommand = $"-c \"{Command.Replace(@"\", @"\\").Replace("\"", "\\\"")}\"";
+				// Finally, we assign the action our finished shell command before it executes
+				StripAction.CommandArguments = ShellCommand;
+			}
+
+			return OutputFiles;
+		}
+
+		protected override void GetLinkArguments_Optimizations(LinkEnvironment linkEnvironment, List<string> arguments)
+		{
+			base.GetLinkArguments_Optimizations(linkEnvironment, arguments);
+
+			bool bLTOEnabled = Options.HasFlag(ClangToolChainOptions.EnableLinkTimeOptimization);
+			bool bThinLTOEnabled = Options.HasFlag(ClangToolChainOptions.EnableThinLTO);
+
+			if (bLTOEnabled && bThinLTOEnabled)
+			{
+				arguments.Remove("-Wl,--thinlto-jobs=all"); // unsupported
+
+				if (linkEnvironment.ThinLTOCacheDirectory != null)
+				{
+					arguments.RemoveAll(x => x.StartsWith("-Wl,--thinlto-cache", StringComparison.Ordinal));
+				}
+			}
+		}
+
+		protected override void GetLinkArguments_Global(LinkEnvironment LinkEnvironment, List<string> Arguments)
+		{
+			base.GetLinkArguments_Global(LinkEnvironment, Arguments);
+
+			Arguments.Add($"-isysroot \"{NormalizeCommandLinePath(Settings.GetSDKPath(LinkEnvironment.Architecture), LinkEnvironment.RootPaths)}\"");
+
+			VersionNumber AppleClangVersion = ((AppleToolChainInfo)Info).AppleClangVersion;
+			// Temp solution for UE-191350
+			if (AppleClangVersion >= new VersionNumber(15) && AppleClangVersion < new VersionNumber(16))
+			{
+				Arguments.Add("-ld_classic");
+			}
+
+			if (!LinkEnvironment.bIsBuildingDLL)
+			{
+				string PlatformName = LinkEnvironment.Architecture == UnrealArch.IOSSimulator ? Settings.SimulatorPlatformDirName : Settings.PlatformDirName;
+				DirectoryReference libSwiftDir = DirectoryReference.Combine(ApplePlatformSDK.GetToolchainDirectory(), "..", "lib", "swift", PlatformName);
+				Arguments.Add($"-L\"{NormalizeCommandLinePath(libSwiftDir, LinkEnvironment.RootPaths)}\"");
+			}
+
+			Arguments.Add("-L\"/usr/lib/swift\"");
+
+			// Pass through architecture and OS info		
+			Arguments.Add(FormatArchitectureArg(LinkEnvironment.Architectures));
+			Arguments.Add($"-target {Settings.GetTargetTuple(LinkEnvironment.Architecture, Target!.Platform, Target!.Type)}");
+
+			Arguments.Add("-dead_strip");
+		}
+
+		protected override void GetArchiveArguments_Global(LinkEnvironment LinkEnvironment, List<string> Arguments)
+		{
+			base.GetArchiveArguments_Global(LinkEnvironment, Arguments);
+			Arguments.Add("-static");
+		}
+
+		#region Stub Xcode Projects
+
+		internal static bool GenerateProjectFiles(FileReference? ProjectFile, string[] Arguments, string? SingleTargetName, ILogger Logger, out DirectoryReference? XcodeProjectFile)
+		{
+			ProjectFileGenerator.bGenerateProjectFiles = true;
+			try
+			{
+				CommandLineArguments CmdLine = new CommandLineArguments(Arguments);
+
+				PlatformProjectGeneratorCollection PlatformProjectGenerators = new PlatformProjectGeneratorCollection();
+				PlatformProjectGenerators.RegisterPlatformProjectGenerator(UnrealTargetPlatform.Mac, new MacProjectGenerator(CmdLine, Logger), Logger);
+				PlatformProjectGenerators.RegisterPlatformProjectGenerator(UnrealTargetPlatform.IOS, new IOSProjectGenerator(CmdLine, Logger), Logger);
+				PlatformProjectGenerators.RegisterPlatformProjectGenerator(UnrealTargetPlatform.TVOS, new TVOSProjectGenerator(CmdLine, Logger), Logger);
+
+				XcodeProjectFileGenerator Generator = new XcodeProjectFileGenerator(ProjectFile, CmdLine);
+				// this could be improved if ProjectFileGenerator constructor took a Arguments param, and it could parse it there instead of the GenerateProjectFilesMode
+				Generator.SingleTargetName = SingleTargetName;
+				// don't need the editor data since these are stub projects
+				ProjectFileGenerator.Current = Generator;
+				bool bSucces = Generator.GenerateProjectFiles(PlatformProjectGenerators, Arguments, Logger);
+				ProjectFileGenerator.Current = null;
+				XcodeProjectFile = Generator.XCWorkspace;
+				return bSucces;
+			}
+			catch (Exception ex)
+			{
+				XcodeProjectFile = null;
+				Logger.LogError(ex.ToString());
+			}
+			finally
+			{
+				ProjectFileGenerator.bGenerateProjectFiles = false;
+			}
+			return false;
+		}
+
+		/// <summary>
+		/// Genearate an run-only Xcode project, that is not meant to be used for anything else besides code-signing/running/etc of the native .app bundle
+		/// </summary>
+		/// <param name="UProjectFile">Location of .uproject file (or null for the engine project</param>
+		/// <param name="Platform">The platform to generate a project for</param>
+		/// <param name="TargetName">The name of the target being built, so we can generate a more minimal project</param>
+		/// <param name="bForDistribution">True if this is making a bild for uploading to app store</param>
+		/// <param name="bNoEntitlements">True if we shouldn't use entitlements</param>
+		/// <param name="Logger">Logging object</param>
+		/// <param name="GeneratedProjectFile">Returns the .xcworkspace that was made</param>
+		internal static void GenerateRunOnlyXcodeProject(FileReference? UProjectFile, UnrealTargetPlatform Platform, string TargetName, bool bForDistribution, bool bNoEntitlements, ILogger Logger, out DirectoryReference? GeneratedProjectFile)
+		{
+			List<string> Options = new()
+			{
+				$"-platforms={Platform}",
+				"-DeployOnly",
+				"-NoIntellisense",
+				"-NoDotNet",
+				"-IgnoreJunk",
+				bNoEntitlements ? "-noEntitlements" : "",
+				bForDistribution ? "-distribution" : "-development",
+				"-IncludeTempTargets",
+				"-projectfileformat=XCode",
+				"-automated",
+			};
+
+			if (!String.IsNullOrEmpty(TargetName))
+			{
+				Options.Add($"-singletarget={TargetName}");
+			}
+
+			if (UProjectFile == null || UProjectFile.IsUnderDirectory(Unreal.EngineDirectory))
+			{
+				// @todo do we need these? where would the bundleid come from if there's no project?
+				//				Options.Add("-bundleID=" + BundleID);
+				//				Options.Add("-appname=" + AppName);
+				// @todo add an option to only add Engine target?
+			}
+			else
+			{
+				Options.Add($"-project=\"{UProjectFile.FullName}\"");
+				Options.Add("-game");
+			}
+
+			// we need to be in Engine/Source for some build.cs files
+			string CurrentCWD = Environment.CurrentDirectory;
+			Environment.CurrentDirectory = DirectoryReference.Combine(Unreal.EngineDirectory, "Source").FullName;
+			GenerateProjectFiles(UProjectFile, Options.ToArray(), TargetName, Logger, out GeneratedProjectFile);
+			Environment.CurrentDirectory = CurrentCWD;
+		}
+
+		static string IdOrNameParam(string Id)
+		{
+			string OldFormat = "^[0-9,a-f]{40}$";
+			string NewFormat = "^[0-9,A-F]{8}-[0-9,A-F]{16}$";
+			if (Regex.Match(Id, OldFormat).Success || Regex.Match(Id, NewFormat).Success)
+			{
+				return $"id={Id}";
+			}
+			return $"name={Id}";
+		}
+
+		internal static int FinalizeAppWithXcode(DirectoryReference XcodeProject, UnrealTargetPlatform Platform, UnrealArchitectures Architectures, bool bUseAutomaticCodeSigning, string SchemeName, string Configuration, string Action, string ExtraOptions, ILogger Logger, List<string>? DestinationIds=null)
+		{
+			// Acquire a different mutex to the regular UBT instance, since this mode will be called as part of a build. We need the mutex to ensure that building two modular configurations 
+			// in parallel don't clash over writing shared *.modules files (eg. DebugGame and Development editors).
+			string MutexName = GlobalSingleInstanceMutex.GetUniqueMutexForPath("UnrealBuildTool_XcodeBuild", Unreal.RootDirectory);
+			using (new GlobalSingleInstanceMutex(MutexName, true))
+			{
+				for (int Pass = 0; Pass < 2; Pass++)
+				{
+					string DestinationParam;
+					// use generic:
+					//   - for Mac, the destination seems to want name="My Mac", not the real name of the Mac
+					//   - when no destinations were supplied
+					//   - when trying to build iOS app to run on local mac 
+					//   - when using manual codesigning because the main use of this is to update provisioning information (reducing potential for issues)
+					if (Platform == UnrealTargetPlatform.Mac || DestinationIds == null || DestinationIds.Count == 0 || DestinationIds[0].Length == 0 || DestinationIds[0].Contains(AppleToolChainSettings.LocalMacUUID) || !bUseAutomaticCodeSigning)
+					{
+						DestinationParam = $"-destination generic/platform=\"{AppleExports.GetDestinationPlatform(Platform, Architectures)}\"";
+					}
+					else
+					{
+						string DestPlatform = AppleExports.GetDestinationPlatform(Platform, Architectures);
+						DestinationParam = String.Join(" ", DestinationIds.Select(x => $"-destination \"platform={DestPlatform},{IdOrNameParam(x)}\""));
+					}
+					List<string> Arguments = new()
+					{
+						"UBT_NO_POST_DEPLOY=true",
+						FileReference.Combine(ApplePlatformSDK.DeveloperDir, "usr/bin/xcodebuild").FullName,
+						Action,
+						$"-workspace \"{XcodeProject.FullName}\"",
+						$"-scheme \"{SchemeName}\"",
+						$"-configuration \"{Configuration}\"",
+						DestinationParam,
+						// xcode gets confused it we _just_ wrote out entitlements while generating the temp project, and it thinks it was modified _during_ building
+						// but it wasn't, it was written before the build started
+						"CODE_SIGN_ALLOW_ENTITLEMENTS_MODIFICATION=YES",
+						ExtraOptions,
+						//$"-sdk {SDKName}",
+					};
+
+					if (Pass == 0)
+					{
+						Arguments.Add("-hideShellScriptEnvironment");
+					}
+
+
+					Console.WriteLine("params: {0}", String.Join(" ", Arguments));
+
+					Process LocalProcess = new Process();
+					LocalProcess.StartInfo = new ProcessStartInfo("/usr/bin/env", String.Join(" ", Arguments));
+					LocalProcess.OutputDataReceived += (Sender, Args) => { LocalProcessOutput(Args, false, Logger); };
+					LocalProcess.ErrorDataReceived += (Sender, Args) =>
+					{
+						if (Args != null && Args.Data != null
+						&& Args.Data.Contains("Failed to load profile") && Args.Data.Contains("<stdin>"))
+						{
+							Logger.LogInformation("Silencing the following provision profile error, it is not affecting code signing:");
+							LocalProcessOutput(Args, false, Logger);
+						}
+						else
+						{
+							LocalProcessOutput(Args, true, Logger);
+						}
+					};
+
+
+					// if first pass succeeded, return 0, otherwise return whatever the exit code is
+					int ExitCode = Utils.RunLocalProcess(LocalProcess);
+					if (ExitCode == 0 || Pass > 0)
+					{
+						return ExitCode;
+					}
+
+					Logger.LogWarning("");
+					Logger.LogWarning("");
+					Logger.LogWarning("**************************************************************************************************************");
+					Logger.LogWarning("* App finalization failed, running again with all environment variables displayed to help diagnose the error *");
+					Logger.LogWarning("**************************************************************************************************************");
+					Logger.LogWarning("");
+					Logger.LogWarning("");
+				}
+			}
+			return -1;
+		}
+
+		static void LocalProcessOutput(DataReceivedEventArgs? Args, bool bIsError, ILogger Logger)
+		{
+			if (Args != null && Args.Data != null)
+			{
+				if (bIsError)
+				{
+					Logger.LogError("{Message}", Args.Data.TrimEnd());
+				}
+				else
+				{
+					Logger.LogInformation("{Message}", Args.Data.TrimEnd());
+				}
+			}
+		}
+
+		#endregion
+	};
+
+	[Serializable]
+	class ApplePostBuildSyncTarget
+	{
+		public FileReference? ProjectFile;
+		public UnrealTargetPlatform Platform;
+		public UnrealArchitectures Architectures;
+		public UnrealTargetConfiguration Configuration;
+		public string TargetName;
+
+		// For iOS/TVOS
+		public bool bCreateStubIPA;
+		public string? RemoteImportProvision;
+		public string? RemoteImportCertificate;
+		public string? RemoteImportCertificatePassword;
+		public DirectoryReference ProjectIntermediateDirectory;
+		public FileReference StubOutputPath;
+
+		public ApplePostBuildSyncTarget(ReadOnlyTargetRules Target, FileItem Executable, DirectoryReference IntermediateDir)
+		{
+			Platform = Target.Platform;
+			Configuration = Target.Configuration;
+			Architectures = Target.Architectures;
+			ProjectFile = Target.ProjectFile;
+			TargetName = Target.Name;
+
+			bCreateStubIPA = Target.IOSPlatform.bCreateStubIPA;
+			RemoteImportProvision = Target.IOSPlatform.ImportProvision;
+			RemoteImportCertificate = Target.IOSPlatform.ImportCertificate;
+			RemoteImportCertificatePassword = Target.IOSPlatform.ImportCertificatePassword;
+			ProjectIntermediateDirectory = IntermediateDir;
+
+			StubOutputPath = Executable.Location;
+		}
+	}
+
+	[ToolMode("ApplePostBuildSync", ToolModeOptions.XmlConfig | ToolModeOptions.BuildPlatforms)]
+	class ApplePostBuildSyncMode : ToolMode
+	{
+		[CommandLine("-Input=", Required = true)]
+		public FileReference? InputFile = null;
+
+		[CommandLine("-XmlConfigCache=")]
+		public FileReference? XmlConfigCache = null;
+
+		// this isn't actually used, but is helpful to pass -modernxcode along in CreatePostBuildSyncAction, and UBT won't
+		// complain that nothing is using it, because where we _do_ use it is outside the normal cmdline parsing functionality
+		[CommandLine("-ModernXcode")]
+		public bool bModernXcode;
+
+		[CommandLine("-noEntitlements")]
+		public bool bNoEntitlements;
+
+		public override Task<int> ExecuteAsync(CommandLineArguments Arguments, ILogger Logger)
+		{
+			Arguments.ApplyTo(this);
+			Arguments.CheckAllArgumentsUsed();
+
+			// Run the PostBuildSync command
+			ApplePostBuildSyncTarget Target = BinaryFormatterUtils.Load<ApplePostBuildSyncTarget>(InputFile!);
+			int ExitCode = PostBuildSync(Target, Logger);
+
+			return Task.FromResult(ExitCode);
+		}
+
+		private int PostBuildSync(ApplePostBuildSyncTarget Target, ILogger Logger)
+		{
+			// generate the IOS plist file every time
+			if (Target.Platform.IsInGroup(UnrealPlatformGroup.IOS))
+			{
+				string GameName = Target.ProjectFile == null ? "UnrealGame" : Target.ProjectFile.GetFileNameWithoutAnyExtensions();
+				// most of these params are uused in modern
+				UEDeployIOS.GenerateIOSPList(Target.ProjectFile, Target.Configuration, AppleToolChain.GetActualProjectDirectory(Target.ProjectFile).FullName, Target.ProjectFile == null, GameName, bIsClient: false,
+					GameName, Unreal.EngineDirectory.FullName, "", null, null, false, Logger);
+			}
+
+			// if xcode is building this, it will also do the Run stuff anyway, so no need to do it here as well
+			if (Environment.GetEnvironmentVariable("UE_BUILD_FROM_XCODE") == "1")
+			{
+				return 0;
+			}
+
+			string ExtraOptions = "";
+			// for mobile builds (which need real codesigning to be able to run), we use dummy codesigning when making a .stub (which will
+			// be sent to Windows an re-codesigned), or when making UnrealGame.app without a .uproject (we will always make a .app
+			// again with a .uproject on the commandline to be able to get the staged data that will be pulled in to the app - we can't run 
+			// without a Staged directory, so this dummy codesigned .app won't be used directly)
+
+			// NOTE: Actually for _now_ we are using legacy-style signing with temp keychain because IPhonePackager cannot codesign
+			// Frameworks, so we have to do full non-dummy signing of stubs until we get IPP working
+			bool bUseDummySigning = Target.Platform != UnrealTargetPlatform.Mac && (Target.ProjectFile == null);
+			bool bCreateStub = Target.Platform.IsInGroup(UnrealPlatformGroup.IOS) && Target.bCreateStubIPA;
+			// if we want dummy signing (no project at all) then we don't want to use legacy signing - that is only used for remote builds from Windows
+			bool bUseLegacyStubSigning = bCreateStub && !bUseDummySigning;
+
+			if (bUseLegacyStubSigning)
+			{
+				// create and run a script that will make a temp keychain, and return the options needed to pass to xcodebuild to use it
+				ExtraOptions += SetupRemoteCodesigning(Target);
+			}
+
+			int ExitCode = AppleExports.BuildWithStubXcodeProject(Target.ProjectFile, Target.Platform, Target.Architectures, Target.Configuration, Target.TargetName,
+				AppleExports.XcodeBuildMode.PostBuildSync, Logger, ExtraOptions, bForceDummySigning: bUseDummySigning);
+
+			// restore the keychain as soon as possible
+			if (bUseLegacyStubSigning)
+			{
+				// cleanup the 
+				CleanupRemoteCodesigning(Target);
+			}
+
+			if (ExitCode != 0)
+			{
+				Logger.LogError("ERROR: Failed to finalize the .app with Xcode. Check the log for more information");
+			}
+
+			if (bCreateStub)
+			{
+				IOSToolChain.PackageStub(Target.StubOutputPath.Directory.FullName, Target.TargetName, Target.StubOutputPath.GetFileNameWithoutExtension(), true, !bUseLegacyStubSigning);
+
+			}
+
+			return ExitCode;
+		}
+
+		// This is hopefully until we can get codesigning of Frameworks working in IPhonePackager, then we can go back to dummy codesigning, without needing
+		// to mess with keychains and what not
+		private static string SetupRemoteCodesigning(ApplePostBuildSyncTarget Target)
+		{
+			FileReference TempKeychain = FileReference.Combine(Target.ProjectIntermediateDirectory!, "TempKeychain.keychain");
+			FileReference SignProjectScript = FileReference.Combine(Target.ProjectIntermediateDirectory!, "SignProject.sh");
+			string MobileProvisionUUID = "";
+			string SigningCertificate = "";
+
+			using (StreamWriter Writer = new StreamWriter(SignProjectScript.FullName))
+			{
+				// Boilerplate
+				Writer.WriteLine("#!/bin/sh");
+				Writer.WriteLine("set -e");
+				Writer.WriteLine("set -x");
+				// Copy the mobile provision into the system store
+				if (Target.RemoteImportProvision == null || Target.RemoteImportCertificate == null)
+				{
+					throw new BuildException("Expecting stub to be run with -ImportCertificate and -ImportProvision when using modern xcode");
+				}
+
+				// copy the provision into standard location
+				Writer.WriteLine("cp -f {0} {1}", Utils.EscapeShellArgument(Target.RemoteImportProvision), Utils.EscapeShellArgument(AppleExports.GetProvisionDirectory().FullName));
+				MobileProvisionContents MobileProvision = MobileProvisionContents.Read(new FileReference(Target.RemoteImportProvision));
+				MobileProvisionUUID = MobileProvision.GetUniqueId();
+
+				// Get the signing certificate to use
+				X509Certificate2 Certificate;
+				try
+				{
+					Certificate = new X509Certificate2(Target.RemoteImportCertificate, Target.RemoteImportCertificatePassword ?? "");
+				}
+				catch (Exception Ex)
+				{
+					throw new BuildException(Ex, "Unable to read certificate '{0}': {1}", Target.RemoteImportCertificate, Ex.Message);
+				}
+				// Read the name from the certificate
+				SigningCertificate = Certificate.GetNameInfo(X509NameType.SimpleName, false);
+
+				// Install a certificate given on the command line to a temporary keychain
+				Writer.WriteLine("security delete-keychain \"{0}\" || true", TempKeychain);
+				Writer.WriteLine("security create-keychain -p \"A\" \"{0}\"", TempKeychain);
+				Writer.WriteLine("security list-keychains -s \"{0}\"", TempKeychain);
+				Writer.WriteLine("security list-keychains");
+				Writer.WriteLine("security set-keychain-settings -t 3600 -l  \"{0}\"", TempKeychain);
+				Writer.WriteLine("security -v unlock-keychain -p \"A\" \"{0}\"", TempKeychain);
+				Writer.WriteLine("security import {0}/Build/IOS/AppleWorldwideDeveloperRelationsCA.pem -k \"{1}\"", Utils.EscapeShellArgument(Unreal.EngineDirectory.FullName), TempKeychain);
+				Writer.WriteLine("security import {0} -P {1} -k \"{2}\" -T /usr/bin/codesign -T /usr/bin/security -t agg", Utils.EscapeShellArgument(Target.RemoteImportCertificate), Utils.EscapeShellArgument(Target.RemoteImportCertificatePassword!), TempKeychain);
+				Writer.WriteLine("security set-key-partition-list -S apple-tool:,apple:,codesign: -s -k \"A\" -D '{0}' -t private {1}", SigningCertificate, TempKeychain);
+			}
+
+			// run the script
+			Utils.RunLocalProcessAndReturnStdOut("/bin/sh", $"\"{SignProjectScript.FullName}\"");
+
+			// Set parameters to make sure it uses the correct identity and keychain
+			// pass back the comandline arguments to xcodebuild to use these certicicates
+			return $" CODE_SIGN_STYLE=Manual CODE_SIGN_IDENTITY=\"{SigningCertificate}\" PROVISIONING_PROFILE_SPECIFIER={MobileProvisionUUID}";
+		}
+
+		private static void CleanupRemoteCodesigning(ApplePostBuildSyncTarget Target)
+		{
+			FileReference TempKeychain = FileReference.Combine(Target.ProjectIntermediateDirectory!, "TempKeychain.keychain");
+
+			FileReference CleanProjectScript = FileReference.Combine(Target.ProjectIntermediateDirectory!, "CleanProject.sh");
+			using (StreamWriter CleanWriter = new StreamWriter(CleanProjectScript.FullName))
+			{
+				CleanWriter.WriteLine("#!/bin/sh");
+				CleanWriter.WriteLine("set -e");
+				CleanWriter.WriteLine("set -x");
+				// Remove the temporary keychain from the search list
+				CleanWriter.WriteLine("security delete-keychain \"{0}\" || true", TempKeychain);
+				// Restore the login keychain as active
+				CleanWriter.WriteLine("security list-keychain -s login.keychain");
+			}
+
+			Utils.RunLocalProcessAndReturnStdOut("/bin/sh", $"\"{CleanProjectScript.FullName}\"");
+		}
+
+		private static FileItem GetPostBuildOutputFile(FileReference Executable, string TargetName, UnrealTargetPlatform Platform)
+		{
+			FileReference StagedExe;
+			if (Platform == UnrealTargetPlatform.Mac)
+			{
+				StagedExe = FileReference.Combine(Executable.Directory, Executable.GetFileName() + ".app/Contents/PkgInfo");
+			}
+			else
+			{
+				StagedExe = FileReference.Combine(Executable.Directory, TargetName + ".app", TargetName);
+			}
+			return FileItem.GetItemByFileReference(StagedExe);
+		}
+
+		public static Action CreatePostBuildSyncAction(ReadOnlyTargetRules Target, FileItem Executable, DirectoryReference IntermediateDir, IActionGraphBuilder Graph, bool bNoEntitlements = false)
+		{
+			ApplePostBuildSyncTarget PostBuildSync = new(Target, Executable, IntermediateDir);
+			FileReference PostBuildSyncFile = FileReference.Combine(IntermediateDir!, "PostBuildSync.dat");
+			BinaryFormatterUtils.Save(PostBuildSyncFile, PostBuildSync);			
+			string PostBuildSyncArguments = String.Format("-modernxcode -Input=\"{0}\" -XmlConfigCache=\"{1}\" -remoteini=\"{2}\" {3}", PostBuildSyncFile, XmlConfig.CacheFile, UnrealBuildTool.GetRemoteIniPath(), bNoEntitlements ? "-noEntitlements" : "");
+			Action PostBuildSyncAction = Graph.CreateRecursiveAction<ApplePostBuildSyncMode>(ActionType.CreateAppBundle, PostBuildSyncArguments);
+			PostBuildSyncAction.PrerequisiteItems.Add(Executable);
+			PostBuildSyncAction.ProducedItems.Add(GetPostBuildOutputFile(Executable.Location, Target.Name, Target.Platform));
+			PostBuildSyncAction.StatusDescription = $"Executing PostBuildSync [{Executable.Location}]";
+
+			if (Target.Platform == UnrealTargetPlatform.IOS || Target.Platform == UnrealTargetPlatform.TVOS)
+			{
+				// @todo: do we need one per Target for IOS? Client? I dont think so for Modern
+				FileReference PlistFile = FileReference.Combine(AppleToolChain.GetActualProjectDirectory(Target.ProjectFile), "Build/IOS/UBTGenerated/Info.Template.plist");
+				PostBuildSyncAction.ProducedItems.Add(FileItem.GetItemByFileReference(PlistFile));
+
+				if (PostBuildSync.bCreateStubIPA)
+				{
+					FileReference StubFile = FileReference.Combine(Executable.Directory.Location, Executable.Location.GetFileNameWithoutExtension() + ".stub");
+					PostBuildSyncAction.ProducedItems.Add(FileItem.GetItemByFileReference(StubFile));
+				}
+			}
+
+			return PostBuildSyncAction;
+		}
+	}
+}
diff --git a/Engine/Source/Programs/UnrealBuildTool/Platform/Apple/AppleToolChainSettings.cs b/Engine/Source/Programs/UnrealBuildTool/Platform/Apple/AppleToolChainSettings.cs
new file mode 100644
--- /dev/null
+++ b/Engine/Source/Programs/UnrealBuildTool/Platform/Apple/AppleToolChainSettings.cs
@@ -0,0 +1,186 @@
+// Copyright Epic Games, Inc. All Rights Reserved.
+
+using EpicGames.Core;
+using Microsoft.Extensions.Logging;
+using System;
+using UnrealBuildBase;
+
+namespace UnrealBuildTool
+{
+	/// <summary>
+	/// Helper class for managing Xcode paths, versions, etc. Helps to differentiate between Apple xcode platforms
+	/// </summary>
+	public abstract class AppleToolChainSettings
+	{
+		/// <summary>
+		/// Cached copy of ApplePlatformSDK.GetToolchainDirectory, for existing code to work
+		/// </summary>
+		public DirectoryReference ToolchainDir => ApplePlatformSDK.GetToolchainDirectory();
+
+		/// <summary>
+		/// Name for Xcode plaform directory under Toolchains
+		/// </summary>
+		public string PlatformDirName;
+
+		/// <summary>
+		/// Name for Xcode simulator platform directory under Toolchains
+		/// </summary>
+		public string SimulatorPlatformDirName = "";
+
+		/// <summary>
+		/// A portion of the target "tuple"
+		/// </summary>
+		private string TargetOSName;
+
+		/// <summary>
+		/// The name apple uses for SDK directories
+		/// </summary>
+		private string TargetSDKName;
+
+		/// <summary>
+		/// The version of the SDK being used to build with
+		/// </summary>
+		public string SDKVersion;
+
+		/// <summary>
+		/// The version of the iOS SDK to target at build time.
+		/// </summary>
+		public string MinTargetVersion;
+
+		/// <summary>
+		/// The build version in a floating point value for easy comparison
+		/// </summary>
+		public readonly float SDKVersionFloat;
+
+		/// <summary>
+		/// The target version in a floating point value for easy comparison
+		/// </summary>
+		public readonly float MinTargetVersionFloat;
+
+		/// <summary>
+		/// For platforms that need a different min version for editor vs game, this can be overridden
+		/// </summary>
+		/// <param name="Type"></param>
+		/// <returns></returns>
+		public virtual string GetTargetVersionForTargetType(TargetType Type) => MinTargetVersion;
+
+		/// <summary>
+		/// Cache SDK dir for device and simulator (for non-Mac)
+		/// </summary>
+		DirectoryReference SDKDir;
+		DirectoryReference? SimulatorSDKDir = null;
+
+		/// <summary>
+		/// Dummy UUID for running iOS binary natively on Mac
+		/// </summary>
+		public static readonly string LocalMacUUID = "10CA18AC-10CA18AC10CA18AC";
+
+		/// <summary>
+		/// Constructor, called by platform sybclasses
+		/// </summary>
+		/// <param name="OSPrefix">SDK name (like "MacOSX")</param>
+		/// <param name="SimulatorOSPrefix">Sinulator SDK name (like "iPhoneOSSimulator")</param>
+		/// <param name="TargetOSName">Platform name used in the -target parameter</param>
+		/// <param name="TargetOSVersion">min OS version this project wants to target (not what it's being built with)</param>
+		/// <param name="bVerbose"></param>
+		/// <param name="Logger"></param>
+		protected AppleToolChainSettings(string OSPrefix, string? SimulatorOSPrefix, string TargetOSName, string TargetOSVersion, bool bVerbose, ILogger Logger)
+		{
+			this.TargetOSName = TargetOSName;
+			TargetSDKName = OSPrefix;
+			PlatformDirName = OSPrefix.ToLower();
+
+			// set up directories
+			SDKDir = ApplePlatformSDK.GetPlatformSDKDirectory(OSPrefix);
+			if (SimulatorOSPrefix != null)
+			{
+				SimulatorSDKDir = ApplePlatformSDK.GetPlatformSDKDirectory(SimulatorOSPrefix);
+				SimulatorPlatformDirName = SimulatorOSPrefix.ToLower();
+			}
+
+			MinTargetVersion = TargetOSVersion;
+			MinTargetVersionFloat = Single.Parse(MinTargetVersion, System.Globalization.CultureInfo.InvariantCulture);
+
+			// cache some info for this OS
+			SDKVersion = ApplePlatformSDK.GetPlatformSDKVersion(OSPrefix) ?? "";
+			SDKVersionFloat = ApplePlatformSDK.GetPlatformSDKVersionFloat(OSPrefix);
+
+			TestXcode(bVerbose, Logger);
+		}
+
+		/// <summary>
+		/// Get the path to the SDK diretory in xcode for the given architecture
+		/// </summary>
+		/// <param name="Architecture"></param>
+		/// <returns></returns>
+		public DirectoryReference GetSDKPath(UnrealArch Architecture)
+		{
+			// note that VisionOS uses IOSSimulator (as TVOS should eventually do as well)
+			if (Architecture == UnrealArch.IOSSimulator || Architecture == UnrealArch.TVOSSimulator)
+			{
+				return SimulatorSDKDir!;
+			}
+			return SDKDir;
+		}
+
+		/// <summary>
+		/// Gets the string used by xcode to taget a platform and version (will return something like "arm64-apple-ios17.0-simulator"
+		/// </summary>
+		/// <param name="Architecture"></param>
+		/// <param name="Platform"></param>
+		/// <param name="TargetType"></param>
+		/// <param name="ForcedVersion"></param>
+		/// <returns></returns>
+		public virtual string GetTargetTuple(UnrealArch Architecture, UnrealTargetPlatform Platform, TargetType TargetType, string? ForcedVersion = null)
+		{
+			string Prefix = Architecture.AppleName;
+			string Suffix = (Architecture == UnrealArch.IOSSimulator || Architecture == UnrealArch.TVOSSimulator) ? "-simulator" : "";
+			string TargetVersion = ForcedVersion ?? GetTargetVersionForTargetType(TargetType);
+
+			return $"{Prefix}-apple-{TargetOSName}{TargetVersion}{Suffix}";
+		}
+
+		/// <summary>
+		/// Find the Xcode developer directory
+		/// </summary>
+		/// <param name="bVerbose"></param>
+		/// <param name="Logger"></param>
+		/// <exception cref="BuildException"></exception>
+		private static void TestXcode(bool bVerbose, ILogger Logger)
+		{
+			DirectoryReference XcodeDeveloperDir = ApplePlatformSDK.DeveloperDir;
+			// make sure we get a full path
+			if (DirectoryReference.Exists(XcodeDeveloperDir) == false)
+			{
+				throw new BuildException("Selected Xcode ('{0}') doesn't exist, cannot continue.", XcodeDeveloperDir);
+			}
+
+			if (XcodeDeveloperDir.ContainsName("CommandLineTools", 0))
+			{
+				throw new BuildException($"Your Mac is set to use CommandLineTools for its build tools ({XcodeDeveloperDir}). Unreal expects Xcode as the build tools. Please install Xcode if it's not already, then do one of the following:\n" +
+					"  - Run Xcode, go to Settings, and in the Locations tab, choose your Xcode in Command Line Tools dropdown.\n" +
+					"  - In Terminal, run 'sudo xcode-select -s /Applications/Xcode.app' (or an alternate location if you installed Xcode to a non-standard location)\n" +
+					"Either way, you will need to enter your Mac password.");
+			}
+
+			if (bVerbose && !XcodeDeveloperDir.FullName.StartsWith("/Applications/Xcode.app"))
+			{
+				Log.TraceInformationOnce("Compiling with non-standard Xcode: {0}", XcodeDeveloperDir);
+			}
+
+			// Installed engine requires Xcode 13
+			if (Unreal.IsEngineInstalled())
+			{
+				string? InstalledSdkVersion = ApplePlatformSDK.InstalledXcodeVersion.Value;
+				if (String.IsNullOrEmpty(InstalledSdkVersion))
+				{
+					throw new BuildException("Unable to get xcode version");
+				}
+				if (Int32.Parse(InstalledSdkVersion.Substring(0, 2)) < 13)
+				{
+					throw new BuildException("Building for macOS, iOS and tvOS requires Xcode 13.4.1 or newer, Xcode " + InstalledSdkVersion + " detected");
+				}
+			}
+		}
+	}
+}
diff --git a/Engine/Source/Programs/UnrealBuildTool/Platform/Apple/RemoteMac.cs b/Engine/Source/Programs/UnrealBuildTool/Platform/Apple/RemoteMac.cs
new file mode 100644
--- /dev/null
+++ b/Engine/Source/Programs/UnrealBuildTool/Platform/Apple/RemoteMac.cs
@@ -0,0 +1,1409 @@
+﻿// Copyright Epic Games, Inc. All Rights Reserved.
+
+using System;
+using System.Collections.Generic;
+using System.Diagnostics;
+using System.IO;
+using System.Linq;
+using System.Text;
+using System.Text.RegularExpressions;
+using EpicGames.Core;
+using Microsoft.Extensions.Logging;
+using UnrealBuildBase;
+
+namespace UnrealBuildTool
+{
+	/// <summary>
+	/// Stores information about how a local directory maps to a remote directory
+	/// </summary>
+	[DebuggerDisplay("{LocalDirectory}")]
+	class RemoteMapping
+	{
+		public DirectoryReference LocalDirectory;
+		public string RemoteDirectory;
+
+		public RemoteMapping(DirectoryReference LocalDirectory, string RemoteDirectory)
+		{
+			this.LocalDirectory = LocalDirectory;
+			this.RemoteDirectory = RemoteDirectory;
+		}
+	}
+
+	/// <summary>
+	/// Handles uploading and building on a remote Mac
+	/// </summary>
+	class RemoteMac
+	{
+		/// <summary>
+		/// These two variables will be loaded from the XML config file in XmlConfigLoader.Init().
+		/// </summary>
+		[XmlConfigFile]
+		private readonly string? ServerName;
+
+		/// <summary>
+		/// The remote username.
+		/// </summary>
+		[XmlConfigFile]
+		private readonly string? UserName;
+
+		/// <summary>
+		/// If set, instead of looking for RemoteToolChainPrivate.key in the usual places (Documents/Unreal, Engine/UnrealBuildTool/SSHKeys or Engine/Build/SSHKeys), this private key will be used.
+		/// </summary>
+		[XmlConfigFile]
+		private FileReference? SshPrivateKey;
+
+		/// <summary>
+		/// The authentication used for Rsync (for the -e rsync flag).
+		/// </summary>
+		[XmlConfigFile]
+		private string RsyncAuthentication = "./ssh -i '${CYGWIN_SSH_PRIVATE_KEY}'";
+
+		/// <summary>
+		/// The authentication used for SSH (probably similar to RsyncAuthentication).
+		/// </summary>
+		[XmlConfigFile]
+		private string SshAuthentication = "-i '${CYGWIN_SSH_PRIVATE_KEY}'";
+
+		/// <summary>
+		/// Save the specified port so that RemoteServerName is the machine address only
+		/// </summary>
+		private readonly int ServerPort = 22;   // Default ssh port
+
+		/// <summary>
+		/// Path to Rsync
+		/// </summary>
+		private readonly FileReference RsyncExe;
+
+		/// <summary>
+		/// Path to SSH
+		/// </summary>
+		private readonly FileReference SshExe;
+
+		/// <summary>
+		/// The project being built. Settings will be read from config files in this project.
+		/// </summary>
+		private readonly FileReference? ProjectFile;
+
+		/// <summary>
+		/// The project descriptor for the project being built.
+		/// </summary>
+		private readonly ProjectDescriptor? ProjectDescriptor;
+
+		/// <summary>
+		/// A set of directories containing additional paths to be built.
+		/// </summary>
+		private readonly List<DirectoryReference>? AdditionalPaths;
+
+		/// <summary>
+		/// The base directory on the remote machine
+		/// </summary>
+		private string RemoteBaseDir;
+
+		/// <summary>
+		/// Mappings from local directories to remote directories
+		/// </summary>
+		private List<RemoteMapping> Mappings;
+
+		/// <summary>
+		/// Arguments that are used by every Ssh call
+		/// </summary>
+		private List<string> CommonSshArguments;
+
+		/// <summary>
+		/// Arguments that are used by every Rsync call
+		/// </summary>
+		private List<string> BasicRsyncArguments;
+
+		/// <summary>
+		/// Arguments that are used by directory Rsync call
+		/// </summary>
+		private List<string> CommonRsyncArguments;
+
+		private List<string> RsyncArgumentsForDownload(List<string> BaseArguments)
+		{
+			List<string> Arguments = new List<string>(BaseArguments);
+			Arguments.Add(String.Format("--rsync-path=\"{0}/rsync\"", RemoteBaseDir));
+			return Arguments;
+		}
+
+		private List<string> RsyncArgumentsForUpload(List<string> BaseArguments, string RemoteDirectory)
+		{
+			List<string> Arguments = new List<string>(BaseArguments);
+			Arguments.Add(String.Format("--rsync-path=\"mkdir -p {0} && {1}/rsync\"", RemoteDirectory, RemoteBaseDir));
+			return Arguments;
+		}
+
+		private string? IniBundleIdentifier = "";
+
+		/// <summary>
+		/// Constructor
+		/// </summary>
+		/// <param name="ProjectFile">Project to read settings from</param>
+		/// <param name="Logger">Logger for output</param>
+		/// <param name="RemoteIniPath">Remote ini path, or null to use UnrealBuildTool.GetRemoteIniPath()</param>
+		/// <param name="bIsPrimary">Is the primary Remote Mac or not. True by default.</param>
+		/// <param name="bPrepareForSecondaryMac">Added arguments when preparing a secondary Mac for debug. False by default.</param>
+		public RemoteMac(FileReference? ProjectFile, ILogger Logger, string? RemoteIniPath = null, bool bIsPrimary = true, bool bPrepareForSecondaryMac = false)
+		{
+			DirectoryReference Cwrsync = DirectoryReference.Combine(Unreal.EngineDirectory, "Extras", "ThirdPartyNotUE", "cwrsync");
+			RsyncExe = FileReference.Combine(Cwrsync, "bin", "rsync.exe");
+			SshExe = FileReference.Combine(Cwrsync, "bin", "ssh.exe");
+			this.ProjectFile = ProjectFile;
+			if (ProjectFile != null)
+			{
+				ProjectDescriptor = ProjectDescriptor.FromFile(ProjectFile);
+				AdditionalPaths = new List<DirectoryReference>();
+				ProjectDescriptor.AddAdditionalPaths(AdditionalPaths, ProjectFile.Directory);
+				if (AdditionalPaths.Count == 0)
+				{
+					AdditionalPaths = null;
+				}
+			}
+
+			// Apply settings from the XML file
+			XmlConfig.ApplyTo(this);
+
+			// Get the project config file path
+			DirectoryReference? EngineIniPath = ProjectFile?.Directory;
+			if (RemoteIniPath == null)
+			{
+				RemoteIniPath = UnrealBuildTool.GetRemoteIniPath();
+			}
+			if (EngineIniPath == null && RemoteIniPath != null)
+			{
+				EngineIniPath = new DirectoryReference(RemoteIniPath!);
+			}
+			ConfigHierarchy Ini = ConfigCache.ReadHierarchy(ConfigHierarchyType.Engine, EngineIniPath, UnrealTargetPlatform.IOS);
+
+			// Read the project settings if we don't have anything in the build configuration settings
+			if (String.IsNullOrEmpty(ServerName))
+			{
+				// Read the server name
+				string IniServerName;
+
+				if ((bIsPrimary ? Ini.GetString("/Script/IOSRuntimeSettings.IOSRuntimeSettings", "RemoteServerName", out IniServerName) : Ini.GetString("/Script/IOSRuntimeSettings.IOSRuntimeSettings", "SecondaryRemoteServerName", out IniServerName)) && !String.IsNullOrEmpty(IniServerName))
+				{
+					ServerName = IniServerName;
+				}
+				else
+				{
+					throw new BuildException("Remote compiling requires a server name. Use the editor (Project Settings > IOS) to set up your remote compilation settings.");
+				}
+
+				// Parse the username
+				string IniUserName;
+				if ((bIsPrimary ? Ini.GetString("/Script/IOSRuntimeSettings.IOSRuntimeSettings", "RSyncUsername", out IniUserName) : Ini.GetString("/Script/IOSRuntimeSettings.IOSRuntimeSettings", "SecondaryRSyncUsername", out IniUserName)) && !String.IsNullOrEmpty(IniUserName))
+				{
+					UserName = IniUserName;
+				}
+			}
+
+			// Split port out from the server name
+			int PortIdx = ServerName.LastIndexOf(':');
+			if (PortIdx != -1)
+			{
+				string Port = ServerName[(PortIdx + 1)..];
+				if (!Int32.TryParse(Port, out ServerPort))
+				{
+					throw new BuildException("Unable to parse port number from '{0}'", ServerName);
+				}
+				ServerName = ServerName[..PortIdx];
+			}
+
+			// If a user name is not set, use the current user
+			if (String.IsNullOrEmpty(UserName))
+			{
+				UserName = Environment.UserName;
+			}
+
+			// Print out the server info
+			Logger.LogInformation("[Remote] Using remote server '{ServerName}' on port {ServerPort} (user '{UserName}')", ServerName, ServerPort, UserName);
+
+			// Get the path to the SSH private key
+			string OverrideSshPrivateKeyPath;
+			if ((bIsPrimary ? Ini.GetString("/Script/IOSRuntimeSettings.IOSRuntimeSettings", "SSHPrivateKeyOverridePath", out OverrideSshPrivateKeyPath) : Ini.GetString("/Script/IOSRuntimeSettings.IOSRuntimeSettings", "SecondarySSHPrivateKeyOverridePath", out OverrideSshPrivateKeyPath)) && !String.IsNullOrEmpty(OverrideSshPrivateKeyPath))
+			{
+				SshPrivateKey = new FileReference(OverrideSshPrivateKeyPath);
+				if (!FileReference.Exists(SshPrivateKey))
+				{
+					throw new BuildException("SSH private key specified in config file ({0}) does not exist.", SshPrivateKey);
+				}
+			}
+
+			Ini.GetString("/Script/IOSRuntimeSettings.IOSRuntimeSettings", "BundleIdentifier", out IniBundleIdentifier);
+
+			// If it's not set, look in the standard locations. If that fails, spawn the batch file to generate one.
+			if (SshPrivateKey == null && !TryGetSshPrivateKey(out SshPrivateKey))
+			{
+				Logger.LogWarning("No SSH private key found for {UserName}@{ServerName}. Launching SSH to generate one.", UserName, ServerName);
+
+				StringBuilder CommandLine = new StringBuilder();
+				CommandLine.AppendFormat("/C \"\"{0}\"", FileReference.Combine(Unreal.EngineDirectory, "Build", "BatchFiles", "MakeAndInstallSSHKey.bat"));
+				CommandLine.AppendFormat(" \"{0}\"", SshExe);
+				CommandLine.AppendFormat(" \"{0}\"", ServerPort);
+				CommandLine.AppendFormat(" \"{0}\"", RsyncExe);
+				CommandLine.AppendFormat(" \"{0}\"", UserName);
+				CommandLine.AppendFormat(" \"{0}\"", ServerName);
+				CommandLine.AppendFormat(" \"{0}\"", DirectoryReference.GetSpecialFolder(Environment.SpecialFolder.MyDocuments));
+				CommandLine.AppendFormat(" \"{0}\"", GetLocalCygwinPath(DirectoryReference.GetSpecialFolder(Environment.SpecialFolder.MyDocuments)!));
+				CommandLine.AppendFormat(" \"{0}\"", Unreal.EngineDirectory);
+				CommandLine.Append('"');
+
+				using (Process ChildProcess = Process.Start(BuildHostPlatform.Current.Shell.FullName, CommandLine.ToString()))
+				{
+					ChildProcess.WaitForExit();
+				}
+
+				if (!TryGetSshPrivateKey(out SshPrivateKey))
+				{
+					throw new BuildException("Failed to generate SSH private key for {0}@{1}.", UserName, ServerName);
+				}
+			}
+
+			// Print the path to the private key
+			Logger.LogInformation("[Remote] Using private key at {SshPrivateKey}", SshPrivateKey);
+
+			// resolve the rest of the strings
+			RsyncAuthentication = ExpandVariables(RsyncAuthentication);
+			SshAuthentication = ExpandVariables(SshAuthentication);
+
+			// Build a list of arguments for SSH
+			CommonSshArguments = new List<string>();
+			CommonSshArguments.Add("-o BatchMode=yes");
+			CommonSshArguments.Add(SshAuthentication);
+			CommonSshArguments.Add(String.Format("-p {0}", ServerPort));
+			CommonSshArguments.Add(String.Format("\"{0}@{1}\"", UserName, ServerName));
+
+			// Build a list of arguments for Rsync
+			BasicRsyncArguments = new List<string>();
+			BasicRsyncArguments.Add("--compress");
+			BasicRsyncArguments.Add("--verbose");
+			BasicRsyncArguments.Add(String.Format("--rsh=\"{0} -p {1}\"", RsyncAuthentication, ServerPort));
+			BasicRsyncArguments.Add("--chmod=ugo=rwx");
+
+			// Build a list of arguments for Rsync filters
+			CommonRsyncArguments = new List<string>(BasicRsyncArguments);
+			CommonRsyncArguments.Add("--copy-links");
+			CommonRsyncArguments.Add("--recursive");
+			if (!bPrepareForSecondaryMac)
+			{
+				CommonRsyncArguments.Add("--delete"); // Delete anything not in the source directory
+				CommonRsyncArguments.Add("--delete-excluded"); // Delete anything not in the source directory
+			}
+			CommonRsyncArguments.Add("--times"); // Preserve modification times
+			CommonRsyncArguments.Add("--omit-dir-times"); // Ignore modification times for directories
+			CommonRsyncArguments.Add("--prune-empty-dirs"); // Remove empty directories from the file list
+
+			// Get the remote base directory
+			string RemoteServerOverrideBuildPath;
+			if ((bIsPrimary ? Ini.GetString("/Script/IOSRuntimeSettings.IOSRuntimeSettings", "RemoteServerOverrideBuildPath", out RemoteServerOverrideBuildPath) : Ini.GetString("/Script/IOSRuntimeSettings.IOSRuntimeSettings", "SecondaryRemoteServerOverrideBuildPath", out RemoteServerOverrideBuildPath)) && !String.IsNullOrEmpty(RemoteServerOverrideBuildPath))
+			{
+				RemoteBaseDir = String.Format("{0}/{1}", RemoteServerOverrideBuildPath.Trim().TrimEnd('/'), Unreal.MachineName);
+			}
+			else
+			{
+				StringBuilder Output;
+				if (ExecuteAndCaptureOutput("'echo ~'", Logger, out Output) != 0)
+				{
+					throw new BuildException("Unable to determine home directory for remote user. SSH output:\n{0}", StringUtils.Indent(Output.ToString(), "  "));
+				}
+				RemoteBaseDir = String.Format("{0}/UE5/Builds/{1}", Output.ToString().Trim().TrimEnd('/'), Unreal.MachineName);
+			}
+
+			Logger.LogInformation("[Remote] Using base directory '{RemoteBaseDir}'", RemoteBaseDir);
+
+			// Build the list of directory mappings between the local and remote machines
+			Mappings = new List<RemoteMapping>();
+			Mappings.Add(new RemoteMapping(Unreal.EngineDirectory, GetRemotePath(Unreal.EngineDirectory)));
+			if (ProjectFile != null && !ProjectFile.IsUnderDirectory(Unreal.EngineDirectory))
+			{
+				Mappings.Add(new RemoteMapping(ProjectFile.Directory, GetRemotePath(ProjectFile.Directory)));
+			}
+			if (AdditionalPaths != null && ProjectFile != null)
+			{
+				foreach (DirectoryReference AdditionalPath in AdditionalPaths)
+				{
+					if (!AdditionalPath.IsUnderDirectory(Unreal.EngineDirectory) &&
+						!AdditionalPath.IsUnderDirectory(ProjectFile.Directory))
+					{
+						Mappings.Add(new RemoteMapping(AdditionalPath, GetRemotePath(AdditionalPath)));
+					}
+				}
+			}
+
+			string RemoteFile = GetRemotePath("rsync");
+
+			List<string> Arguments = new List<string>(CommonRsyncArguments);
+			// use MacOS default system rsync for the upload of the 2.6.9 rsync binary
+			Arguments.Add(String.Format("--rsync-path=\"mkdir -p {0} && rsync\"", RemoteBaseDir));
+			Arguments.Add(String.Format("\"{0}\"", GetLocalCygwinPath(FileReference.Combine(Cwrsync, "mac", "rsync"))));
+			Arguments.Add(String.Format("\"{0}@{1}\":'{2}'", UserName, ServerName, RemoteFile));
+			Arguments.Add("-q");
+
+			int Result = Rsync(String.Join(" ", Arguments), Logger);
+			if (Result != 0)
+			{
+				throw new BuildException("Error while running Rsync (exit code {0})", Result);
+			}
+			// Remove quaratine extended attribute if remote machine set it.
+			Execute(RemoteBaseDir, "if xattr -p com.apple.quarantine rsync > /dev/null 2>&1; then xattr -d com.apple.quarantine rsync; fi", Logger);
+		}
+
+		/// <summary>
+		/// Attempts to get the SSH private key from the standard locations
+		/// </summary>
+		/// <param name="OutPrivateKey">If successful, receives the location of the private key that was found</param>
+		/// <returns>True if a private key was found, false otherwise</returns>
+		private bool TryGetSshPrivateKey(out FileReference? OutPrivateKey)
+		{
+			// Build a list of all the places to look for a private key
+			List<DirectoryReference> Locations = new List<DirectoryReference>();
+			Locations.Add(DirectoryReference.Combine(DirectoryReference.GetSpecialFolder(Environment.SpecialFolder.ApplicationData)!, "Unreal Engine", "UnrealBuildTool"));
+			Locations.Add(DirectoryReference.Combine(DirectoryReference.GetSpecialFolder(Environment.SpecialFolder.Personal)!, "Unreal Engine", "UnrealBuildTool"));
+			if (ProjectFile != null)
+			{
+				Locations.Add(DirectoryReference.Combine(ProjectFile.Directory, "Restricted", "NotForLicensees", "Build"));
+				Locations.Add(DirectoryReference.Combine(ProjectFile.Directory, "Restricted", "NoRedist", "Build"));
+				Locations.Add(DirectoryReference.Combine(ProjectFile.Directory, "Build"));
+			}
+			Locations.Add(DirectoryReference.Combine(Unreal.EngineDirectory, "Restricted", "NotForLicensees", "Build"));
+			Locations.Add(DirectoryReference.Combine(Unreal.EngineDirectory, "Restricted", "NoRedist", "Build"));
+			Locations.Add(DirectoryReference.Combine(Unreal.EngineDirectory, "Build"));
+
+			// Find the first that exists
+			foreach (DirectoryReference Location in Locations)
+			{
+				FileReference KeyFile = FileReference.Combine(Location, "SSHKeys", ServerName!, UserName!, "RemoteToolChainPrivate.key");
+				if (FileReference.Exists(KeyFile))
+				{
+					// MacOS Mojave includes a new version of SSH that generates keys that are incompatible with our version of SSH. Make sure the detected keys have the right signature.
+					string Text = FileReference.ReadAllText(KeyFile);
+					if (Text.Contains("---BEGIN RSA PRIVATE KEY---"))
+					{
+						OutPrivateKey = KeyFile;
+						return true;
+					}
+				}
+			}
+
+			// Nothing found
+			OutPrivateKey = null;
+			return false;
+		}
+
+		/// <summary>
+		/// Expand all the variables in the given string
+		/// </summary>
+		/// <param name="Input">The input string</param>
+		/// <returns>String with any variables expanded</returns>
+		private string ExpandVariables(string Input)
+		{
+			string Result = Input;
+			Result = Result.Replace("${SSH_PRIVATE_KEY}", SshPrivateKey!.FullName);
+			Result = Result.Replace("${CYGWIN_SSH_PRIVATE_KEY}", GetLocalCygwinPath(SshPrivateKey));
+			return Result;
+		}
+
+		/// <summary>
+		/// Flush the remote machine, removing all existing files
+		/// </summary>
+		public void FlushRemote(ILogger Logger)
+		{
+			Logger.LogInformation("[Remote] Deleting all files under {RemoteBaseDir}...", RemoteBaseDir);
+			Execute("/", String.Format("rm -rf \"{0}\"", RemoteBaseDir), Logger);
+		}
+
+		/// <summary>
+		/// Returns true if the remote executor supports this target platform
+		/// </summary>
+		/// <param name="Platform">The platform to check</param>
+		/// <returns>True if the remote mac handles this target platform</returns>
+		public static bool HandlesTargetPlatform(UnrealTargetPlatform Platform)
+		{
+			return BuildHostPlatform.Current.Platform == UnrealTargetPlatform.Win64 && (Platform == UnrealTargetPlatform.Mac || Platform == UnrealTargetPlatform.IOS || Platform == UnrealTargetPlatform.TVOS);
+		}
+
+		/// <summary>
+		/// Clean a target remotely
+		/// </summary>
+		/// <param name="TargetDesc">Descriptor for the target to build</param>
+		/// <param name="Logger">Logger for diagnostic output</param>
+		/// <returns>True if the build succeeded, false otherwise</returns>
+		public bool Clean(TargetDescriptor TargetDesc, ILogger Logger)
+		{
+			// Translate all the arguments for the remote
+			List<string> RemoteArguments = GetRemoteArgumentsForTarget(TargetDesc, null);
+			RemoteArguments.Add("-Clean");
+
+			// Upload the workspace
+			DirectoryReference TempDir = CreateTempDirectory(TargetDesc);
+			UploadWorkspace(TempDir, Logger);
+
+			// Execute the compile
+			Logger.LogInformation("[Remote] Executing clean...");
+
+			StringBuilder BuildCommandLine = new StringBuilder("Engine/Build/BatchFiles/Mac/Build.sh");
+			foreach (string RemoteArgument in RemoteArguments)
+			{
+				BuildCommandLine.AppendFormat(" {0}", EscapeShellArgument(RemoteArgument));
+			}
+
+			int Result = Execute(GetRemotePath(Unreal.RootDirectory), BuildCommandLine.ToString(), Logger);
+			return Result == 0;
+		}
+
+		/// <summary>
+		/// Build a target remotely
+		/// </summary>
+		/// <param name="TargetDesc">Descriptor for the target to build</param>
+		/// <param name="RemoteLogFile">Path to store the remote log file</param>
+		/// <param name="bSkipPreBuildTargets">If true then any PreBuildTargets will be skipped</param>
+		/// <param name="Logger">Logger for diagnostic output</param>
+		/// <returns>True if the build succeeded, false otherwise</returns>
+		public bool Build(TargetDescriptor TargetDesc, FileReference RemoteLogFile, bool bSkipPreBuildTargets, ILogger Logger)
+		{
+			// Compile the rules assembly
+			RulesAssembly RulesAssembly = RulesCompiler.CreateTargetRulesAssembly(TargetDesc.ProjectFile, TargetDesc.Name, false, false, false, TargetDesc.ForeignPlugin, TargetDesc.bBuildPluginAsLocal, Logger);
+
+			// Create the target rules
+			TargetRules Rules = RulesAssembly.CreateTargetRules(TargetDesc.Name, TargetDesc.Platform, TargetDesc.Configuration, TargetDesc.Architectures, TargetDesc.ProjectFile, TargetDesc.AdditionalArguments, Logger);
+			if (!bSkipPreBuildTargets)
+			{
+				foreach (TargetInfo PreBuildTargetInfo in Rules.PreBuildTargets)
+				{
+					RemoteMac PreBuildTargetRemoteMac = new RemoteMac(ProjectFile, Logger);
+					TargetDescriptor PreBuildTargetDesc = new TargetDescriptor(PreBuildTargetInfo.ProjectFile, PreBuildTargetInfo.Name, PreBuildTargetInfo.Platform, PreBuildTargetInfo.Configuration, PreBuildTargetInfo.Architectures, PreBuildTargetInfo.Arguments);
+
+					Logger.LogInformation("[Remote] Building pre target [{PreTarget}] for [{Target}] ", PreBuildTargetDesc.ToString(), TargetDesc.ToString());
+					if (!PreBuildTargetRemoteMac.Build(PreBuildTargetDesc, RemoteLogFile, false, Logger))
+					{
+						return false;
+					}
+				}
+			}
+
+			// Get the directory for working files
+			DirectoryReference TempDir = CreateTempDirectory(TargetDesc);
+
+			// Map the path containing the remote log file
+			bool bLogIsMapped = false;
+			foreach (RemoteMapping Mapping in Mappings)
+			{
+				if (RemoteLogFile.Directory.FullName.Equals(Mapping.LocalDirectory.FullName, StringComparison.OrdinalIgnoreCase))
+				{
+					bLogIsMapped = true;
+					break;
+				}
+			}
+			if (!bLogIsMapped)
+			{
+				Mappings.Add(new RemoteMapping(RemoteLogFile.Directory, GetRemotePath(RemoteLogFile.Directory)));
+			}
+
+			// Path to the local manifest file. This has to be translated from the remote format after the build is complete.
+			List<FileReference> LocalManifestFiles = new List<FileReference>();
+
+			// Path to the remote manifest file
+			FileReference RemoteManifestFile = FileReference.Combine(TempDir, "Manifest.xml");
+
+			// Prepare the arguments we will pass to the remote build
+			List<string> RemoteArguments = GetRemoteArgumentsForTarget(TargetDesc, LocalManifestFiles);
+			RemoteArguments.Add(String.Format("-Log={0}", GetRemotePath(RemoteLogFile)));
+			RemoteArguments.Add(String.Format("-Manifest={0}", GetRemotePath(RemoteManifestFile)));
+			RemoteArguments.Add(String.Format("-SkipPreBuildTargets"));
+
+			// Handle any per-platform setup that is required
+			if (TargetDesc.Platform == UnrealTargetPlatform.IOS || TargetDesc.Platform == UnrealTargetPlatform.TVOS)
+			{
+				// Always generate a .stub
+				RemoteArguments.Add("-CreateStub");
+
+				// Cannot use makefiles, since we need PostBuildSync() to generate the IPA (and that requires a TargetRules instance)
+				RemoteArguments.Add("-NoUBTMakefiles");
+
+				// Get the provisioning data for this project
+				IOSProvisioningData ProvisioningData = ((IOSPlatform)UEBuildPlatform.GetBuildPlatform(TargetDesc.Platform)).ReadProvisioningData(TargetDesc.ProjectFile, TargetDesc.AdditionalArguments.HasOption("-distribution"), IniBundleIdentifier);
+				if (ProvisioningData == null || ProvisioningData.MobileProvisionFile == null)
+				{
+					throw new BuildException("Unable to find mobile provision for {0}. See log for more information.", TargetDesc.Name);
+				}
+
+				// Create a local copy of the provision
+				FileReference MobileProvisionFile = FileReference.Combine(TempDir, ProvisioningData.MobileProvisionFile.GetFileName());
+				if (FileReference.Exists(MobileProvisionFile))
+				{
+					FileReference.SetAttributes(MobileProvisionFile, FileAttributes.Normal);
+				}
+				FileReference.Copy(ProvisioningData.MobileProvisionFile, MobileProvisionFile, true);
+				Logger.LogInformation("[Remote] Uploading {MobileProvisionFile}", MobileProvisionFile);
+				UploadFile(MobileProvisionFile, Logger);
+
+				// Extract the certificate for the project. Try to avoid calling IPP if we already have it.
+				FileReference CertificateFile = FileReference.Combine(TempDir, "Certificate.p12");
+
+				FileReference CertificateInfoFile = FileReference.Combine(TempDir, "Certificate.txt");
+				string CertificateInfoContents = String.Format("{0}\n{1}", ProvisioningData.MobileProvisionFile, FileReference.GetLastWriteTimeUtc(ProvisioningData.MobileProvisionFile).Ticks);
+
+				if (!FileReference.Exists(CertificateFile) || !FileReference.Exists(CertificateInfoFile) || FileReference.ReadAllText(CertificateInfoFile) != CertificateInfoContents)
+				{
+					Logger.LogInformation("[Remote] Exporting certificate for {ProvisioningDataMobileProvisionFile}...", ProvisioningData.MobileProvisionFile);
+
+					StringBuilder Arguments = new StringBuilder("ExportCertificate");
+					if (TargetDesc.ProjectFile == null)
+					{
+						Arguments.AppendFormat(" \"{0}\"", Unreal.EngineSourceDirectory);
+					}
+					else
+					{
+						Arguments.AppendFormat(" \"{0}\"", TargetDesc.ProjectFile.Directory);
+					}
+					Arguments.AppendFormat(" -provisionfile \"{0}\"", ProvisioningData.MobileProvisionFile);
+					Arguments.AppendFormat(" -outputcertificate \"{0}\"", CertificateFile);
+					if (TargetDesc.Platform == UnrealTargetPlatform.TVOS)
+					{
+						Arguments.Append(" -tvos");
+					}
+
+					ProcessStartInfo StartInfo = new ProcessStartInfo();
+					StartInfo.FileName = FileReference.Combine(Unreal.EngineDirectory, "Binaries", "DotNET", "IOS", "IPhonePackager.exe").FullName;
+					StartInfo.Arguments = Arguments.ToString();
+					if (Utils.RunLocalProcessAndLogOutput(StartInfo, Logger) != 0)
+					{
+						throw new BuildException("IphonePackager failed.");
+					}
+
+					FileReference.WriteAllText(CertificateInfoFile, CertificateInfoContents);
+				}
+
+				// Upload the certificate to the remote
+				Logger.LogInformation("[Remote] Uploading {CertificateFile}", CertificateFile);
+				UploadFile(CertificateFile, Logger);
+
+				// Tell the remote UBT instance to use them
+				RemoteArguments.Add(String.Format("-ImportProvision={0}", GetRemotePath(MobileProvisionFile)));
+				RemoteArguments.Add(String.Format("-ImportCertificate={0}", GetRemotePath(CertificateFile)));
+				RemoteArguments.Add(String.Format("-ImportCertificatePassword=A"));
+			}
+
+			// Upload the workspace files
+			UploadWorkspace(TempDir, Logger);
+
+			// Execute the compile
+			Logger.LogInformation("[Remote] Executing build");
+
+			StringBuilder BuildCommandLine = new StringBuilder("Engine/Build/BatchFiles/Mac/Build.sh");
+			foreach (string RemoteArgument in RemoteArguments)
+			{
+				BuildCommandLine.AppendFormat(" {0}", EscapeShellArgument(RemoteArgument));
+			}
+
+			int Result = Execute(GetRemotePath(Unreal.RootDirectory), BuildCommandLine.ToString(), Logger);
+			if (Result != 0)
+			{
+				if (RemoteLogFile != null)
+				{
+					Logger.LogInformation("[Remote] Downloading {RemoteLogFile}", RemoteLogFile);
+					DownloadFile(RemoteLogFile, Logger);
+				}
+				return false;
+			}
+
+			// Download the manifest
+			Logger.LogInformation("[Remote] Downloading {RemoteManifestFile}", RemoteManifestFile);
+			DownloadFile(RemoteManifestFile, Logger);
+
+			// Convert the manifest to local form
+			BuildManifest Manifest = Utils.ReadClass<BuildManifest>(RemoteManifestFile.FullName, Logger);
+			for (int Idx = 0; Idx < Manifest.BuildProducts.Count; Idx++)
+			{
+				Manifest.BuildProducts[Idx] = GetLocalPath(Manifest.BuildProducts[Idx]).FullName;
+			}
+
+			// Download the files from the remote
+			Logger.LogInformation("[Remote] Downloading build products");
+
+			List<FileReference> FilesToDownload = new List<FileReference>();
+			FilesToDownload.Add(RemoteLogFile);
+			FilesToDownload.AddRange(Manifest.BuildProducts.Select(x => new FileReference(x)));
+			DownloadFiles(FilesToDownload, Logger);
+
+			// Copy remote FrameworkAssets directory as it could contain resource bundles that must be packaged locally.
+			DirectoryReference BaseDir = DirectoryReference.FromFile(TargetDesc.ProjectFile) ?? Unreal.EngineDirectory;
+			DirectoryReference FrameworkAssetsDir = DirectoryReference.Combine(BaseDir, "Intermediate", TargetDesc.Platform == UnrealTargetPlatform.IOS ? "IOS" : "TVOS", "FrameworkAssets");
+			if (RemoteDirectoryExists(FrameworkAssetsDir, Logger))
+			{
+				Logger.LogInformation("[Remote] Downloading {FrameworkAssetsDir}", FrameworkAssetsDir);
+				DownloadDirectory(FrameworkAssetsDir, Logger);
+			}
+
+			// Write out all the local manifests
+			foreach (FileReference LocalManifestFile in LocalManifestFiles)
+			{
+				Logger.LogInformation("[Remote] Writing {LocalManifestFile}", LocalManifestFile);
+				Utils.WriteClass<BuildManifest>(Manifest, LocalManifestFile.FullName, "", Logger);
+			}
+
+			return true;
+		}
+
+		/// <summary>
+		/// Creates a temporary directory for the given target
+		/// </summary>
+		/// <param name="TargetDesc">The target descriptor</param>
+		/// <returns>Directory to use for temporary files</returns>
+		static DirectoryReference CreateTempDirectory(TargetDescriptor TargetDesc)
+		{
+			DirectoryReference BaseDir = DirectoryReference.FromFile(TargetDesc.ProjectFile) ?? Unreal.EngineDirectory;
+			DirectoryReference TempDir = DirectoryReference.Combine(BaseDir, "Intermediate", "Remote", TargetDesc.Name, TargetDesc.Platform.ToString(), TargetDesc.Configuration.ToString());
+			DirectoryReference.CreateDirectory(TempDir);
+			return TempDir;
+		}
+
+		/// <summary>
+		/// Translate the arguments for a target descriptor for the remote machine
+		/// </summary>
+		/// <param name="TargetDesc">The target descriptor</param>
+		/// <param name="LocalManifestFiles">Manifest files to be output from this target</param>
+		/// <return>List of remote arguments</return>
+		List<string> GetRemoteArgumentsForTarget(TargetDescriptor TargetDesc, List<FileReference>? LocalManifestFiles)
+		{
+			List<string> RemoteArguments = new List<string>();
+			RemoteArguments.Add(TargetDesc.Name);
+			RemoteArguments.Add(TargetDesc.Platform.ToString());
+			RemoteArguments.Add(TargetDesc.Configuration.ToString());
+			RemoteArguments.Add("-SkipRulesCompile"); // Use the rules assembly built locally
+			RemoteArguments.Add(String.Format("-XmlConfigCache={0}", GetRemotePath(XmlConfig.CacheFile!))); // Use the XML config cache built locally, since the remote won't have it
+
+			string? RemoteIniPath = UnrealBuildTool.GetRemoteIniPath();
+			if (!String.IsNullOrEmpty(RemoteIniPath))
+			{
+				RemoteArguments.Add(String.Format("-remoteini={0}", GetRemotePath(RemoteIniPath)));
+			}
+
+			if (TargetDesc.ProjectFile != null)
+			{
+				RemoteArguments.Add(String.Format("-Project={0}", GetRemotePath(TargetDesc.ProjectFile)));
+			}
+
+			foreach (string LocalArgument in TargetDesc.AdditionalArguments)
+			{
+				int EqualsIdx = LocalArgument.IndexOf('=');
+				if (EqualsIdx == -1)
+				{
+					RemoteArguments.Add(LocalArgument);
+					continue;
+				}
+
+				string Key = LocalArgument[..EqualsIdx];
+				string Value = LocalArgument[(EqualsIdx + 1)..];
+
+				if (Key.Equals("-Log", StringComparison.OrdinalIgnoreCase))
+				{
+					// We are already writing to the local log file. The remote will produce a different log (RemoteLogFile)
+					continue;
+				}
+				if (Key.Equals("-Manifest", StringComparison.OrdinalIgnoreCase) && LocalManifestFiles != null)
+				{
+					LocalManifestFiles.Add(new FileReference(Value));
+					continue;
+				}
+
+				string RemoteArgument = LocalArgument;
+				foreach (RemoteMapping Mapping in Mappings)
+				{
+					if (Value.StartsWith(Mapping.LocalDirectory.FullName, StringComparison.OrdinalIgnoreCase))
+					{
+						RemoteArgument = String.Format("{0}={1}", Key, GetRemotePath(Value));
+						break;
+					}
+				}
+				RemoteArguments.Add(RemoteArgument);
+			}
+			return RemoteArguments;
+		}
+
+		/// <summary>
+		/// Runs the actool utility on a directory to create an Assets.car file
+		/// </summary>
+		/// <param name="Platform">The target platform</param>
+		/// <param name="InputDir">Input directory containing assets</param>
+		/// <param name="OutputFile">Path to the Assets.car file to produce</param>
+		/// <param name="Logger">Logger for output</param>
+		public void RunAssetCatalogTool(UnrealTargetPlatform Platform, DirectoryReference InputDir, FileReference OutputFile, ILogger Logger)
+		{
+			Logger.LogInformation("Running asset catalog tool for {Platform}: {InputDir} -> {OutputFile}", Platform, InputDir, OutputFile);
+
+			string RemoteInputDir = GetRemotePath(InputDir);
+			UploadDirectory(InputDir, Logger);
+
+			string RemoteOutputFile = GetRemotePath(OutputFile);
+			Execute(RemoteBaseDir, String.Format("rm -f {0}", EscapeShellArgument(RemoteOutputFile)), Logger);
+
+			string RemoteOutputDir = Path.GetDirectoryName(RemoteOutputFile)!.Replace(Path.DirectorySeparatorChar, '/');
+			Execute(RemoteBaseDir, String.Format("mkdir -p {0}", EscapeShellArgument(RemoteOutputDir)), Logger);
+
+			string RemoteArguments = IOSToolChain.GetAssetCatalogArgs(Platform, RemoteInputDir, RemoteOutputDir);
+			if (Execute(RemoteBaseDir, String.Format("/usr/bin/xcrun {0}", RemoteArguments), Logger) != 0)
+			{
+				throw new BuildException("Failed to run actool.");
+			}
+			DownloadFile(OutputFile, Logger);
+		}
+
+		/// <summary>
+		/// Prepare the remotely built iOS or tvOS client to de debugged on Mac
+		/// </summary>
+		/// <param name="CookedDataDirectory">Where the content has been previously cooked</param>
+		/// <param name="ProjectFile">The project file</param>
+		/// <param name="Logger">The logger</param>
+		public void PrepareToDebug(string CookedDataDirectory, FileReference ProjectFile, ILogger Logger)
+		{
+			Logger.LogInformation("[Remote] Uploading the default uprojectdirs ...");
+			UploadFile(new FileReference("Default.uprojectdirs"), Logger);
+
+			Logger.LogInformation("[Remote] Generating an Xcode Project file...");
+			string CommandLine = EscapeShellArgument(GetRemotePath(Unreal.EngineDirectory));
+			CommandLine += "/Build/BatchFiles/Mac/GenerateProjectFiles.sh";
+			StringBuilder BuildCommandLine = new StringBuilder(CommandLine);
+
+			Execute(GetRemotePath(Unreal.RootDirectory), CommandLine, Logger);
+
+			Logger.LogInformation("[Remote] Uploading cooked data directory...");
+			Logger.LogInformation(CookedDataDirectory);
+			UploadDirectory(new DirectoryReference(CookedDataDirectory), Logger);
+
+			Logger.LogInformation("An Xcode project file has been generated on your remote Mac at the above address. Please access it to debug your IPA.");
+		}
+
+		/// <summary>
+		/// Retrieves data from the primary Mac to the Windows machine
+		/// </summary>
+		/// <param name="ProjectFile">The project file</param>
+		/// <param name="Logger">The logger</param>
+		/// <param name="ClientPlatform">The platfrom for the client (typically iOS or tvOS)</param>
+		public void RetrieveFilesGeneratedOnPrimaryMac(FileReference ProjectFile, ILogger Logger, string ClientPlatform)
+		{
+			Logger.LogInformation("[Remote] Downloading files generated on the Primary Remote Mac ...");
+			DownloadDirectory(DirectoryReference.Combine(Unreal.EngineDirectory, "Intermediate"), Logger);
+			DownloadDirectory(DirectoryReference.Combine(Unreal.EngineDirectory, "Saved"), Logger);
+
+			DownloadDirectory(DirectoryReference.Combine(Unreal.EngineDirectory, "Binaries", ClientPlatform), Logger);
+			DownloadDirectory(DirectoryReference.Combine(ProjectFile.Directory, "Intermediate"), Logger);
+			DownloadDirectory(DirectoryReference.Combine(ProjectFile.Directory, "Binaries", ClientPlatform), Logger);
+			DownloadDirectory(DirectoryReference.Combine(new DirectoryReference("UE5.xcworkspace")), Logger);
+
+		}
+
+		/// <summary>
+		/// Upload retrieved data from the primary Mac from the Windows machine to the secondary Mac
+		/// </summary>
+		/// <param name="ProjectFile">The project file</param>
+		/// <param name="Logger">The logger</param>
+		public void UploadToSecondaryMac(FileReference ProjectFile, ILogger Logger)
+		{
+			Logger.LogInformation("[Remote] Uploading workspace to the secondary Mac ...");
+
+			UploadWorkspace(DirectoryReference.Combine(Unreal.EngineDirectory, ".."), Logger);
+			UploadFile(new FileReference("Default.uprojectdirs"), Logger);
+			UploadDirectory(DirectoryReference.Combine(ProjectFile.Directory, "Intermediate"), Logger);
+			UploadDirectory(DirectoryReference.Combine(ProjectFile.Directory, "Binaries"), Logger);
+			UploadDirectory(DirectoryReference.Combine(ProjectFile.Directory, "Saved"), Logger);
+
+			Logger.LogInformation("[Remote] Generating an Xcode Project file...");
+			string CommandLine = EscapeShellArgument(GetRemotePath(Unreal.EngineDirectory));
+			CommandLine += "/Build/BatchFiles/Mac/GenerateProjectFiles.sh";
+			StringBuilder BuildCommandLine = new StringBuilder(CommandLine);
+			Execute(GetRemotePath(Unreal.RootDirectory), CommandLine, Logger);
+
+			UploadDirectory(new DirectoryReference("UE5.xcworkspace"), Logger);
+		}
+
+		/// <summary>
+		/// Convers a remote path into local form
+		/// </summary>
+		/// <param name="RemotePath">The remote filename</param>
+		/// <returns>Local filename corresponding to the remote path</returns>
+		private FileReference GetLocalPath(string RemotePath)
+		{
+			foreach (RemoteMapping Mapping in Mappings)
+			{
+				if (RemotePath.StartsWith(Mapping.RemoteDirectory, StringComparison.InvariantCultureIgnoreCase) && RemotePath.Length > Mapping.RemoteDirectory.Length && RemotePath[Mapping.RemoteDirectory.Length] == '/')
+				{
+					return FileReference.Combine(Mapping.LocalDirectory, RemotePath[(Mapping.RemoteDirectory.Length + 1)..]);
+				}
+			}
+			throw new BuildException("Unable to map remote path '{0}' to local path", RemotePath);
+		}
+
+		/// <summary>
+		/// Converts a local path into a remote one
+		/// </summary>
+		/// <param name="LocalPath">The local path to convert</param>
+		/// <returns>Equivalent remote path</returns>
+		private string GetRemotePath(FileSystemReference LocalPath)
+		{
+			return GetRemotePath(LocalPath.FullName);
+		}
+
+		/// <summary>
+		/// Converts a local path into a remote one
+		/// </summary>
+		/// <param name="LocalPath">The local path to convert</param>
+		/// <returns>Equivalent remote path</returns>
+		private string GetRemotePath(string LocalPath)
+		{
+			return String.Format("{0}/{1}", RemoteBaseDir, LocalPath.Replace(":", "").Replace('\\', '/').Replace(' ', '_')).Replace('(', '_').Replace(')', '_').Replace('[', '_').Replace(']', '_');
+		}
+
+		/// <summary>
+		/// Gets the local path in Cygwin format (eg. /cygdrive/C/...)
+		/// </summary>
+		/// <param name="InPath">Local path</param>
+		/// <returns>Path in cygwin format</returns>
+		private static string GetLocalCygwinPath(FileSystemReference InPath)
+		{
+			if (InPath.FullName.Length < 2 || InPath.FullName[1] != ':')
+			{
+				throw new BuildException("Invalid local path for converting to cygwin format ({0}).", InPath);
+			}
+			return String.Format("/cygdrive/{0}{1}", InPath.FullName[..1], InPath.FullName[2..].Replace('\\', '/'));
+		}
+
+		/// <summary>
+		/// Escapes spaces and brackets in a shell command argument
+		/// </summary>
+		/// <param name="Argument">The argument to escape</param>
+		/// <returns>The escaped argument</returns>
+		private static string EscapeShellArgument(string Argument)
+		{
+			return Argument.Replace(" ", "\\ ").Replace("[", "\\\\[").Replace("]", "\\\\]");
+		}
+
+		/// <summary>
+		/// Upload a single file to the remote
+		/// </summary>
+		/// <param name="LocalFile">The file to upload</param>
+		/// <param name="Logger">Logger for output</param>
+		void UploadFile(FileReference LocalFile, ILogger Logger)
+		{
+			string RemoteFile = GetRemotePath(LocalFile);
+			string RemoteDirectory = GetRemotePath(LocalFile.Directory);
+
+			List<string> Arguments = RsyncArgumentsForUpload(CommonRsyncArguments, RemoteDirectory);
+			Arguments.Add(String.Format("\"{0}\"", GetLocalCygwinPath(LocalFile)));
+			Arguments.Add(String.Format("\"{0}@{1}\":'{2}'", UserName, ServerName, RemoteFile));
+			Arguments.Add("-q");
+
+			int Result = Rsync(String.Join(" ", Arguments), Logger);
+			if (Result != 0)
+			{
+				throw new BuildException("Error while running Rsync (exit code {0})", Result);
+			}
+		}
+
+		/// <summary>
+		/// Upload a single file to the remote
+		/// </summary>
+		/// <param name="LocalDirectory">The base directory to copy</param>
+		/// <param name="RemoteDirectory">The remote directory</param>
+		/// <param name="LocalFileList">The file to upload</param>
+		/// <param name="Logger">Logger for output</param>
+		void UploadFiles(DirectoryReference LocalDirectory, string RemoteDirectory, FileReference LocalFileList, ILogger Logger)
+		{
+			List<string> Arguments = RsyncArgumentsForUpload(BasicRsyncArguments, RemoteDirectory);
+			Arguments.Add(String.Format("--files-from=\"{0}\"", GetLocalCygwinPath(LocalFileList)));
+			Arguments.Add(String.Format("\"{0}/\"", GetLocalCygwinPath(LocalDirectory)));
+			Arguments.Add(String.Format("\"{0}@{1}\":'{2}/'", UserName, ServerName, RemoteDirectory));
+			Arguments.Add("-q");
+
+			int Result = Rsync(String.Join(" ", Arguments), Logger);
+			if (Result != 0)
+			{
+				throw new BuildException("Error while running Rsync (exit code {0})", Result);
+			}
+		}
+
+		/// <summary>
+		/// Upload a single directory to the remote
+		/// </summary>
+		/// <param name="LocalDirectory">The local directory to upload</param>
+		/// <param name="Logger">Logger for output</param>
+		void UploadDirectory(DirectoryReference LocalDirectory, ILogger Logger)
+		{
+			string RemoteDirectory = GetRemotePath(LocalDirectory);
+
+			List<string> Arguments = RsyncArgumentsForUpload(CommonRsyncArguments, RemoteDirectory);
+			Arguments.Add(String.Format("\"{0}/\"", GetLocalCygwinPath(LocalDirectory)));
+			Arguments.Add(String.Format("\"{0}@{1}\":'{2}/'", UserName, ServerName, RemoteDirectory));
+			Arguments.Add("-q");
+
+			int Result = Rsync(String.Join(" ", Arguments), Logger);
+			if (Result != 0)
+			{
+				throw new BuildException("Error while running Rsync (exit code {0})", Result);
+			}
+		}
+
+		/// <summary>
+		/// Uploads a directory to the remote using a specific filter list
+		/// </summary>
+		/// <param name="LocalDirectory">The local directory to copy from</param>
+		/// <param name="RemoteDirectory">The remote directory to copy to</param>
+		/// <param name="FilterLocations">List of paths to filter</param>
+		/// <param name="Logger">Logger for output</param>
+		void UploadDirectory(DirectoryReference LocalDirectory, string RemoteDirectory, List<FileReference> FilterLocations, ILogger Logger)
+		{
+			List<string> Arguments = RsyncArgumentsForUpload(CommonRsyncArguments, RemoteDirectory);
+			foreach (FileReference FilterLocation in FilterLocations)
+			{
+				Arguments.Add(String.Format("--filter=\"merge {0}\"", GetLocalCygwinPath(FilterLocation)));
+			}
+			Arguments.Add("--exclude='*'");
+			Arguments.Add(String.Format("\"{0}/\"", GetLocalCygwinPath(LocalDirectory)));
+			Arguments.Add(String.Format("\"{0}@{1}\":'{2}/'", UserName, ServerName, RemoteDirectory));
+
+			int Result = Rsync(String.Join(" ", Arguments), Logger);
+			if (Result != 0)
+			{
+				throw new BuildException("Error while running Rsync (exit code {0})", Result);
+			}
+		}
+
+		/// <summary>
+		/// Upload all the files in the workspace for the current project
+		/// </summary>
+		void UploadWorkspace(DirectoryReference TempDir, ILogger Logger)
+		{
+			// Path to the scripts to be uploaded
+			FileReference ScriptPathsFileName = FileReference.Combine(Unreal.EngineDirectory, "Build", "Rsync", "RsyncEngineScripts.txt");
+
+			// Read the list of scripts to be uploaded
+			List<string> ScriptPaths = new List<string>();
+			foreach (string Line in FileReference.ReadAllLines(ScriptPathsFileName))
+			{
+				string FileToUpload = Line.Trim();
+				if (FileToUpload.Length > 0 && FileToUpload[0] != '#')
+				{
+					ScriptPaths.Add(FileToUpload);
+				}
+			}
+
+			// Fixup the line endings
+			List<FileReference> TargetFiles = new List<FileReference>();
+			foreach (string ScriptPath in ScriptPaths)
+			{
+				FileReference SourceFile = FileReference.Combine(Unreal.EngineDirectory, ScriptPath.TrimStart('/'));
+				if (!FileReference.Exists(SourceFile))
+				{
+					throw new BuildException("Missing script required for remote upload: {0}", SourceFile);
+				}
+
+				FileReference TargetFile = FileReference.Combine(TempDir, SourceFile.MakeRelativeTo(Unreal.EngineDirectory));
+				if (!FileReference.Exists(TargetFile) || FileReference.GetLastWriteTimeUtc(TargetFile) < FileReference.GetLastWriteTimeUtc(SourceFile))
+				{
+					DirectoryReference.CreateDirectory(TargetFile.Directory);
+					string ScriptText = FileReference.ReadAllText(SourceFile);
+					FileReference.WriteAllText(TargetFile, ScriptText.Replace("\r\n", "\n"));
+				}
+				TargetFiles.Add(TargetFile);
+			}
+
+			// Write a file that protects all the scripts from being overridden by the standard engine filters
+			FileReference ScriptProtectList = FileReference.Combine(TempDir, "RsyncEngineScripts-Protect.txt");
+			using (StreamWriter Writer = new StreamWriter(ScriptProtectList.FullName))
+			{
+				foreach (string ScriptPath in ScriptPaths)
+				{
+					Writer.WriteLine("protect {0}", ScriptPath);
+				}
+			}
+
+			// Upload these files to the remote
+			Logger.LogInformation("[Remote] Uploading scripts...");
+			UploadFiles(TempDir, GetRemotePath(Unreal.EngineDirectory), ScriptPathsFileName, Logger);
+
+			// Upload the config files
+			Logger.LogInformation("[Remote] Uploading config files...");
+			UploadFile(XmlConfig.CacheFile!, Logger);
+
+			// Upload the engine files
+			List<FileReference> EngineFilters = new List<FileReference>();
+			EngineFilters.Add(ScriptProtectList);
+			if (Unreal.IsEngineInstalled())
+			{
+				EngineFilters.Add(FileReference.Combine(Unreal.EngineDirectory, "Build", "Rsync", "RsyncEngineInstalled.txt"));
+
+				// Upload MarketplaceRules.dll (if it exists) when in InstallBuild. The path is specified in RulesCompiler::CreateMarketplaceRulesAssembly()
+				DirectoryReference MarketplaceEngineDir = DirectoryReference.Combine(Unreal.WritableEngineDirectory, "Intermediate", "Build", "BuildRules");
+				if (DirectoryReference.Exists(MarketplaceEngineDir))
+				{
+					Logger.LogInformation("[Remote] Uploading MarketplaceRules Engine files located in {MarketplaceEngineDir}", MarketplaceEngineDir);
+					UploadDirectory(MarketplaceEngineDir, Logger);
+				}
+			}
+			EngineFilters.Add(FileReference.Combine(Unreal.EngineDirectory, "Build", "Rsync", "RsyncEngine.txt"));
+
+			Logger.LogInformation("[Remote] Uploading engine files...");
+			UploadDirectory(Unreal.EngineDirectory, GetRemotePath(Unreal.EngineDirectory), EngineFilters, Logger);
+
+			// Upload the project files
+			DirectoryReference? ProjectDir = null;
+			if (ProjectFile != null && !ProjectFile.IsUnderDirectory(Unreal.EngineDirectory))
+			{
+				ProjectDir = ProjectFile.Directory;
+			}
+			else if (!String.IsNullOrEmpty(UnrealBuildTool.GetRemoteIniPath()))
+			{
+				ProjectDir = new DirectoryReference(UnrealBuildTool.GetRemoteIniPath()!);
+				if (ProjectDir.IsUnderDirectory(Unreal.EngineDirectory))
+				{
+					ProjectDir = null;
+				}
+			}
+			if (ProjectDir != null)
+			{
+				List<FileReference> ProjectFilters = new List<FileReference>();
+
+				FileReference CustomFilter = FileReference.Combine(ProjectDir, "Build", "Rsync", "RsyncProject.txt");
+				if (FileReference.Exists(CustomFilter))
+				{
+					ProjectFilters.Add(CustomFilter);
+				}
+				ProjectFilters.Add(FileReference.Combine(Unreal.EngineDirectory, "Build", "Rsync", "RsyncProject.txt"));
+
+				Logger.LogInformation("[Remote] Uploading project files...");
+				UploadDirectory(ProjectDir, GetRemotePath(ProjectDir), ProjectFilters, Logger);
+			}
+
+			if (AdditionalPaths != null)
+			{
+				foreach (DirectoryReference AdditionalPath in AdditionalPaths)
+				{
+					List<FileReference> CustomFilters = new List<FileReference>();
+
+					FileReference CustomFilter = FileReference.Combine(AdditionalPath, "Build", "Rsync", "RsyncProject.txt");
+					if (FileReference.Exists(CustomFilter))
+					{
+						CustomFilters.Add(CustomFilter);
+					}
+					CustomFilters.Add(FileReference.Combine(Unreal.EngineDirectory, "Build", "Rsync", "RsyncProject.txt"));
+
+					Logger.LogInformation("[Remote] Uploading additional path files [{Dir}]...", AdditionalPath);
+					UploadDirectory(AdditionalPath, GetRemotePath(AdditionalPath), CustomFilters, Logger);
+				}
+			}
+
+			Execute("/", String.Format("rm -rf {0}/Intermediate/IOS/*.plist", GetRemotePath(Unreal.EngineDirectory)), Logger, true);
+			Execute("/", String.Format("rm -rf {0}/Intermediate/TVOS/*.plist", GetRemotePath(Unreal.EngineDirectory)), Logger, true);
+			if (ProjectFile != null)
+			{
+				Execute("/", String.Format("rm -rf {0}/Intermediate/IOS/*.plist", GetRemotePath(ProjectFile.Directory)), Logger, true);
+				Execute("/", String.Format("rm -rf {0}/Intermediate/TVOS/*.plist", GetRemotePath(ProjectFile.Directory)), Logger, true);
+			}
+
+			// Convert CRLF to LF for all shell scripts
+			Execute(RemoteBaseDir, String.Format("for i in {0}/Build/BatchFiles/Mac/*.sh; do mv $i $i.crlf; tr -d '\r' < $i.crlf > $i; done", EscapeShellArgument(GetRemotePath(Unreal.EngineDirectory))), Logger);
+
+			// Fixup permissions on any shell scripts
+			Execute(RemoteBaseDir, String.Format("chmod +x {0}/Build/BatchFiles/Mac/*.sh", EscapeShellArgument(GetRemotePath(Unreal.EngineDirectory))), Logger);
+		}
+
+		/// <summary>
+		/// Downloads a single file from the remote
+		/// </summary>
+		/// <param name="LocalFile">The file to download</param>
+		/// <param name="Logger">Logger for output</param>
+		void DownloadFile(FileReference LocalFile, ILogger Logger)
+		{
+			RemoteMapping? Mapping = Mappings.FirstOrDefault(x => LocalFile.IsUnderDirectory(x.LocalDirectory));
+			if (Mapping == null)
+			{
+				throw new BuildException("File for download '{0}' is not under any mapped directory.", LocalFile);
+			}
+
+			List<string> Arguments = RsyncArgumentsForDownload(CommonRsyncArguments);
+			Arguments.Add(String.Format("\"{0}@{1}\":'{2}/{3}'", UserName, ServerName, Mapping.RemoteDirectory, LocalFile.MakeRelativeTo(Mapping.LocalDirectory).Replace('\\', '/')));
+			Arguments.Add(String.Format("\"{0}/\"", GetLocalCygwinPath(LocalFile.Directory)));
+			Arguments.Add("-q");
+
+			int Result = Rsync(String.Join(" ", Arguments), Logger);
+			if (Result != 0)
+			{
+				throw new BuildException("Unable to download '{0}' from the remote Mac (exit code {1}).", LocalFile, Result);
+			}
+		}
+
+		/// <summary>
+		/// Download multiple files from the remote Mac
+		/// </summary>
+		/// <param name="Files">List of local files to download</param>
+		/// <param name="Logger">Logger for output</param>
+		void DownloadFiles(IEnumerable<FileReference> Files, ILogger Logger)
+		{
+			List<FileReference>[] FileGroups = new List<FileReference>[Mappings.Count];
+			for (int Idx = 0; Idx < Mappings.Count; Idx++)
+			{
+				FileGroups[Idx] = new List<FileReference>();
+			}
+			foreach (FileReference File in Files)
+			{
+				int MappingIdx = Mappings.FindIndex(x => File.IsUnderDirectory(x.LocalDirectory));
+				if (MappingIdx == -1)
+				{
+					throw new BuildException("File for download '{0}' is not under the engine or project directory.", File);
+				}
+				FileGroups[MappingIdx].Add(File);
+			}
+			for (int Idx = 0; Idx < Mappings.Count; Idx++)
+			{
+				if (FileGroups[Idx].Count > 0)
+				{
+					FileReference DownloadListLocation = FileReference.Combine(Unreal.EngineDirectory, "Intermediate", "Rsync", "Download.txt");
+					DirectoryReference.CreateDirectory(DownloadListLocation.Directory);
+					FileReference.WriteAllLines(DownloadListLocation, FileGroups[Idx].Select(x => x.MakeRelativeTo(Mappings[Idx].LocalDirectory).Replace('\\', '/')));
+
+					List<string> Arguments = RsyncArgumentsForDownload(CommonRsyncArguments);
+					Arguments.Add(String.Format("--files-from=\"{0}\"", GetLocalCygwinPath(DownloadListLocation)));
+					Arguments.Add(String.Format("\"{0}@{1}\":'{2}/'", UserName, ServerName, Mappings[Idx].RemoteDirectory));
+					Arguments.Add(String.Format("\"{0}/\"", GetLocalCygwinPath(Mappings[Idx].LocalDirectory)));
+
+					int Result = Rsync(String.Join(" ", Arguments), Logger);
+					if (Result != 0)
+					{
+						throw new BuildException("Unable to download files from remote Mac (exit code {0})", Result);
+					}
+				}
+			}
+		}
+
+		/// <summary>
+		/// Checks whether a directory exists on the remote machine
+		/// </summary>
+		/// <param name="LocalDirectory">Path to the directory on the local machine</param>
+		/// <param name="Logger">Logger for output</param>
+		/// <returns>True if the remote directory exists</returns>
+		private bool RemoteDirectoryExists(DirectoryReference LocalDirectory, ILogger Logger)
+		{
+			string RemoteDirectory = GetRemotePath(LocalDirectory);
+			return Execute(Unreal.RootDirectory, String.Format("[ -d {0} ]", EscapeShellArgument(RemoteDirectory)), Logger) == 0;
+		}
+
+		/// <summary>
+		/// Download a directory from the remote Mac
+		/// </summary>
+		/// <param name="LocalDirectory">Directory to download</param>
+		/// <param name="Logger">Logger for output</param>
+		private void DownloadDirectory(DirectoryReference LocalDirectory, ILogger Logger)
+		{
+			DirectoryReference.CreateDirectory(LocalDirectory);
+
+			string RemoteDirectory = GetRemotePath(LocalDirectory);
+
+			List<string> Arguments = RsyncArgumentsForDownload(CommonRsyncArguments);
+			Arguments.Add(String.Format("\"{0}@{1}\":'{2}/'", UserName, ServerName, RemoteDirectory));
+			Arguments.Add(String.Format("\"{0}/\"", GetLocalCygwinPath(LocalDirectory)));
+
+			int Result = Rsync(String.Join(" ", Arguments), Logger);
+			if (Result != 0)
+			{
+				throw new BuildException("Unable to download '{0}' from the remote Mac (exit code {1}).", LocalDirectory, Result);
+			}
+		}
+
+		/// <summary>
+		/// Execute Rsync
+		/// </summary>
+		/// <param name="Arguments">Arguments for the Rsync command</param>
+		/// <param name="Logger">Logger for output</param>
+		/// <returns>Exit code from Rsync</returns>
+		private int Rsync(string Arguments, ILogger Logger)
+		{
+			using (Process RsyncProcess = new Process())
+			{
+				DataReceivedEventHandler OutputHandler = (E, Args) => { RsyncOutput(Args, false, Logger); };
+				DataReceivedEventHandler ErrorHandler = (E, Args) => { RsyncOutput(Args, true, Logger); };
+
+				RsyncProcess.StartInfo.FileName = RsyncExe.FullName;
+				RsyncProcess.StartInfo.Arguments = Arguments;
+				RsyncProcess.StartInfo.WorkingDirectory = SshExe.Directory.FullName;
+				RsyncProcess.OutputDataReceived += OutputHandler;
+				RsyncProcess.ErrorDataReceived += ErrorHandler;
+
+				Logger.LogDebug("[Rsync] {File} {Args}", Utils.MakePathSafeToUseWithCommandLine(RsyncProcess.StartInfo.FileName), RsyncProcess.StartInfo.Arguments);
+				return Utils.RunLocalProcess(RsyncProcess);
+			}
+		}
+
+		/// <summary>
+		/// Handles data output by rsync
+		/// </summary>
+		/// <param name="Args">The received data</param>e
+		/// <param name="bStdErr">whether the data was received on stderr</param>
+		/// <param name="Logger">Logger for output</param>
+		private void RsyncOutput(DataReceivedEventArgs Args, bool bStdErr, ILogger Logger)
+		{
+			if (Args.Data != null)
+			{
+				if (bStdErr)
+				{
+					Logger.LogError("  {Output}", Args.Data);
+				}
+				else
+				{
+					Logger.LogInformation("  {Output}", Args.Data);
+				}
+			}
+		}
+
+		/// <summary>
+		/// Execute a command on the remote in the remote equivalent of a local directory
+		/// </summary>
+		/// <param name="WorkingDir"></param>
+		/// <param name="Command"></param>
+		/// <param name="Logger">Logger for output</param>
+		/// <param name="bSilent"></param>
+		/// <returns></returns>
+		public int Execute(DirectoryReference WorkingDir, string Command, ILogger Logger, bool bSilent = false)
+		{
+			return Execute(GetRemotePath(WorkingDir), Command, Logger, bSilent);
+		}
+
+		/// <summary>
+		/// Execute a remote command, capturing the output text
+		/// </summary>
+		/// <param name="WorkingDirectory">The remote working directory</param>
+		/// <param name="Command">Command to be executed</param>
+		/// <param name="Logger">Logger for output</param>
+		/// <param name="bSilent">If true, logging is suppressed</param>
+		/// <returns></returns>
+		protected int Execute(string WorkingDirectory, string Command, ILogger Logger, bool bSilent = false)
+		{
+			string FullCommand = String.Format("cd {0} && {1}", EscapeShellArgument(WorkingDirectory), Command);
+			using (Process SSHProcess = new Process())
+			{
+				DataReceivedEventHandler OutputHandler = (E, Args) => { SshOutput(Args, false, Logger); };
+				DataReceivedEventHandler ErrorHandler = (E, Args) => { SshOutput(Args, true, Logger); };
+
+				SSHProcess.StartInfo.FileName = SshExe.FullName;
+				SSHProcess.StartInfo.WorkingDirectory = SshExe.Directory.FullName;
+				SSHProcess.StartInfo.Arguments = String.Format("{0} {1}", String.Join(" ", CommonSshArguments), FullCommand);
+				if (!bSilent)
+				{
+					SSHProcess.OutputDataReceived += OutputHandler;
+					SSHProcess.ErrorDataReceived += ErrorHandler;
+				}
+
+				Logger.LogDebug("[SSH] {Exe} {Args}", Utils.MakePathSafeToUseWithCommandLine(SSHProcess.StartInfo.FileName), SSHProcess.StartInfo.Arguments);
+				return Utils.RunLocalProcess(SSHProcess);
+			}
+		}
+
+		/// <summary>
+		/// Handler for output from running remote SSH commands
+		/// </summary>
+		/// <param name="Args"></param>
+		/// <param name="bStdErr">whether the data was received on stderr</param>
+		/// <param name="Logger">Logger for output</param>
+		private void SshOutput(DataReceivedEventArgs Args, bool bStdErr, ILogger Logger)
+		{
+			if (Args.Data != null)
+			{
+				string FormattedOutput = ConvertRemotePathsToLocal(Args.Data);
+				if (bStdErr)
+				{
+					Logger.LogError("  {Output}", FormattedOutput);
+				}
+				else
+				{
+					Logger.LogInformation("  {Output}", FormattedOutput);
+				}
+			}
+		}
+
+		/// <summary>
+		/// Execute a remote command, capturing the output text
+		/// </summary>
+		/// <param name="Command">Command to be executed</param>
+		/// <param name="Logger">Logger for output</param>
+		/// <param name="Output">Receives the output text</param>
+		/// <returns></returns>
+		protected int ExecuteAndCaptureOutput(string Command, ILogger Logger, out StringBuilder Output)
+		{
+			StringBuilder FullCommand = new StringBuilder();
+			foreach (string CommonSshArgument in CommonSshArguments)
+			{
+				FullCommand.AppendFormat("{0} ", CommonSshArgument);
+			}
+			FullCommand.Append(Command.Replace("\"", "\\\""));
+
+			using (Process SSHProcess = new Process())
+			{
+				Output = new StringBuilder();
+
+				StringBuilder OutputLocal = Output;
+				DataReceivedEventHandler OutputHandler = (E, Args) =>
+				{
+					if (Args.Data != null)
+					{
+						OutputLocal.Append(Args.Data);
+					}
+				};
+
+				SSHProcess.StartInfo.FileName = SshExe.FullName;
+				SSHProcess.StartInfo.WorkingDirectory = SshExe.Directory.FullName;
+				SSHProcess.StartInfo.Arguments = FullCommand.ToString();
+				SSHProcess.OutputDataReceived += OutputHandler;
+				SSHProcess.ErrorDataReceived += OutputHandler;
+
+				Logger.LogDebug("[SSH] {Exe} {Args}", Utils.MakePathSafeToUseWithCommandLine(SSHProcess.StartInfo.FileName), SSHProcess.StartInfo.Arguments);
+				return Utils.RunLocalProcess(SSHProcess);
+			}
+		}
+
+		/// <summary>
+		/// Converts any remote paths within the given string to local format
+		/// </summary>
+		/// <param name="Text">The text containing strings to convert</param>
+		/// <returns>The string with paths converted to local format</returns>
+		private string ConvertRemotePathsToLocal(string Text)
+		{
+			// Try to match any source file with the remote base directory in front of it
+			string Pattern = String.Format("(?<![a-zA-Z=]){0}[^:]*\\.(?:cpp|inl|h|hpp|hh|txt)(?![a-zA-Z])", Regex.Escape(RemoteBaseDir));
+
+			// Find the matches, and early out if there are none
+			MatchCollection Matches = Regex.Matches(Text, Pattern, RegexOptions.IgnoreCase);
+			if (Matches.Count == 0)
+			{
+				return Text;
+			}
+
+			// Replace any remote paths with local ones
+			StringBuilder Result = new StringBuilder();
+			int StartIdx = 0;
+			foreach (Match? Match in Matches)
+			{
+				// Append the text leading up to this path
+				Result.Append(Text, StartIdx, Match!.Index - StartIdx);
+
+				// Try to convert the path
+				string Path = Match.Value;
+				foreach (RemoteMapping Mapping in Mappings)
+				{
+					if (Path.StartsWith(Mapping.RemoteDirectory))
+					{
+						Path = Mapping.LocalDirectory + Path[Mapping.RemoteDirectory.Length..].Replace('/', '\\');
+						break;
+					}
+				}
+
+				// Append the path to the output string
+				Result.Append(Path);
+
+				// Move past this match
+				StartIdx = Match.Index + Match.Length;
+			}
+			Result.Append(Text, StartIdx, Text.Length - StartIdx);
+			return Result.ToString();
+		}
+	}
+}
diff --git a/Engine/Source/Programs/UnrealBuildTool/Platform/IOS/IOSToolChain.cs b/Engine/Source/Programs/UnrealBuildTool/Platform/IOS/IOSToolChain.cs
--- a/Engine/Source/Programs/UnrealBuildTool/Platform/IOS/IOSToolChain.cs
+++ b/Engine/Source/Programs/UnrealBuildTool/Platform/IOS/IOSToolChain.cs
@@ -15,19 +15,6 @@ using UnrealBuildBase;
 
 namespace UnrealBuildTool
 {
-	class IOSToolChainSettings : AppleToolChainSettings
-	{
-		public IOSToolChainSettings(IOSProjectSettings ProjectSettings, ILogger Logger)
-			: this("iPhoneOS", "iPhoneSimulator", "ios", ProjectSettings, Logger)
-		{
-		}
-
-		protected IOSToolChainSettings(string DevicePlatformName, string SimulatorPlatformName, string TargetOSName, IOSProjectSettings ProjectSettings, ILogger Logger)
-			: base(DevicePlatformName, SimulatorPlatformName, TargetOSName, ProjectSettings.RuntimeVersion, true, Logger)
-		{
-		}
-	}
-
 	class IOSToolChain : AppleToolChain
 	{
 		private static List<FileItem> BundleDependencies = new List<FileItem>();
diff --git a/Engine/Source/Programs/UnrealBuildTool/Platform/IOS/IOSToolChainSettings.cs b/Engine/Source/Programs/UnrealBuildTool/Platform/IOS/IOSToolChainSettings.cs
new file mode 100644
--- /dev/null
+++ b/Engine/Source/Programs/UnrealBuildTool/Platform/IOS/IOSToolChainSettings.cs
@@ -0,0 +1,19 @@
+// Copyright Epic Games, Inc. All Rights Reserved.
+
+using Microsoft.Extensions.Logging;
+
+namespace UnrealBuildTool
+{
+	class IOSToolChainSettings : AppleToolChainSettings
+	{
+		public IOSToolChainSettings(IOSProjectSettings ProjectSettings, ILogger Logger)
+			: this("iPhoneOS", "iPhoneSimulator", "ios", ProjectSettings, Logger)
+		{
+		}
+
+		protected IOSToolChainSettings(string DevicePlatformName, string SimulatorPlatformName, string TargetOSName, IOSProjectSettings ProjectSettings, ILogger Logger)
+			: base(DevicePlatformName, SimulatorPlatformName, TargetOSName, ProjectSettings.RuntimeVersion, true, Logger)
+		{
+		}
+	}
+}
diff --git a/Engine/Source/Programs/UnrealBuildTool/Platform/TVOS/TVOSToolChain.cs b/Engine/Source/Programs/UnrealBuildTool/Platform/TVOS/TVOSToolChain.cs
--- a/Engine/Source/Programs/UnrealBuildTool/Platform/TVOS/TVOSToolChain.cs
+++ b/Engine/Source/Programs/UnrealBuildTool/Platform/TVOS/TVOSToolChain.cs
@@ -4,14 +4,6 @@ using Microsoft.Extensions.Logging;
 
 namespace UnrealBuildTool
 {
-	class TVOSToolChainSettings : IOSToolChainSettings
-	{
-		public TVOSToolChainSettings(IOSProjectSettings ProjectSettings, ILogger Logger) 
-			: base("AppleTVOS", "AppleTVSimulator", "tvos", ProjectSettings, Logger)
-		{
-		}
-	}
-
 	class TVOSToolChain : IOSToolChain
 	{
 		public TVOSToolChain(ReadOnlyTargetRules InTarget, TVOSProjectSettings InProjectSettings, ILogger InLogger)
diff --git a/Engine/Source/Programs/UnrealBuildTool/Platform/TVOS/TVOSToolChainSettings.cs b/Engine/Source/Programs/UnrealBuildTool/Platform/TVOS/TVOSToolChainSettings.cs
new file mode 100644
--- /dev/null
+++ b/Engine/Source/Programs/UnrealBuildTool/Platform/TVOS/TVOSToolChainSettings.cs
@@ -0,0 +1,14 @@
+// Copyright Epic Games, Inc. All Rights Reserved.
+
+using Microsoft.Extensions.Logging;
+
+namespace UnrealBuildTool
+{
+	class TVOSToolChainSettings : IOSToolChainSettings
+	{
+		public TVOSToolChainSettings(IOSProjectSettings ProjectSettings, ILogger Logger) 
+			: base("AppleTVOS", "AppleTVSimulator", "tvos", ProjectSettings, Logger)
+		{
+		}
+	}
+}
diff --git a/Engine/Source/Programs/UnrealBuildTool/ToolChain/AppleToolChain.cs b/Engine/Source/Programs/UnrealBuildTool/ToolChain/AppleToolChain.cs
deleted file mode 100644
diff --git a/Engine/Source/Programs/UnrealBuildTool/ToolChain/RemoteMac.cs b/Engine/Source/Programs/UnrealBuildTool/ToolChain/RemoteMac.cs
deleted file mode 100644
