Don't generate RTFM __attribute__ when doing VisualCpp static analysis

#rnx

diff --git a/Engine/Source/Programs/UnrealBuildTool/Configuration/UEBuildModule.cs b/Engine/Source/Programs/UnrealBuildTool/Configuration/UEBuildModule.cs
--- a/Engine/Source/Programs/UnrealBuildTool/Configuration/UEBuildModule.cs
+++ b/Engine/Source/Programs/UnrealBuildTool/Configuration/UEBuildModule.cs
@@ -2,6 +2,7 @@
 
 using System;
 using System.Collections.Generic;
+using System.Diagnostics;
 using System.IO;
 using System.Linq;
 using System.Reflection;
@@ -361,7 +362,9 @@ namespace UnrealBuildTool
 			// disable AutoRTFM on this function the compiler would spit nails. But it'll happily consume the C-style attribute at that location and do
 			// the right thing, so we use that instead.
 			// Caution: Hard coded constant '1' in this string must be kept in sync with Engine\Source\Runtime\AutoRTFM\Public\AutoRTFMConstants.h.
-			string AutoRTFMDisable = (Rules.Target.bUseAutoRTFMCompiler && Rules.bDisableAutoRTFMInstrumentation && Rules.bInjectAutoRTFMAttributeDisables) ?
+			Debug.Assert(Rules.Target.StaticAnalyzer == StaticAnalyzer.None || Rules.Target.StaticAnalyzer == StaticAnalyzer.Default);
+			bool bNonClangSA = Rules.Target.StaticAnalyzer == StaticAnalyzer.Default && Rules.Target.Platform.IsInGroup(UnrealPlatformGroup.Microsoft) && !Rules.Target.WindowsPlatform.Compiler.IsClang();
+			string AutoRTFMDisable = (Rules.Target.bUseAutoRTFMCompiler && Rules.bDisableAutoRTFMInstrumentation && Rules.bInjectAutoRTFMAttributeDisables && !bNonClangSA) ?
 				"__attribute__((autortfm(/*autortfm_mode_disable*/ 1)))" : String.Empty;
 
 			// We precalculate these because they are reused over and over again and causing tons of allocations
