[AutoRTFM] Match LLVM / AutoRTFM diagnostics in UBT

For some reason, LLVM diagnostics are printed as: `<severity>: <file>:<location>: <message>`
While Clang diagnostics are printed as: `<file>:<location>: <severity>: <message>`

Update UBT to handle both the LLVM and Clang forms.

Reworked the AutoRTFMCompilerPass to standardize the diagnostics to avoid formatting issues breaking the matcher in the future.

diff --git a/Engine/Source/Programs/UnrealBuildTool/Matchers/CompileEventMatcher.cs b/Engine/Source/Programs/UnrealBuildTool/Matchers/CompileEventMatcher.cs
--- a/Engine/Source/Programs/UnrealBuildTool/Matchers/CompileEventMatcher.cs
+++ b/Engine/Source/Programs/UnrealBuildTool/Matchers/CompileEventMatcher.cs
@@ -62,6 +62,8 @@ namespace UnrealBuildTool.Matchers
 		static readonly Regex s_clangDiagnosticPattern = new Regex($"^\\s*{FilePattern}\\s*{ClangLocationPattern}:\\s*{ClangSeverity}\\s*:");
 		static readonly Regex s_clangNotePattern = new Regex($"^\\s*{FilePattern}\\s*{ClangLocationPattern}:\\s*note:");
 		static readonly Regex s_clangMarkerPattern = new Regex(@"^(\s*)[\^~][\s\^~]*$");
+		static readonly Regex s_llvmDiagnosticPattern = new Regex($"^\\s*{ClangSeverity}:\\s*{FilePattern}\\s*{ClangLocationPattern}:\\s*");
+		static readonly Regex s_llvmNotePattern = new Regex($"^\\s*note:\\s*(?:{FilePattern}\\s*{ClangLocationPattern}:)?");
 		static readonly Regex s_xcodeIDEWatchExtensionPattern = new Regex(@"xcodebuild.*Requested but did not find extension point with identifier.*for extension.*\.watchOS of plug-in com\.apple\.dt\.IDEWatchSupportCore");
 		static readonly Regex s_scriptCompilePattern = new Regex(@"^\s*[A-Za-z0-9_\.]+ ERROR:.* [A-Za-z_]+ failed to compile\.");
 		static readonly Regex s_cscSummaryPattern = new Regex(@"^\s+\d+ (?:Warning|Error)\(s\)");
@@ -140,9 +142,10 @@ namespace UnrealBuildTool.Matchers
 					return eventMatch;
 				}
 
-				// Try to match a Clang diagnostic
+				// Try to match a Clang or LLVM diagnostic
 				Match? match;
-				if (builder.Current.TryMatch(s_clangDiagnosticPattern, out match) && IsSourceFile(match))
+				if ((builder.Current.TryMatch(s_clangDiagnosticPattern, out match) ||
+				     builder.Current.TryMatch(s_llvmDiagnosticPattern, out match)) && IsSourceFile(match))
 				{
 					LogLevel level = GetLogLevelFromSeverity(match);
 
@@ -155,7 +158,8 @@ namespace UnrealBuildTool.Matchers
 					{
 						SkipClangMarker(builder);
 
-						if (!builder.Next.TryMatch(s_clangNotePattern, out match))
+						if (!(builder.Next.TryMatch(s_clangNotePattern, out match) ||
+						      builder.Next.TryMatch(s_llvmNotePattern, out match)))
 						{
 							break;
 						}
