Move the location of the SxS manifest for Win64 arm64 builds to be unique per target.

diff --git a/Engine/Source/Programs/UnrealBuildTool/Platform/Windows/VCToolChain.cs b/Engine/Source/Programs/UnrealBuildTool/Platform/Windows/VCToolChain.cs
--- a/Engine/Source/Programs/UnrealBuildTool/Platform/Windows/VCToolChain.cs
+++ b/Engine/Source/Programs/UnrealBuildTool/Platform/Windows/VCToolChain.cs
@@ -2594,8 +2594,6 @@ namespace UnrealBuildTool
 			OptionalUpdateSxSManifest(Target, MakefileBuilder);
 		}
 
-		static object SxSManifestLock = new();
-
 		protected void OptionalUpdateSxSManifest(ReadOnlyTargetRules Target, TargetMakefileBuilder MakefileBuilder)
 		{
 			if (Target.WindowsPlatform.ManifestFile != null && !Target.bShouldCompileAsDLL && !Target.Architecture.bIsX64)
@@ -2612,44 +2610,40 @@ namespace UnrealBuildTool
 
 				// load or create the intermediate SxS manifest
 				bool bDirty = false;
-				DirectoryReference BaseDir = Target.ProjectFile?.Directory ?? Unreal.EngineDirectory;
-				FileReference IntermediateManifest = FileReference.Combine(  BaseDir, "Intermediate", "SxSManifest", AssemblyName, AssemblyName + ".manifest");
-				FileReference TemplateManifest = FileReference.Combine(Unreal.EngineDirectory, "Build", "Windows", "Resources", AssemblyName, AssemblyName + ".manifest");			
+				FileReference IntermediateManifest = GetSxSIntermediateManifest(Target);
+				FileReference TemplateManifest = FileReference.Combine(Unreal.EngineDirectory, "Build", "Windows", "Resources", AssemblyName, AssemblyName + ".manifest");
 				XDocument XmlManifest;
 
-				lock (SxSManifestLock)
+				if (FileReference.Exists(IntermediateManifest))
 				{
-					if (FileReference.Exists(IntermediateManifest))
-					{
-						XmlManifest = XDocument.Load(IntermediateManifest.FullName);
-					}
-					else
-					{
-						XmlManifest = XDocument.Load(TemplateManifest.FullName);
-						bDirty = true;
-					}
+					XmlManifest = XDocument.Load(IntermediateManifest.FullName);
+				}
+				else
+				{
+					XmlManifest = XDocument.Load(TemplateManifest.FullName);
+					bDirty = true;
+				}
 
-					// update & save the intermediate SxS manifest
-					if (XmlManifest.Root != null && XmlManifest.Root.Name.LocalName == "assembly")
-					{
-						XNamespace XmlNs = XmlManifest.Root.Name.Namespace;
-						HashSet<string> ExistingDLLNames = [.. XmlManifest.Root.Elements(XmlNs + "file").Select(X => X.Attribute("name")!.Value)];
+				// update & save the intermediate SxS manifest
+				if (XmlManifest.Root != null && XmlManifest.Root.Name.LocalName == "assembly")
+				{
+					XNamespace XmlNs = XmlManifest.Root.Name.Namespace;
+					HashSet<string> ExistingDLLNames = [.. XmlManifest.Root.Elements(XmlNs + "file").Select(X => X.Attribute("name")!.Value)];
 
-						foreach (FileItem DLL in DLLs)
+					foreach (FileItem DLL in DLLs)
+					{
+						string DLLName = DLL.Location.MakeRelativeTo(ManifestDir);
+						if (!ExistingDLLNames.Contains(DLLName))
 						{
-							string DLLName = DLL.Location.MakeRelativeTo(ManifestDir);
-							if (!ExistingDLLNames.Contains(DLLName))
-							{
-								XmlManifest.Root.Add(new XElement(XmlNs + "file", new XAttribute("name", DLLName)));
-								bDirty = true;
-							}
+							XmlManifest.Root.Add(new XElement(XmlNs + "file", new XAttribute("name", DLLName)));
+							bDirty = true;
 						}
 					}
-					if (bDirty)
-					{
-						DirectoryReference.CreateDirectory(IntermediateManifest.Directory);
-						XmlManifest.Save(IntermediateManifest.FullName);
-					}
+				}
+				if (bDirty)
+				{
+					DirectoryReference.CreateDirectory(IntermediateManifest.Directory);
+					XmlManifest.Save(IntermediateManifest.FullName);
 				}
 				// include the intermediate SxS manifest as a dependency
 				MakefileBuilder.Makefile.InternalDependencies.Add( FileItem.GetItemByFileReference(IntermediateManifest) );
@@ -2719,9 +2713,7 @@ namespace UnrealBuildTool
 			if (Target.WindowsPlatform.ManifestFile != null && !Target.bShouldCompileAsDLL && !Target.Architecture.bIsX64)
 			{
 				string AssemblyName = Target.Architecture.ToString().ToLower();
-				DirectoryReference BaseDir = Target.ProjectFile?.Directory ?? Unreal.EngineDirectory;
-
-				FileReference SrcFile = FileReference.Combine(  BaseDir, "Intermediate", "SxSManifest", AssemblyName, AssemblyName + ".manifest"); // created in OptionalUpdateSxSManifest during a post-build step
+				FileReference SrcFile = GetSxSIntermediateManifest(Target);
 				FileReference DstFile = FileReference.Combine(ExeDir, AssemblyName, AssemblyName + ".manifest");
 
 				TargetFileToSourceFile[DstFile] = SrcFile;
@@ -3988,11 +3980,21 @@ namespace UnrealBuildTool
 
 			if (Target.WindowsPlatform.ManifestFile != null && Binary.Type == UEBuildBinaryType.Executable && !Target.Architecture.bIsX64)
 			{
-				string AssemblyName = Target.Architecture.ToString().ToLower();
-				DirectoryReference BaseDir = Target.ProjectFile?.Directory ?? Unreal.EngineDirectory;
-				FileReference IntermediateManifest = FileReference.Combine(BaseDir, "Intermediate", "SxSManifest", AssemblyName, AssemblyName + ".manifest");
+				FileReference IntermediateManifest = GetSxSIntermediateManifest(Target);
 				BuildProducts.TryAdd(IntermediateManifest, BuildProductType.BuildResource);
 			}
 		}
+
+		private static FileReference GetSxSIntermediateManifest(ReadOnlyTargetRules Target)
+		{
+			DirectoryReference BaseDir = Target.ProjectFile?.Directory ?? Unreal.EngineDirectory;
+			string PlatformIntermediateDirectory = UEBuildTarget.GetPlatformIntermediateFolder(Target.Platform, Target.Architectures, false);
+			string TargetIntermediateName = UEBuildTarget.GetTargetIntermediateFolderName(Target.Name, Target.IntermediateEnvironment);
+			string TargetConfiguration = Target.Configuration.ToString();
+			DirectoryReference IntermediateDirectory = DirectoryReference.Combine(BaseDir, PlatformIntermediateDirectory, TargetIntermediateName, TargetConfiguration);
+
+			string ManifestName = $"{Target.Architecture}.manifest";
+			return FileReference.Combine(IntermediateDirectory, ManifestName);
+		}
 	}
 }
