Add back in the .editorconfig file for UBT's Matchers directory, as these files are also included directly by the JobDriver project in Horde.

diff --git a/Engine/Source/Programs/UnrealBuildTool/Matchers/.editorconfig b/Engine/Source/Programs/UnrealBuildTool/Matchers/.editorconfig
new file mode 100644
--- /dev/null
+++ b/Engine/Source/Programs/UnrealBuildTool/Matchers/.editorconfig
@@ -0,0 +1,19 @@
+# NOTE: Requires **VS2019 16.3** or later
+
+# EditorConfig is awesome:http://EditorConfig.org
+
+# Dotnet code style settings:
+[*.cs]
+
+# IMPORTANT:
+# These files are currently imported directly by Engine/Source/Programs/Horde/Drivers/JobDriver/JobDriver.cs in addition to being a part of UBT.
+# As such, analysis builds need to be done on both UBT and Horde when changes are made to this file.
+
+# IDE0005: Using directive is unnecessary
+dotnet_diagnostic.IDE0005.severity = silent
+
+# MA0009: Regular expressions should not be vulnerable to Denial of Service attacks
+dotnet_diagnostic.MA0009.severity = none
+
+# MA0051: Method is too long
+dotnet_diagnostic.MA0051.severity = none
