Enable input cache on build machines

#rnx

diff --git a/Engine/Source/Programs/UnrealBuildTool/Configuration/UEBuildTarget.cs b/Engine/Source/Programs/UnrealBuildTool/Configuration/UEBuildTarget.cs
--- a/Engine/Source/Programs/UnrealBuildTool/Configuration/UEBuildTarget.cs
+++ b/Engine/Source/Programs/UnrealBuildTool/Configuration/UEBuildTarget.cs
@@ -2761,9 +2761,9 @@ namespace UnrealBuildTool
 				Makefile.UHTTargetSettings.GenerateEnumMaxValues = true;
 			}
 
-			Makefile.UHTTargetSettings.EnableInputCacheRead = BuildConfiguration.bEnableUHTInputCache && !BuildConfiguration.bForceHeaderGeneration;
-			Makefile.UHTTargetSettings.EnableInputCacheWrite = BuildConfiguration.bEnableUHTInputCache;
-
+			bool bEnabledUHTInputCache = BuildConfiguration.bEnableUHTInputCache || Unreal.IsBuildMachine();
+			Makefile.UHTTargetSettings.EnableInputCacheRead = bEnabledUHTInputCache && !BuildConfiguration.bForceHeaderGeneration;
+			Makefile.UHTTargetSettings.EnableInputCacheWrite = bEnabledUHTInputCache;
 
 #if __VPROJECT_AVAILABLE__
 			{
