Improved stats for input cache

#rnx

diff --git a/Engine/Source/Programs/Shared/EpicGames.UHT/Utils/UhtInputCache.cs b/Engine/Source/Programs/Shared/EpicGames.UHT/Utils/UhtInputCache.cs
--- a/Engine/Source/Programs/Shared/EpicGames.UHT/Utils/UhtInputCache.cs
+++ b/Engine/Source/Programs/Shared/EpicGames.UHT/Utils/UhtInputCache.cs
@@ -1137,7 +1137,8 @@ namespace EpicGames.UHT.Utils
 						return new();
 					}
 				}
-				
+
+				module.Session.IncrementTotalCacheFilesRead();
 				return new(cachedFiles, reader.RemainingMemory);
 			}
 			catch (Exception ex)
@@ -1196,6 +1197,8 @@ namespace EpicGames.UHT.Utils
 						cacheStream.Write(header.InputCacheEntry.CompressedData.Span);
 					}
 				}
+
+				module.Session.IncrementTotalCacheFilesWritten();
 			}
 			catch(Exception ex) 
 			{
diff --git a/Engine/Source/Programs/Shared/EpicGames.UHT/Utils/UhtSession.cs b/Engine/Source/Programs/Shared/EpicGames.UHT/Utils/UhtSession.cs
--- a/Engine/Source/Programs/Shared/EpicGames.UHT/Utils/UhtSession.cs
+++ b/Engine/Source/Programs/Shared/EpicGames.UHT/Utils/UhtSession.cs
@@ -1063,6 +1063,18 @@ namespace EpicGames.UHT.Utils
 		public int TotalFilesCacheRead => Interlocked.Add(ref _totalFilesCacheRead, 0);
 		private int _totalFilesCacheRead = 0;
 
+		/// <summary>
+		/// Number of cache files read
+		/// </summary>
+		public int TotalCacheFilesRead => Interlocked.Add(ref _totalCacheFilesRead, 0);
+		private int _totalCacheFilesRead = 0;
+
+		/// <summary>
+		/// Number of cache files written
+		/// </summary>
+		public int TotalCacheFilesWritten => Interlocked.Add(ref _totalCacheFilesWritten, 0);
+		private int _totalCacheFilesWritten = 0;
+
 		/// <summary>
 		/// Number of source files failed to deserialize
 		/// </summary>
@@ -1093,6 +1105,24 @@ namespace EpicGames.UHT.Utils
 			return Interlocked.Increment(ref _totalFilesCacheRead);
 		}
 
+		/// <summary>
+		/// Increment the total number of cache files read
+		/// </summary>
+		/// <returns></returns>
+		public int IncrementTotalCacheFilesRead()
+		{
+			return Interlocked.Increment(ref _totalCacheFilesRead);
+		}
+
+		/// <summary>
+		/// Increment the total number of cache files written
+		/// </summary>
+		/// <returns></returns>
+		public int IncrementTotalCacheFilesWritten()
+		{
+			return Interlocked.Increment(ref _totalCacheFilesWritten);
+		}
+
 		/// <summary>
 		/// Increment the total number of errors deserializing cached files
 		/// </summary>
diff --git a/Engine/Source/Programs/UnrealBuildTool/Modes/UnrealHeaderToolMode.cs b/Engine/Source/Programs/UnrealBuildTool/Modes/UnrealHeaderToolMode.cs
--- a/Engine/Source/Programs/UnrealBuildTool/Modes/UnrealHeaderToolMode.cs
+++ b/Engine/Source/Programs/UnrealBuildTool/Modes/UnrealHeaderToolMode.cs
@@ -860,12 +860,20 @@ namespace UnrealBuildTool.Modes
 				// Generate summary message
 				stopwatch.Stop();
 				double elapsed = stopwatch.Elapsed.TotalSeconds;
-				LogLevel level = elapsed > 5 || Options.Mode != null ? LogLevel.Information : LogLevel.Debug;
+				LogLevel level = elapsed > 5 || Options.Mode != null || Unreal.IsBuildMachine() ? LogLevel.Information : LogLevel.Debug;
 				if (!Session.HasErrors || Session.TotalFilesWritten != 0)
 				{
 					string ActualTargetName = Session.Manifest != null ? Session.Manifest.TargetName : "UnknownTarget";
-					Logger.Log(level, "UHT processed {TargetName} in {Time} seconds ({TotalFilesWritten} generated files written, {TotalFilesRead} H files read, {TotalFilesCacheRead} H files read from cache)",
-						ActualTargetName, stopwatch.Elapsed.TotalSeconds, Session.TotalFilesWritten, Session.TotalFilesRead, Session.TotalFilesCacheRead);
+					if (Session.EnableInputCacheRead || Session.EnableInputCacheWrite)
+					{
+						Logger.Log(level, "UHT processed {TargetName} in {Time} seconds (Generated Files: {TotalFilesWritten}, Source Files Read (disk/cache): {TotalFilesRead}/{TotalFileCacheRead}, Cache Files (read/written): {TotalCacheFilesRead}/{TotalCacheFilesWritten})",
+							ActualTargetName, stopwatch.Elapsed.TotalSeconds, Session.TotalFilesWritten, Session.TotalFilesRead, Session.TotalFilesCacheRead, Session.TotalCacheFilesRead, Session.TotalCacheFilesWritten);
+					}
+					else
+					{
+						Logger.Log(level, "UHT processed {TargetName} in {Time} seconds ({TotalFilesWritten} generated files written)",
+							ActualTargetName, stopwatch.Elapsed.TotalSeconds, Session.TotalFilesWritten);
+					}
 				}
 				if (Session.TotalFilesCacheReadErrors != 0 || Session.TotalFilesCacheWriteErrors != 0)
 				{
