[UBT]
* Added "/dximplicitlyinstantiatevirtualfunction-" to msvc compiles to make msvc behave like clang/gcc when it comes to instantiating virtual methods in templated types.

diff --git a/Engine/Source/Programs/UnrealBuildTool/Platform/Windows/VCToolChain.cs b/Engine/Source/Programs/UnrealBuildTool/Platform/Windows/VCToolChain.cs
--- a/Engine/Source/Programs/UnrealBuildTool/Platform/Windows/VCToolChain.cs
+++ b/Engine/Source/Programs/UnrealBuildTool/Platform/Windows/VCToolChain.cs
@@ -1259,6 +1259,10 @@ namespace UnrealBuildTool
 				// Allow large object files to avoid hitting the 2^16 section limit when running with -StressTestUnity.
 				// Note: not needed for clang, it implicitly upgrades COFF files to bigobj format when necessary.
 				Arguments.Add("/bigobj");
+
+				// This flag changes so msvc delay instantiation of virtual methods until they are needed.
+				// This brings msvc behavior closer to clang/gcc and also reduces pch/binary size a bit while increasing obj files a tiny bit (net win)
+				Arguments.Add("/dximplicitlyinstantiatevirtualfunction-");
 			}
 
 			FPSemanticsMode FPSemantics = CompileEnvironment.FPSemantics;
