UnrealBuildTool: Add braces to statement (ide0011)

#rnx

diff --git a/Engine/Source/Programs/UnrealBuildTool/Configuration/UEBuildModule.cs b/Engine/Source/Programs/UnrealBuildTool/Configuration/UEBuildModule.cs
--- a/Engine/Source/Programs/UnrealBuildTool/Configuration/UEBuildModule.cs
+++ b/Engine/Source/Programs/UnrealBuildTool/Configuration/UEBuildModule.cs
@@ -1715,7 +1715,9 @@ namespace UnrealBuildTool
 			UEBuildModule[] ordered = new UEBuildModule[nodeCount];
 			int counter = 0;
 			foreach (int idx in outIdx)
+			{
 				ordered[counter++] = NodeList[idx];
+			}
 
 			return ordered;
 		}
diff --git a/Engine/Source/Programs/UnrealBuildTool/Configuration/UEBuildModuleCPP.cs b/Engine/Source/Programs/UnrealBuildTool/Configuration/UEBuildModuleCPP.cs
--- a/Engine/Source/Programs/UnrealBuildTool/Configuration/UEBuildModuleCPP.cs
+++ b/Engine/Source/Programs/UnrealBuildTool/Configuration/UEBuildModuleCPP.cs
@@ -1388,9 +1388,15 @@ namespace UnrealBuildTool
 				BloomFilter dirFilter = new();
 				BloomFilter fileFilter = new();
 				foreach (DirectoryItem dir in dirItem.EnumerateDirectories())
+				{
 					dirFilter.Add(dir.Name);
+				}
+
 				foreach (FileItem file in dirItem.EnumerateFiles())
+				{
 					fileFilter.Add(file.Name);
+				}
+
 				UserIncludePaths.Add((dirItem, dirFilter, fileFilter));
 			}
 
@@ -2371,7 +2377,10 @@ namespace UnrealBuildTool
 									if (NameToHeaderFile.TryGetValue(ExpectedName, out HeaderFile))
 									{
 										if (list == null)
+										{
 											list = new();
+										}
+
 										list.Add(new InvalidIncludeDirective(CppFile.Location, HeaderFile));
 									}
 								}
diff --git a/Engine/Source/Programs/UnrealBuildTool/Configuration/UEBuildTarget.cs b/Engine/Source/Programs/UnrealBuildTool/Configuration/UEBuildTarget.cs
--- a/Engine/Source/Programs/UnrealBuildTool/Configuration/UEBuildTarget.cs
+++ b/Engine/Source/Programs/UnrealBuildTool/Configuration/UEBuildTarget.cs
@@ -4440,7 +4440,9 @@ namespace UnrealBuildTool
 						return true;
 					}
 					return false;
-				}) != 0) ;
+				}) != 0)
+				{
+				}
 			}
 
 			// Merge remaining engine modules
diff --git a/Engine/Source/Programs/UnrealBuildTool/Executors/UnrealBuildAccelerator/UBAAgentCoordinatorHorde.cs b/Engine/Source/Programs/UnrealBuildTool/Executors/UnrealBuildAccelerator/UBAAgentCoordinatorHorde.cs
--- a/Engine/Source/Programs/UnrealBuildTool/Executors/UnrealBuildAccelerator/UBAAgentCoordinatorHorde.cs
+++ b/Engine/Source/Programs/UnrealBuildTool/Executors/UnrealBuildAccelerator/UBAAgentCoordinatorHorde.cs
@@ -624,7 +624,9 @@ namespace UnrealBuildTool
 						//arguments.Add("-NoStore"); // This means that no cas files will be written to disk.. so no state is stored between runs
 
 						if (!String.IsNullOrEmpty(desc))
+						{
 							arguments.Add($"-Description=\"{desc}\"");
+						}
 
 						if (_config.bLogEnabled)
 						{
diff --git a/Engine/Source/Programs/UnrealBuildTool/Executors/UnrealBuildAccelerator/UBAExecutor.cs b/Engine/Source/Programs/UnrealBuildTool/Executors/UnrealBuildAccelerator/UBAExecutor.cs
--- a/Engine/Source/Programs/UnrealBuildTool/Executors/UnrealBuildAccelerator/UBAExecutor.cs
+++ b/Engine/Source/Programs/UnrealBuildTool/Executors/UnrealBuildAccelerator/UBAExecutor.cs
@@ -316,7 +316,9 @@ namespace UnrealBuildTool
 					{
 						UnrealBuildAcceleratorCacheConfig? cacheConfig = coordinator.RequestCacheServer(_cancellationSource.Token);
 						if (cacheConfig != null)
+						{
 							cacheConfigs.Add(cacheConfig);
+						}
 					}
 				}
 			}
@@ -448,7 +450,10 @@ namespace UnrealBuildTool
 					foreach ((uint id, DirectoryReference vfs, DirectoryReference local) in action.RootPaths)
 					{
 						if (!isFirst)
+						{
 							await writer.WriteAsync(", ");
+						}
+
 						isFirst = false;
 						await writer.WriteAsync(vfs.FullName);
 						await writer.WriteAsync(';');
@@ -975,7 +980,10 @@ namespace UnrealBuildTool
 
 						scheduler.Cancel();
 						foreach (ICacheClient client in cache)
+						{
 							client.Disconnect();
+						}
+
 						server?.StopServer(); // Make sure all remove processes are returned. We can't have any callbacks after this
 						session?.CancelAll(); // Cancel all processes native side
 						storage?.SaveCasTable();
@@ -1065,7 +1073,9 @@ namespace UnrealBuildTool
 		{
 			string description = action.StatusDescription;
 			if (!String.IsNullOrEmpty(action.CommandDescription))
+			{
 				description = $"{action.StatusDescription} ({action.CommandDescription})";
+			}
 
 			// Let pch compile at higher priority
 			System.Diagnostics.ProcessPriorityClass priority = ProcessPriority;
@@ -1168,7 +1178,9 @@ namespace UnrealBuildTool
 		uint GetMemoryGroup(LinkedAction action)
 		{
 			if (action.ActionType != ActionType.Compile)
+			{
 				return 0;
+			}
 
 			FileItem? pch = action.ProducedItems.FirstOrDefault(i => i.HasExtension(".pch"));
 
diff --git a/Engine/Source/Programs/UnrealBuildTool/Modes/IWYUMode.cs b/Engine/Source/Programs/UnrealBuildTool/Modes/IWYUMode.cs
--- a/Engine/Source/Programs/UnrealBuildTool/Modes/IWYUMode.cs
+++ b/Engine/Source/Programs/UnrealBuildTool/Modes/IWYUMode.cs
@@ -777,7 +777,10 @@ namespace UnrealBuildTool
 						if (isGeneratedHeader)
 						{
 							if (generatedHeaders == null)
+							{
 								generatedHeaders = new();
+							}
+
 							generatedHeaders.Add(info);
 						}
 					}
@@ -800,11 +803,13 @@ namespace UnrealBuildTool
 							var exists = new HashSet<IWYUIncludeEntry>(mainInfo.Includes, new IWYUIncludeEntryFullPathComparer());
 							exists.UnionWith(headerAsMain!.Includes);
 							foreach (var i in headerInCpp.Includes)
+							{
 								if (!exists.Contains(i))
 								{
 									// TODO: Note that h and cpp might have different printable include paths since h and cpp are often not in same folder
 									mainInfo.Includes.Add(i);
 								}
+							}
 						}
 
 					}
@@ -823,8 +828,13 @@ namespace UnrealBuildTool
 					foreach (var generatedHeader in generatedHeaders)
 					{
 						foreach (var i in generatedHeader.Includes)
+						{
 							if (!exists.Contains(i))
+							{
 								mainInfo.Includes.Add(i);
+							}
+						}
+
 						generatedHeader.Includes = generatedHeader.IncludesSeenInFile;
 					}
 				});
@@ -1394,7 +1404,9 @@ namespace UnrealBuildTool
 		{
 			SortedDictionary<string, List<string>>? sortedLog = null;
 			if (!bWrite)
+			{
 				sortedLog = [];
+			}
 
 			List<ValueTuple<IWYUInfo, List<string>>> UpdatedFiles = new(); // <FilePath>, <NewLines>
 			HashSet<IWYUFile> SkippedFiles = new();
@@ -1994,7 +2006,10 @@ namespace UnrealBuildTool
 				{
 					Console.WriteLine(kv.Key);
 					foreach (string line in kv.Value)
+					{
 						Console.WriteLine(line);
+					}
+
 					Console.WriteLine();
 				}
 				return 0;
diff --git a/Engine/Source/Programs/UnrealBuildTool/Modes/TestMode.cs b/Engine/Source/Programs/UnrealBuildTool/Modes/TestMode.cs
--- a/Engine/Source/Programs/UnrealBuildTool/Modes/TestMode.cs
+++ b/Engine/Source/Programs/UnrealBuildTool/Modes/TestMode.cs
@@ -231,7 +231,9 @@ namespace UnrealBuildTool
 					CurrentLevelType = CurrentLevelType.BaseType!;
 				}
 				else
+				{
 					return FindMethod;
+				}
 			}
 			return null;
 		}
@@ -248,7 +250,9 @@ namespace UnrealBuildTool
 					CurrentLevelType = CurrentLevelType.BaseType!;
 				}
 				else
+				{
 					return FindField;
+				}
 			}
 			return null;
 		}
diff --git a/Engine/Source/Programs/UnrealBuildTool/Platform/Windows/VCToolChain.cs b/Engine/Source/Programs/UnrealBuildTool/Platform/Windows/VCToolChain.cs
--- a/Engine/Source/Programs/UnrealBuildTool/Platform/Windows/VCToolChain.cs
+++ b/Engine/Source/Programs/UnrealBuildTool/Platform/Windows/VCToolChain.cs
@@ -3003,7 +3003,10 @@ namespace UnrealBuildTool
 			CompileAction.ProducedItems.Add(ObjectFile);
 			CompileAction.ProducedItems.Add(FileItem.GetItemByFileReference(OutputFile));
 			if (OutputHeader != null)
+			{
 				CompileAction.ProducedItems.Add(FileItem.GetItemByFileReference(OutputHeader));
+			}
+
 			CompileAction.DeleteItems.Add(FileItem.GetItemByFileReference(OutputFile));
 			CompileAction.DeleteItems.Add(FileItem.GetItemByFileReference(OutputFile.ChangeExtension(".tli"))); // May not be created, depending on Attributes
 			CompileAction.StatusDescription = TypeLibrary.Header;
diff --git a/Engine/Source/Programs/UnrealBuildTool/System/ActionGraph.cs b/Engine/Source/Programs/UnrealBuildTool/System/ActionGraph.cs
--- a/Engine/Source/Programs/UnrealBuildTool/System/ActionGraph.cs
+++ b/Engine/Source/Programs/UnrealBuildTool/System/ActionGraph.cs
@@ -300,21 +300,31 @@ namespace UnrealBuildTool
 				bool ShouldWarn(FileItem item)
 				{
 					if (item.Location.FullName.Length <= Unreal.RootDirectory.FullName.Length + BuildConfiguration.MaxNestedPathLength)
+					{
 						return false;
+					}
 
 					if (!item.Location.IsUnderDirectory(Unreal.RootDirectory))
+					{
 						return false;
+					}
 
 					if (!item.Location.FullName.Contains("/Restricted/NotForLicensees/", StringComparison.OrdinalIgnoreCase)) //Be more relaxed for internal only code
+					{
 						return false;
+					}
 
 					if (item.Location.FullName.Contains("/Intermediate/", StringComparison.OrdinalIgnoreCase))
 					{
 						if (item.Location.ContainsName("H", 1)) // Ignore -IncludeHeader items
+						{
 							return false;
+						}
 
 						if (item.Location.FullName.EndsWith(".i.PVS-Studio.log")) // Ignore PVS Studio intermediate items
+						{
 							return false;
+						}
 					}
 
 					return true;
diff --git a/Engine/Source/Programs/UnrealBuildTool/System/FileHashCache.cs b/Engine/Source/Programs/UnrealBuildTool/System/FileHashCache.cs
--- a/Engine/Source/Programs/UnrealBuildTool/System/FileHashCache.cs
+++ b/Engine/Source/Programs/UnrealBuildTool/System/FileHashCache.cs
@@ -326,7 +326,10 @@ namespace UnrealBuildTool
 
 							XxHash128 xx = new();
 							while ((read = fs.Read(buffer, 0, buffer.Length)) > 0)
+							{
 								xx.Append(buffer.AsSpan(0, read));
+							}
+
 							isEqual = newHash == ToUInt128LittleEndian(xx.GetHashAndReset());
 						}
 						else if (data is string content)
diff --git a/Engine/Source/Programs/UnrealBuildTool/System/FileMetadataPrefetch.cs b/Engine/Source/Programs/UnrealBuildTool/System/FileMetadataPrefetch.cs
--- a/Engine/Source/Programs/UnrealBuildTool/System/FileMetadataPrefetch.cs
+++ b/Engine/Source/Programs/UnrealBuildTool/System/FileMetadataPrefetch.cs
@@ -90,7 +90,9 @@ namespace UnrealBuildTool
 				foreach (DirectoryItem? dir in Directories)
 				{
 					if (dir != null)
+					{
 						queue.Enqueue(() => ScanDirectoryTree(queue, dir));
+					}
 				}
 				queue.Drain();
 			});
diff --git a/Engine/Source/Programs/UnrealBuildTool/System/HotReload.cs b/Engine/Source/Programs/UnrealBuildTool/System/HotReload.cs
--- a/Engine/Source/Programs/UnrealBuildTool/System/HotReload.cs
+++ b/Engine/Source/Programs/UnrealBuildTool/System/HotReload.cs
@@ -1045,7 +1045,10 @@ namespace UnrealBuildTool
 
 								string NewSourceDependencyPath = NewSourceDependencyFile.FullName;
 								if (Action.RootPaths.GetVfsOverlayPath(NewSourceDependencyFile, out string? VirtualPath))
+								{
 									NewSourceDependencyPath = VirtualPath;
+								}
+
 								ResponseLines[Idx] = SourceDependencyPrefix + "\"" + NewSourceDependencyPath + "\"";
 								break;
 							}
diff --git a/Engine/Source/Programs/UnrealBuildTool/System/RulesAssembly.cs b/Engine/Source/Programs/UnrealBuildTool/System/RulesAssembly.cs
--- a/Engine/Source/Programs/UnrealBuildTool/System/RulesAssembly.cs
+++ b/Engine/Source/Programs/UnrealBuildTool/System/RulesAssembly.cs
@@ -810,21 +810,27 @@ namespace UnrealBuildTool
 					{
 						ReadOnlySpan<char> span = info.File.FullName.AsSpan(Unreal.EngineDirectory.FullName.Length + 1);
 						if (bloomFilter.MightContainSubString(span, Path.DirectorySeparatorChar))
+						{
 							mightContain = true;
+						}
 					}
 
 					if (!mightContain && ProjectDir != null && info.File.IsUnderDirectory(ProjectDir))
 					{
 						ReadOnlySpan<char> span = info.File.FullName.AsSpan(ProjectDir.FullName.Length + 1);
 						if (bloomFilter.MightContainSubString(span, Path.DirectorySeparatorChar))
+						{
 							mightContain = true;
+						}
 					}
 
 					if (mightContain)
 					{
 						if (info.File.ContainsAnyNames(DisallowedPlatformsAndGroups, Unreal.EngineDirectory) ||
 						(ProjectDir != null && info.File.ContainsAnyNames(DisallowedPlatformsAndGroups, ProjectDir)))
+						{
 							DisallowedPluginsList.Add(info);
+						}
 					}
 				}
 
