UnrealBuildTool: Code style fixes
Modifiers are not ordered (ide0036)
Avoid multiple blank lines (ide2000)
Embedded statements must be on their own line (ide2001)
Consecutive braces must not have blank line between them (ide2002)

#rnx

diff --git a/Engine/Source/Programs/UnrealBuildTool/Artifacts/Artifacts.cs b/Engine/Source/Programs/UnrealBuildTool/Artifacts/Artifacts.cs
--- a/Engine/Source/Programs/UnrealBuildTool/Artifacts/Artifacts.cs
+++ b/Engine/Source/Programs/UnrealBuildTool/Artifacts/Artifacts.cs
@@ -77,7 +77,7 @@ namespace UnrealBuildTool.Artifacts
 		public IArtifactDirectoryMapping? DirectoryMapping { get; set; } = null;
 
 		/// <inheritdoc/>
-		public readonly override int GetHashCode()
+		public override readonly int GetHashCode()
 		{
 			return HashCode.Combine(Key.GetHashCode(), ActionKey.GetHashCode(), Inputs.GetHashCode(), Outputs.GetHashCode());
 		}
diff --git a/Engine/Source/Programs/UnrealBuildTool/Configuration/TargetRules.cs b/Engine/Source/Programs/UnrealBuildTool/Configuration/TargetRules.cs
--- a/Engine/Source/Programs/UnrealBuildTool/Configuration/TargetRules.cs
+++ b/Engine/Source/Programs/UnrealBuildTool/Configuration/TargetRules.cs
@@ -540,7 +540,7 @@ namespace UnrealBuildTool
 	/// <summary>
 	/// Writable additional properties.
 	/// </summary>
-	sealed public class AdditionalProperties : ReadOnlyAdditionalProperties
+	public sealed class AdditionalProperties : ReadOnlyAdditionalProperties
 	{
 		/// <summary>
 		/// Set property
diff --git a/Engine/Source/Programs/UnrealBuildTool/Configuration/UEBuildModuleCPP.cs b/Engine/Source/Programs/UnrealBuildTool/Configuration/UEBuildModuleCPP.cs
--- a/Engine/Source/Programs/UnrealBuildTool/Configuration/UEBuildModuleCPP.cs
+++ b/Engine/Source/Programs/UnrealBuildTool/Configuration/UEBuildModuleCPP.cs
@@ -896,7 +896,6 @@ namespace UnrealBuildTool
 				}
 			}
 
-
 			// Engine modules will always use unity build mode unless MinSourceFilesForUnityBuildOverride is specified in
 			// the module rules file.  By default, game modules only use unity of they have enough source files for that
 			// to be worthwhile.  If you have a lot of small game modules, consider specifying MinSourceFilesForUnityBuildOverride=0
@@ -1003,7 +1002,6 @@ namespace UnrealBuildTool
 				LinkInputFiles.AddRange(CompileFilesWithToolChain(Target, ToolChain, GeneratedCPPCompileEnvironment, ModuleCompileEnvironment, NormalGeneratedFiles, AdaptiveGeneratedFiles, Graph, Logger).ObjectFiles);
 			}
 
-
 			// Compile Swift - which may generate a header that the cpp files need below
 			if (InputFiles.SwiftFiles.Count > 0 && Rules.SwiftInteropHeader != null)
 			{
@@ -1158,7 +1156,7 @@ namespace UnrealBuildTool
 				else
 				{
 					LinkInputFiles.AddRange(ToolChain.CompileAllCPPFiles(MMCompileEnvironment, InputFiles.MMFiles, IntermediateDirectory, Name, Graph).ObjectFiles);
-				};
+				}
 			}
 
 			// RC files for all modules in a binary will be compiled in UEBuildBinary.cs using the BinaryCompileEnvironment
@@ -1379,7 +1377,6 @@ namespace UnrealBuildTool
 			using Pooled<HashSet<string>> rented3 = HashSetPool<string>.Rent();
 			HashSet<string> IncludesNotFound = rented3.Value;
 
-
 			// Go through user directories and create bloom filter for files and dirs so we can early out quickly
 			List<(DirectoryItem, BloomFilter, BloomFilter)> UserIncludePaths = new(CompileEnvironment.UserIncludePaths.Count);
 			foreach (DirectoryReference UserInclude in CompileEnvironment.UserIncludePaths)
@@ -1502,8 +1499,6 @@ namespace UnrealBuildTool
 				}
 			}
 
-
-
 			bool FoundAllModules = true;
 			HashSet<UEBuildModule> OptModules = new();
 			foreach (FileItem IncludeFile in FoundIncludeFileItems)
@@ -2040,7 +2035,6 @@ namespace UnrealBuildTool
 					AdaptiveUnityEnvironment.bSupportEditAndContinue = true;
 				}
 
-
 				// Create a per-file PCH
 				CPPOutput AdaptiveOutput;
 				if (CompileEnvironment.PrecompiledHeaderAction == PrecompiledHeaderAction.Include && Target.bAdaptiveUnityCreatesDedicatedPCH)
@@ -2771,7 +2765,7 @@ namespace UnrealBuildTool
 							throw new Exception("BuildSettingsDefinitions can't contain keys or values with character ';'");
 						}
 
-						sb.Append(key).Append(';').Append(value).Append(';'); ;
+						sb.Append(key).Append(';').Append(value).Append(';');
 					}
 
 					Result.Definitions.Add($"UE_BUILDSETTINGS_VARIABLES=\"{sb}\"");
diff --git a/Engine/Source/Programs/UnrealBuildTool/Configuration/UEBuildTarget.cs b/Engine/Source/Programs/UnrealBuildTool/Configuration/UEBuildTarget.cs
--- a/Engine/Source/Programs/UnrealBuildTool/Configuration/UEBuildTarget.cs
+++ b/Engine/Source/Programs/UnrealBuildTool/Configuration/UEBuildTarget.cs
@@ -2711,15 +2711,15 @@ namespace UnrealBuildTool
 
 		internal class BuildContext
 		{
-			required internal SourceFileMetadataCache MetadataCache;
-			required internal TargetDescriptor Descriptor;
-			required internal TargetMakefile Makefile;
-			required internal List<UEBuildBinary> OriginalBinaries;
-			required internal HashSet<UEBuildModuleCPP> ModulesToGenerateHeadersFor;
-			required internal CppCompileEnvironment GlobalCompileEnvironment;
-			required internal UEToolChain TargetToolChain;
-			required internal LinkEnvironment GlobalLinkEnvironment;
-			required internal List<GenerateHeaderFunc> GenerateHeaderFuncs;
+			internal required SourceFileMetadataCache MetadataCache;
+			internal required TargetDescriptor Descriptor;
+			internal required TargetMakefile Makefile;
+			internal required List<UEBuildBinary> OriginalBinaries;
+			internal required HashSet<UEBuildModuleCPP> ModulesToGenerateHeadersFor;
+			internal required CppCompileEnvironment GlobalCompileEnvironment;
+			internal required UEToolChain TargetToolChain;
+			internal required LinkEnvironment GlobalLinkEnvironment;
+			internal required List<GenerateHeaderFunc> GenerateHeaderFuncs;
 		}
 
 		internal BuildContext BuildMakefile1(BuildConfiguration BuildConfiguration, TargetDescriptor TargetDescriptor, ILogger Logger, SourceFileMetadataCache MetadataCache)
@@ -2860,7 +2860,6 @@ namespace UnrealBuildTool
 				}
 			}
 
-
 			return new BuildContext()
 			{
 				MetadataCache = MetadataCache,
@@ -2877,10 +2876,10 @@ namespace UnrealBuildTool
 
 		internal class CodeGenContext
 		{
-			required internal TargetDescriptor Descriptor;
-			required internal TargetMakefile Makefile;
-			required internal string TargetName;
-			required internal List<GenerateHeaderFunc> GenerateHeaderFuncs;
+			internal required TargetDescriptor Descriptor;
+			internal required TargetMakefile Makefile;
+			internal required string TargetName;
+			internal required List<GenerateHeaderFunc> GenerateHeaderFuncs;
 			internal Task? uhtRun;
 			internal Task? uhtTimestamp;
 			internal Task? vniRun;
diff --git a/Engine/Source/Programs/UnrealBuildTool/Executors/UnrealBuildAccelerator/UBAExecutor.cs b/Engine/Source/Programs/UnrealBuildTool/Executors/UnrealBuildAccelerator/UBAExecutor.cs
--- a/Engine/Source/Programs/UnrealBuildTool/Executors/UnrealBuildAccelerator/UBAExecutor.cs
+++ b/Engine/Source/Programs/UnrealBuildTool/Executors/UnrealBuildAccelerator/UBAExecutor.cs
@@ -95,7 +95,6 @@ namespace UnrealBuildTool
 			}
 		}
 
-
 		public override string Name => !UBAConfig.bDisableRemote && _agentCoordinators.Count > 0 ? "Unreal Build Accelerator" : "Unreal Build Accelerator local";
 
 		public static new bool IsAvailable() => EpicGames.UBA.Utils.IsAvailable();
@@ -369,7 +368,6 @@ namespace UnrealBuildTool
 			{
 				_session.DisableRemoteExecution();
 			}
-
 		}
 
 		public void UpdateStatus(uint statusRow, uint statusColumn, string statusText, LogEntryType statusType, string? statusLink)
@@ -534,7 +532,6 @@ namespace UnrealBuildTool
 				_cacheableActionsLeft = inputActions.Count(x => x.ArtifactMode.HasFlag(ArtifactMode.Enabled));
 			}
 
-
 			await WriteActionOutputFileAsync(logger, inputActions);
 
 			try
@@ -621,7 +618,6 @@ namespace UnrealBuildTool
 
 		public override bool VerifyOutputs => UBAConfig.bWriteToDisk;
 
-
 		(byte[]?, uint) BuildKnownInputs(LinkedAction action)
 		{
 			if (!UBAConfig.bUseKnownInputs)
@@ -662,7 +658,12 @@ namespace UnrealBuildTool
 
 		class AgentCoordinatorScheduler : IUbaAgentCoordinatorScheduler
 		{
-			internal AgentCoordinatorScheduler(IScheduler scheduler, IServer server) { _scheduler = scheduler; _server = server; }
+			internal AgentCoordinatorScheduler(IScheduler scheduler, IServer server)
+			{
+				_scheduler = scheduler;
+				_server = server;
+			}
+
 			internal IScheduler? _scheduler;
 			internal IServer? _server;
 			private readonly object _lockObject = new();
@@ -783,7 +784,7 @@ namespace UnrealBuildTool
 								}
 								++localRetryActions;
 								return ProcessFinishedResponse.RerunLocal;
-							};
+							}
 
 							if (exitCode != 0 && !logLines.Any())
 							{
diff --git a/Engine/Source/Programs/UnrealBuildTool/Modes/IWYUMode.cs b/Engine/Source/Programs/UnrealBuildTool/Modes/IWYUMode.cs
--- a/Engine/Source/Programs/UnrealBuildTool/Modes/IWYUMode.cs
+++ b/Engine/Source/Programs/UnrealBuildTool/Modes/IWYUMode.cs
@@ -484,7 +484,6 @@ namespace UnrealBuildTool
 					}
 				}
 
-
 				Queue<UEBuildModule> toUpdateQueue = [];
 				HashSet<UEBuildModule> toUpdateHandled = [];
 				if (UEModuleToUpdate != null)
@@ -493,7 +492,6 @@ namespace UnrealBuildTool
 					toUpdateHandled.Add(UEModuleToUpdate);
 				}
 
-
 				foreach (UEBuildModule Module in NameToModule.Values)
 				{
 					string ModuleDir = Module.ModuleDirectory.FullName.Replace('\\', '/');
@@ -714,7 +712,7 @@ namespace UnrealBuildTool
 				});
 				parsingIwyu.Finish();
 
-				ConcurrentBag<(IWYUInfo mainInfo, List<IWYUInfo> generatedHeaders)> generatedHeadersToFixup = new();;
+				ConcurrentBag<(IWYUInfo mainInfo, List<IWYUInfo> generatedHeaders)> generatedHeadersToFixup = new();
 
 				// Now traverse main infos and find additional infos and fix them up
 				Logger.LogInformation($"Adding additional iwyu infos...");
@@ -811,7 +809,6 @@ namespace UnrealBuildTool
 								}
 							}
 						}
-
 					}
 				});
 
@@ -1045,7 +1042,6 @@ namespace UnrealBuildTool
 						}
 					}
 				}
-
 			});
 			resolvingIncludes.Finish();
 
@@ -2318,7 +2314,6 @@ namespace UnrealBuildTool
 			Logger.LogInformation("----------------------------------------------------------");
 			int End = Math.Min(20, ByteSizeDiffList.Count);
 
-
 			int LongestName = 0;
 			for (int I = 0; I != End; ++I)
 			{
diff --git a/Engine/Source/Programs/UnrealBuildTool/Platform/Mac/MacToolChain.cs b/Engine/Source/Programs/UnrealBuildTool/Platform/Mac/MacToolChain.cs
--- a/Engine/Source/Programs/UnrealBuildTool/Platform/Mac/MacToolChain.cs
+++ b/Engine/Source/Programs/UnrealBuildTool/Platform/Mac/MacToolChain.cs
@@ -786,7 +786,7 @@ namespace UnrealBuildTool
 							RPathEntry = UnstagedBundlePrefix + DestDir.MakeRelativeTo(ReferencingFile.Directory);
 							LinkArguments.Add($"-rpath \"@loader_path/{RPathEntry}\"");
 						}
-					};
+					}
 
 					// first look in the RuntimeDependency list to look for remapped locations, anything not remapped we will deal with 
 					// after this loop
diff --git a/Engine/Source/Programs/UnrealBuildTool/ProjectFiles/VisualStudio/VCProject.cs b/Engine/Source/Programs/UnrealBuildTool/ProjectFiles/VisualStudio/VCProject.cs
--- a/Engine/Source/Programs/UnrealBuildTool/ProjectFiles/VisualStudio/VCProject.cs
+++ b/Engine/Source/Programs/UnrealBuildTool/ProjectFiles/VisualStudio/VCProject.cs
@@ -495,12 +495,12 @@ namespace UnrealBuildTool
 
 		class ProjectConfigAndTargetCombination
 		{
-			readonly public UnrealTargetPlatform? Platform;
-			readonly public UnrealTargetConfiguration Configuration;
-			readonly public string ProjectPlatformName;
-			readonly public string ProjectConfigurationName;
-			readonly public ProjectTarget? ProjectTarget;
-			readonly public UnrealArch? Architecture;
+			public readonly UnrealTargetPlatform? Platform;
+			public readonly UnrealTargetConfiguration Configuration;
+			public readonly string ProjectPlatformName;
+			public readonly string ProjectConfigurationName;
+			public readonly ProjectTarget? ProjectTarget;
+			public readonly UnrealArch? Architecture;
 
 			public ProjectConfigAndTargetCombination(UnrealTargetPlatform? InPlatform, UnrealTargetConfiguration InConfiguration, string InProjectPlatformName, string InProjectConfigurationName, ProjectTarget? InProjectTarget, UnrealArch? InArchitecture)
 			{
@@ -1567,7 +1567,7 @@ namespace UnrealBuildTool
 							string NewValue = String.Join(";", PathArray);
 							CommonProjectFileContent.AppendLine($"    <{kvp.Value}>{PropertyValuePrefix}{NewValue}</{kvp.Value}>");
 						}
-					};
+					}
 
 					// Use short names on these because there are _a ton_ of them in the big projects. the vcxproj of our biggest project went from 70mb to 31mb
 
diff --git a/Engine/Source/Programs/UnrealBuildTool/ProjectFiles/VisualStudio/VCProjectFileGenerator.cs b/Engine/Source/Programs/UnrealBuildTool/ProjectFiles/VisualStudio/VCProjectFileGenerator.cs
--- a/Engine/Source/Programs/UnrealBuildTool/ProjectFiles/VisualStudio/VCProjectFileGenerator.cs
+++ b/Engine/Source/Programs/UnrealBuildTool/ProjectFiles/VisualStudio/VCProjectFileGenerator.cs
@@ -955,7 +955,7 @@ namespace UnrealBuildTool
 									// Recurse into subfolders
 									FolderProcessorFunction(LocalVCSolutionFileContent, CurFolder.SubFolders);
 								}
-							};
+							}
 						FolderProcessorFunction(VCSolutionFileContent, RootFolder.SubFolders);
 
 						VCSolutionFileContent.AppendLine("	EndGlobalSection");
@@ -1126,7 +1126,7 @@ namespace UnrealBuildTool
 									Architecture = Arch
 								}
 							);
-						};
+						}
 
 						UnrealArchitectures? Architectures = GetPlatformArchitecturesToGenerate(BuildPlatform, ProjectTarget);
 						if (Architectures == null)
diff --git a/Engine/Source/Programs/UnrealBuildTool/ProjectFiles/VisualStudioCode/VSCodeProjectFileGenerator.cs b/Engine/Source/Programs/UnrealBuildTool/ProjectFiles/VisualStudioCode/VSCodeProjectFileGenerator.cs
--- a/Engine/Source/Programs/UnrealBuildTool/ProjectFiles/VisualStudioCode/VSCodeProjectFileGenerator.cs
+++ b/Engine/Source/Programs/UnrealBuildTool/ProjectFiles/VisualStudioCode/VSCodeProjectFileGenerator.cs
@@ -1044,7 +1044,7 @@ namespace UnrealBuildTool
 								OutFile.AddField("cwd", MakeUnquotedPathString(ProjectRoot, EPathType.Absolute));
 								OutFile.BeginObject("env");
 								{
-									OutFile.AddField("PATH", $"{Unreal.DotnetDirectory}{Path.PathSeparator}${{env:PATH}}");;
+									OutFile.AddField("PATH", $"{Unreal.DotnetDirectory}{Path.PathSeparator}${{env:PATH}}");
 									OutFile.AddField("DOTNET_MULTILEVEL_LOOKUP", "0");
 									OutFile.AddField("DOTNET_ROLL_FORWARD", "LatestMajor");
 								}
@@ -1057,7 +1057,8 @@ namespace UnrealBuildTool
 				}
 			}
 
-			if (InProject.Name.Equals("AutomationTool", StringComparison.OrdinalIgnoreCase)) {
+			if (InProject.Name.Equals("AutomationTool", StringComparison.OrdinalIgnoreCase))
+			{
 				WriteAutomationScriptTask(InProject, OutFile);
 			}
 		}
@@ -1517,7 +1518,8 @@ namespace UnrealBuildTool
 						string TaskName = String.Format("{0} ({1})", Target.Name, BuildProduct.Config);
 						string BuildTaskName = String.Format("{0} {1} {2} Build", Target.Name, HostPlatform, BuildProduct.Config);
 						string[]? Args = null;
-						if (Target.Name.Equals("AutomationTool", StringComparison.OrdinalIgnoreCase)) {
+						if (Target.Name.Equals("AutomationTool", StringComparison.OrdinalIgnoreCase))
+						{
 							BuildTaskName = String.Format("{0}+Scripts {1} {2} Build", Target.Name, HostPlatform, BuildProduct.Config);
 							Args = new string[] {"-NoCompile"};
 						}
diff --git a/Engine/Source/Programs/UnrealBuildTool/ProjectFiles/VisualStudioWorkspace/VSWorkspaceProjectFile.cs b/Engine/Source/Programs/UnrealBuildTool/ProjectFiles/VisualStudioWorkspace/VSWorkspaceProjectFile.cs
--- a/Engine/Source/Programs/UnrealBuildTool/ProjectFiles/VisualStudioWorkspace/VSWorkspaceProjectFile.cs
+++ b/Engine/Source/Programs/UnrealBuildTool/ProjectFiles/VisualStudioWorkspace/VSWorkspaceProjectFile.cs
@@ -160,7 +160,6 @@ namespace UnrealBuildTool
 								ExportedTargetProjects.Add(TargetInfo);
 							}
 						}
-
 					}
 					catch (Exception Ex)
 					{
@@ -450,7 +449,7 @@ namespace UnrealBuildTool
 					}
 				}
 
-				return string.Format("{0}{1}_{2}", Target.Platform.ToString(), Architecture != null ? "_" + Architecture.ToString() : string.Empty, Target.Configuration.ToString()); ;
+				return string.Format("{0}{1}_{2}", Target.Platform.ToString(), Architecture != null ? "_" + Architecture.ToString() : string.Empty, Target.Configuration.ToString());
 			}
 
 			return string.Empty;
diff --git a/Engine/Source/Programs/UnrealBuildTool/ProjectFiles/VisualStudioWorkspace/VSWorkspaceProjectFileGenerator.cs b/Engine/Source/Programs/UnrealBuildTool/ProjectFiles/VisualStudioWorkspace/VSWorkspaceProjectFileGenerator.cs
--- a/Engine/Source/Programs/UnrealBuildTool/ProjectFiles/VisualStudioWorkspace/VSWorkspaceProjectFileGenerator.cs
+++ b/Engine/Source/Programs/UnrealBuildTool/ProjectFiles/VisualStudioWorkspace/VSWorkspaceProjectFileGenerator.cs
@@ -366,7 +366,7 @@ namespace UnrealBuildTool
 			{ "NMakeReBuildCommandLine", "$(_UProjectNMakeReBuildCommandLine)" }
 		};
 
-		private readonly static HashSet<string> NodesAllowList = new HashSet<string>
+		private static readonly HashSet<string> NodesAllowList = new HashSet<string>
 		{
 			"CleanDependsOn",
 			"Configuration",
@@ -390,7 +390,7 @@ namespace UnrealBuildTool
 			"xml"
 		};
 
-		private readonly static HashSet<string> NodesBlockedList = new HashSet<string>
+		private static readonly HashSet<string> NodesBlockedList = new HashSet<string>
 		{
 			"AdditionalOptions",
 			"ClCompile",
diff --git a/Engine/Source/Programs/UnrealBuildTool/System/SourceFileMetadataCache.cs b/Engine/Source/Programs/UnrealBuildTool/System/SourceFileMetadataCache.cs
--- a/Engine/Source/Programs/UnrealBuildTool/System/SourceFileMetadataCache.cs
+++ b/Engine/Source/Programs/UnrealBuildTool/System/SourceFileMetadataCache.cs
@@ -211,7 +211,7 @@ namespace UnrealBuildTool
 
 					bModified = true;
 					return SourceFileInfo;
-				};
+				}
 
 				return FileToSourceFileInfo.AddOrUpdate(SourceFile, _ =>
 				{
