Do not generate source-maps for Xcode projects if VFS is not used

#rnx

diff --git a/Engine/Source/Programs/UnrealBuildTool/ProjectFiles/Xcode/XcodeProject.cs b/Engine/Source/Programs/UnrealBuildTool/ProjectFiles/Xcode/XcodeProject.cs
--- a/Engine/Source/Programs/UnrealBuildTool/ProjectFiles/Xcode/XcodeProject.cs
+++ b/Engine/Source/Programs/UnrealBuildTool/ProjectFiles/Xcode/XcodeProject.cs
@@ -577,6 +577,16 @@ namespace UnrealBuildTool.XcodeProjectXcconfig
 				}
 			}
 
+			if (TargetRules?.bUseVFS == true)
+			{
+				CppRootPaths RootPaths = new() { bUseVfs = TargetRules.bUseVFS };
+				RootPaths[CppRootPathFolder.Root] = Unreal.RootDirectory;
+				if (UProjectFileLocation != null)
+				{
+					RootPaths[CppRootPathFolder.Project] = UProjectFileLocation.Directory;
+				}
+			}
+
 			return true;
 		}
 
@@ -673,8 +683,6 @@ namespace UnrealBuildTool.XcodeProjectXcconfig
 			FileBatch.bEnableRTTI = CompileEnvironment.bUseRTTI;
 			FileBatch.SystemIncludePaths.UnionWith(CompileEnvironment.SystemIncludePaths);
 			FileBatch.UserIncludePaths.UnionWith(CompileEnvironment.UserIncludePaths);
-
-			RootPaths = CompileEnvironment.RootPaths;
 		}
 
 		private List<UnrealBuildConfig> GetSupportedBuildConfigs(List<UnrealTargetPlatform> Platforms, List<UnrealTargetConfiguration> Configurations, ILogger Logger)
