[UBT]
* Changed so FileHashCache always check file matching before writing

diff --git a/Engine/Source/Programs/UnrealBuildTool/System/FileHashCache.cs b/Engine/Source/Programs/UnrealBuildTool/System/FileHashCache.cs
--- a/Engine/Source/Programs/UnrealBuildTool/System/FileHashCache.cs
+++ b/Engine/Source/Programs/UnrealBuildTool/System/FileHashCache.cs
@@ -262,13 +262,11 @@ namespace UnrealBuildTool
 		{
 			Partition? partition = GetPartition(fileItem.Location, logger);
 
-			bool inDatabase = false;
 			bool writtenBefore = false;
 			if (partition != null)
 			{
 				if (partition.Files.TryGetValue(fileItem, out FileEntry fileEntry))
 				{
-					inDatabase = true;
 					writtenBefore = fileEntry.written;
 
 					// If hash is the same we check last write time
@@ -301,79 +299,76 @@ namespace UnrealBuildTool
 			}
 			if (allowAsync)
 			{
-				return Task.Run(() => InternalWriteFileIfChanged2(inDatabase, writtenBefore, partition, fileItem, newHash, data, logger, allowAsync, alwaysUpdateTimestamp));
+				return Task.Run(() => InternalWriteFileIfChanged2(writtenBefore, partition, fileItem, newHash, data, logger, allowAsync, alwaysUpdateTimestamp));
 			}
 
-			InternalWriteFileIfChanged2(inDatabase, writtenBefore, partition, fileItem, newHash, data, logger, allowAsync, alwaysUpdateTimestamp);
+			InternalWriteFileIfChanged2(writtenBefore, partition, fileItem, newHash, data, logger, allowAsync, alwaysUpdateTimestamp);
 			return null;
 		}
 
-		private void InternalWriteFileIfChanged2(bool inDatabase, bool writtenBefore, Partition? partition, FileItem fileItem, UInt128 newHash, object data, ILogger logger, bool allowAsync, bool alwaysUpdateTimestamp)
+		private void InternalWriteFileIfChanged2(bool writtenBefore, Partition? partition, FileItem fileItem, UInt128 newHash, object data, ILogger logger, bool allowAsync, bool alwaysUpdateTimestamp)
 		{
-			if (!inDatabase)
+			if (fileItem.Exists)
 			{
-				if (fileItem.Exists)
+				try
 				{
-					try
+					bool isEqual = false;
+					if (data is MemoryStream)
 					{
-						bool isEqual = false;
-						if (data is MemoryStream)
-						{
-							const int BufferSize = 1024 * 1024; // 1 MB
-							byte[] buffer = new byte[BufferSize];
-							using FileStream fs = new(fileItem.FullName, FileMode.Open, FileAccess.Read, FileShare.Read, BufferSize, FileOptions.SequentialScan);
-							int read;
+						const int BufferSize = 1024 * 1024; // 1 MB
+						byte[] buffer = new byte[BufferSize];
+						using FileStream fs = new(fileItem.FullName, FileMode.Open, FileAccess.Read, FileShare.Read, BufferSize, FileOptions.SequentialScan);
+						int read;
 
-							XxHash128 xx = new();
-							while ((read = fs.Read(buffer, 0, buffer.Length)) > 0)
-							{
-								xx.Append(buffer.AsSpan(0, read));
-							}
-
-							isEqual = newHash == ToUInt128LittleEndian(xx.GetHashAndReset());
-						}
-						else if (data is string content)
+						XxHash128 xx = new();
+						while ((read = fs.Read(buffer, 0, buffer.Length)) > 0)
 						{
-							string oldContent = FileReference.ReadAllText(fileItem.Location, new UTF8Encoding(false));
-							isEqual = String.Equals(oldContent, content, StringComparison.Ordinal);
+							xx.Append(buffer.AsSpan(0, read));
 						}
-						else if (data is IEnumerable<string> contentLines)
+
+						isEqual = newHash == ToUInt128LittleEndian(xx.GetHashAndReset());
+					}
+					else if (data is string content)
+					{
+						string oldContent = FileReference.ReadAllText(fileItem.Location, new UTF8Encoding(false));
+						isEqual = String.Equals(oldContent, content, StringComparison.Ordinal);
+					}
+					else if (data is IEnumerable<string> contentLines)
+					{
+						string[] oldLines = FileReference.ReadAllLines(fileItem.Location, new UTF8Encoding(false));
+						int index = 0;
+						isEqual = true;
+						foreach (string str in contentLines)
 						{
-							string[] oldLines = FileReference.ReadAllLines(fileItem.Location, new UTF8Encoding(false));
-							int index = 0;
-							isEqual = true;
-							foreach (string str in contentLines)
+							if (index >= oldLines.Length || str != oldLines[index++])
 							{
-								if (index >= oldLines.Length || str != oldLines[index++])
-								{
-									isEqual = false;
-									break;
-								}
+								isEqual = false;
+								break;
 							}
 						}
+					}
 
-						if (isEqual)
+					if (isEqual)
+					{
+						DateTime lastWrite = fileItem.LastWriteTimeUtc;
+						if (alwaysUpdateTimestamp)
 						{
-							DateTime lastWrite = fileItem.LastWriteTimeUtc;
-							if (alwaysUpdateTimestamp)
-							{
-								lastWrite = DateTime.UtcNow;
-								File.SetLastWriteTimeUtc(fileItem.FullName, lastWrite);
-								fileItem.ResetCachedInfo();
-							}
-							if (partition != null)
+							lastWrite = DateTime.UtcNow;
+							File.SetLastWriteTimeUtc(fileItem.FullName, lastWrite);
+							fileItem.ResetCachedInfo();
+						}
+						if (partition != null)
+						{
+							if (partition.Files.TryAdd(fileItem, new() { hash = newHash, lastWriteTimeUtc = lastWrite.Ticks }))
 							{
-								if (partition.Files.TryAdd(fileItem, new() { hash = newHash, lastWriteTimeUtc = lastWrite.Ticks }))
-								{
-									partition.Modified = true;
-								}
+								partition.Modified = true;
 							}
-							return;
 						}
+						return;
 					}
-					catch
-					{
-					}
+				}
+				catch
+				{
 				}
 			}
 
