[IWYU]
* Lots of fixes to handle edge case scenarios in IWYU
* Force Cpp23 when building with iwyu since it is more strict on template instantiations (so includes are not removed that shouldn't be removed

diff --git a/Engine/Source/Programs/UnrealBuildTool/Configuration/TargetRules.cs b/Engine/Source/Programs/UnrealBuildTool/Configuration/TargetRules.cs
--- a/Engine/Source/Programs/UnrealBuildTool/Configuration/TargetRules.cs
+++ b/Engine/Source/Programs/UnrealBuildTool/Configuration/TargetRules.cs
@@ -3287,6 +3287,8 @@ namespace UnrealBuildTool
 				bUseVFS = false; // Do not use vfs because IWYUMode can't handle virtual paths (yet)
 				bUseUnityBuild = false; // No unity for iwyu
 				bUseVerseBPVM = false;
+				CppStandard = CppStandardVersion.Cpp23; // Must be set to 23 to handle function default values same as msvc
+				CppStandardEngine = CppStandardVersion.Cpp23; // Must be set to 23 to handle function default values same as msvc
 				// bAlwaysUseUnityForGeneratedFiles = false; // This should be used but adds too many extra files atm.
 			}
 
diff --git a/Engine/Source/Programs/UnrealBuildTool/Modes/IWYUMode.cs b/Engine/Source/Programs/UnrealBuildTool/Modes/IWYUMode.cs
--- a/Engine/Source/Programs/UnrealBuildTool/Modes/IWYUMode.cs
+++ b/Engine/Source/Programs/UnrealBuildTool/Modes/IWYUMode.cs
@@ -27,19 +27,19 @@ using UnrealBuildBase;
 // 
 //  1. Will build lyra editor target using include-what-you-use.exe instead of clang.exe, and preview how iwyu would update module LyraGame and all the modules depending on it
 //  
-//   .build target LyraEditor linux development -- -Mode=IWYU -ModuleToUpdate=LyraGame -UpdateDependents
+//   .build target LyraEditor win64 development -- -Mode=IWYU -ModuleToUpdate=LyraGame -UpdateDependents
 //   
 //  2. Same as 1 but will check out files from p4 and modify them
 //  
-//   .build target LyraEditor linux development -- -Mode=IWYU -ModuleToUpdate=LyraGame -UpdateDependents -Write
+//   .build target LyraEditor win64 development -- -Mode=IWYU -ModuleToUpdate=LyraGame -UpdateDependents -Write
 //   
 //  3. Will build and then update all modules that has Lyra/Plugins in its path. Note it will only include modules that are part of LyraEditor target.
 // 
-//   .build target LyraEditor linux development -- -Mode=IWYU -PathToUpdate=Lyra/Plugins -Write
+//   .build target LyraEditor win64 development -- -Mode=IWYU -PathToUpdate=Lyra/Plugins -Write
 //   
 //  4. Update Niagara plugins private code (not public api) with preview
 // 
-//   .build target UnrealEditor linux development -- -Mode=IWYU -PathToUpdate=Engine/Plugins/FX/Niagara -UpdateOnlyPrivate
+//   .build target UnrealEditor win64 development -- -Mode=IWYU -PathToUpdate=Engine/Plugins/FX/Niagara -UpdateOnlyPrivate
 
 namespace UnrealBuildTool
 {
@@ -58,6 +58,11 @@ namespace UnrealBuildTool
 		/// </summary>
 		public string Printable { get; set; } = "";
 
+		/// <summary>
+		/// Reason include exists
+		/// </summary>
+		public string Reason { get; set; } = "";
+
 		/// <summary>
 		/// Reference to file info. Not populated with .iwyu file
 		/// </summary>
@@ -72,6 +77,11 @@ namespace UnrealBuildTool
 		/// If true #include was found inside a declaration such as a namespace or class/struct/function
 		/// </summary>
 		public bool InsideDecl { get; set; }
+
+		public bool InPath(string path)
+		{
+			return Full.StartsWith(path, StringComparison.OrdinalIgnoreCase);
+		}
 	}
 
 	/// <summary>
@@ -131,34 +141,34 @@ namespace UnrealBuildTool
 		/// <summary>
 		/// Includes that this file needs. This is what iwyu have decided is used.
 		/// </summary>
-		public List<IWYUIncludeEntry> Includes { get; set; } = new List<IWYUIncludeEntry>();
+		public List<IWYUIncludeEntry> Includes { get; set; } = [];
 
 		/// <summary>
 		/// Forward declarations that this file needs. This is what iwyu have decided is used.
 		/// </summary>
-		public List<IWYUForwardEntry> ForwardDeclarations { get; set; } = new List<IWYUForwardEntry>();
+		public List<IWYUForwardEntry> ForwardDeclarations { get; set; } = [];
 
 		/// <summary>
 		/// Includes seen in the file. This is how it looked like on disk.
 		/// </summary>
-		public List<IWYUIncludeEntry> IncludesSeenInFile { get; set; } = new List<IWYUIncludeEntry>();
+		public List<IWYUIncludeEntry> IncludesSeenInFile { get; set; } = [];
 
 		/// <summary>
 		/// Includes that didnt end up in the Includes list but is needed by the code file
 		/// </summary>
-		public List<IWYUMissingInclude> MissingIncludes { get; set; } = new List<IWYUMissingInclude>();
+		public List<IWYUMissingInclude> MissingIncludes { get; set; } = [];
 
 		/// <summary>
 		/// Transitive includes. This is all the includes that someone gets for free when including this file
 		/// Note, this is based on what iwyu believes should be in all includes. So it will not look at "seen includes"
 		/// and only use its generated list.
 		/// </summary>
-		public Dictionary<string, string> TransitiveIncludes = new();
+		public Dictionary<string, string> TransitiveIncludes = [];
 
 		/// <summary>
 		/// Transitive forward declarations. Same as transitive includes
 		/// </summary>
-		public Dictionary<string, string> TransitiveForwardDeclarations = new();
+		public Dictionary<string, string> TransitiveForwardDeclarations = [];
 
 		/// <summary>
 		/// Which .iwyu file that produced this info. Note, it might be null for some special cases (like .generated.h files)
@@ -197,6 +207,13 @@ namespace UnrealBuildTool
 		public string? Name;
 	}
 
+	sealed class IWYUUnsupportedFormattingException : Exception
+	{
+		internal IWYUUnsupportedFormattingException(string? message, Exception? exception = null) : base(message, exception)
+		{
+		}
+	}
+
 	/// <summary>
 	/// Profiles different unity sizes and prints out the different size and its timings
 	/// </summary>
@@ -428,15 +445,20 @@ namespace UnrealBuildTool
 			string TargetName = $"{TargetDescriptors[0].Name}_{TargetDescriptors[0].Configuration}_{TargetDescriptors[0].Platform}";
 
 			// Calculate file paths filter to figure out which files that IWYU will run on
-			HashSet<string> ValidPaths = new();
-			Dictionary<string, UEBuildModule> PathToModule = new();
-			Dictionary<string, UEBuildModule> NameToModule = new();
+			HashSet<string> ValidPaths = [];
+			Dictionary<string, UEBuildModule> PathToModule = [];
+			Dictionary<string, UEBuildModule> NameToModule = [];
+			List<string> systemPaths = [];
 			{
 				using ITimelineEvent _ = Timeline.ScopeEvent("CreatingTargets");
 				// Create target to be able to traverse all existing modules
 				Logger.LogInformation($"Creating BuildTarget for {TargetName}...");
 				UEBuildTarget Target = UEBuildTarget.Create(TargetDescriptors[0], BuildConfiguration, Logger);
 
+				var env = Target.Rules.WindowsPlatform.Environment!;
+				systemPaths.Add(env.ToolChainDir.FullName.Replace('\\', '/'));
+				systemPaths.Add(env.WindowsSdkDir.FullName.Replace('\\', '/'));
+
 				Logger.LogInformation($"Calculating file filter for IWYU...");
 
 				UEBuildModule? UEModuleToUpdate = null;
@@ -586,7 +608,7 @@ namespace UnrealBuildTool
 			}
 
 			Dictionary<string, IWYUInfo> Infos = new();
-			ConcurrentBag<IWYUInfo> GeneratedCppInfos = new();
+			ConcurrentBag<(IWYUInfo, List<IWYUIncludeEntry>)> GeneratedCppInfos = [];
 
 			using (ISourceFileWorkingSet WorkingSet = new EmptySourceFileWorkingSet())
 			{
@@ -724,6 +746,7 @@ namespace UnrealBuildTool
 					IWYUInfo? headerAsMain = null;
 					IWYUInfo? headerInCpp = null;
 					List<IWYUInfo>? generatedHeaders = null;
+					List<IWYUInfo>? generatedSourceFiles = null;
 
 					// Go through all infos in iwyu
 					foreach (IWYUInfo info in IWYUFile.Files)
@@ -746,10 +769,18 @@ namespace UnrealBuildTool
 							break;
 						}
 
-						if (isCpp && info.File.Contains(".gen.cpp", StringComparison.Ordinal))
+						bool isGeneratedCpp = isCpp && info.File.EndsWith(".gen.cpp", StringComparison.Ordinal);
+
+						if (!isGeneratedCpp && info.File.EndsWith(".ipp", StringComparison.Ordinal))
+						{
+							isGeneratedCpp = true;
+						}
+
+						if (isGeneratedCpp)
 						{
 							info.Source = IWYUFile;
-							GeneratedCppInfos.Add(info);
+							generatedSourceFiles ??= new();
+							generatedSourceFiles.Add(info);
 							continue;
 						}
 
@@ -785,11 +816,56 @@ namespace UnrealBuildTool
 
 					// Fix up main file
 
-					if (!mainInfo.IsCpp & generatedHeaders != null)
+					if (!mainInfo.IsCpp && generatedHeaders != null)
 					{
 						generatedHeadersToFixup.Add((mainInfo!, generatedHeaders!));
 					}
 
+					// We want to resolve reference to generated source already here to make sure transitive includes end up correctly
+					// .. so we will swap Includes and IncludesSeenInFile
+					if (generatedSourceFiles != null)
+					{
+						foreach (IWYUInfo generatedFile in generatedSourceFiles)
+						{
+							GeneratedCppInfos.Add((generatedFile, generatedFile.Includes));
+							generatedFile.Includes = generatedFile.IncludesSeenInFile;
+						}
+						if (mainInfo.IsCpp)
+						{
+							List<IWYUIncludeEntry> generatedIncludes = [];
+							foreach (IWYUIncludeEntry include in mainInfo.IncludesSeenInFile)
+							{
+								foreach (IWYUInfo generatedFile in generatedSourceFiles)
+								{
+									if (include.Full == generatedFile.File)
+									{
+										include.Resolved = generatedFile;
+										generatedIncludes.Add(include);
+										break;
+									}
+								}
+							}
+
+							foreach (IWYUIncludeEntry generatedInclude in generatedIncludes)
+							{
+								bool included = false;
+								foreach (IWYUIncludeEntry include in mainInfo.Includes)
+								{
+									if (include.Full == generatedInclude.Full)
+									{
+										include.Resolved = generatedInclude.Resolved;
+										included = true;
+										break;
+									}
+								}
+								if (!included)
+								{
+									mainInfo.Includes.Add(generatedInclude);
+								}
+							}
+						}
+					}
+
 					// This is annoying. But for some reason list of files needed in header can be different
 					// based on if it builds standalone or with the cpp file.
 					// This should not happen if header is first in include list.
@@ -812,6 +888,8 @@ namespace UnrealBuildTool
 					}
 				});
 
+				// Generated headers will never be modified and will keep the includes seen in file.
+				// We also add all the includes the generated files have into the owning file
 				Parallel2.ForEach(generatedHeadersToFixup, (entry) =>
 				{
 					IWYUInfo mainInfo = entry.mainInfo;
@@ -869,6 +947,11 @@ namespace UnrealBuildTool
 				return new(Include, Entry);
 			}
 
+			// Always keep
+			// Windows/WindowsHWrapper.h
+			// Windows/AllowWindowsPlatformTypes.h
+			// Windows/HideWindowsPlatformTypes.h
+
 			Dictionary<string, IWYUIncludeEntry?> SpecialIncludes = new Dictionary<string, IWYUIncludeEntry?>(
 				new[]
 				{
@@ -879,6 +962,9 @@ namespace UnrealBuildTool
 					GetInfo("Misc/OptionalFwd.h", "Source/Runtime/Core/Public/Misc/OptionalFwd.h"),
 					GetInfo("Templates/SharedPointerFwd.h", "Source/Runtime/Core/Public/Templates/SharedPointerFwd.h"),
 					GetInfo("Misc/ExpressionParserTypesFwd.h", "Source/Runtime/Core/Public/Misc/ExpressionParserTypesFwd.h"),
+					GetInfo("Misc/NotNull.h", "Source/Runtime/Core/Public/Misc/NotNull.h"),
+					GetInfo("UObject/StrongObjectPtrTemplatesFwd.h", "Source/Runtime/Core/Public/UObject/StrongObjectPtrTemplatesFwd.h"),
+					GetInfo("Delegates/Delegate.h", "Source/Runtime/Core/Public/Delegates/Delegate.h"),
 				}
 			);
 
@@ -891,8 +977,13 @@ namespace UnrealBuildTool
 				{ "TOptional", SpecialIncludes["Misc/OptionalFwd.h"] },
 				{ "TSharedPtr", SpecialIncludes["Templates/SharedPointerFwd.h"] },
 				{ "TSharedRef", SpecialIncludes["Templates/SharedPointerFwd.h"] },
+				{ "TWeakPtr", SpecialIncludes["Templates/SharedPointerFwd.h"] },
+				{ "TSharedFromThis", SpecialIncludes["Templates/SharedPointerFwd.h"] },
 				{ "TCompiledToken", SpecialIncludes["Misc/ExpressionParserTypesFwd.h"] },
 				{ "TExpressionToken", SpecialIncludes["Misc/ExpressionParserTypesFwd.h"] },
+				{ "TStrongObjectPtr", SpecialIncludes["UObject/StrongObjectPtrTemplatesFwd.h"] },
+				{ "TNotNull", SpecialIncludes["Misc/NotNull.h"] },
+				{ "TMulticastDelegateRegistration", SpecialIncludes["Delegates/Delegate.h"] },
 			};
 			/*
 			// Add all .generated.h files as entries in the lookup and explicitly add the includes they have which will never be removed
@@ -937,9 +1028,6 @@ namespace UnrealBuildTool
 			Logger.LogInformation($"Resolving Includes...");
 			using ITimelineEvent resolvingIncludes = Timeline.ScopeEvent("ResolvingIncludes");
 
-			LinuxPlatformSDK PlatformSDK = new LinuxPlatformSDK(Logger);
-			DirectoryReference? BaseLinuxPath = PlatformSDK.GetBaseLinuxPathForArchitecture(LinuxPlatform.DefaultHostArchitecture);
-			string SystemPath = BaseLinuxPath!.FullName.Replace('\\', '/');
 
 			ConcurrentDictionary<IWYUInfo, IWYUInfo> IncludedInsideDecl = new();
 			Dictionary<string, IWYUInfo> SkippedHeaders = new();
@@ -982,11 +1070,14 @@ namespace UnrealBuildTool
 				}
 
 				// We don't want to mess around with third party includes.. since we can't see the hierarchy we just assumes that they are optimally included
-				Info.Includes.RemoveAll(Entry => Entry.Full.Contains("/ThirdParty/", StringComparison.Ordinal) || Entry.Full.StartsWith(SystemPath));
+				Info.Includes.RemoveAll(Entry =>
+				{
+					return Entry.Full.Contains("/ThirdParty/") || systemPaths.Any(systemPath => Entry.InPath(systemPath));
+				});
 
 				foreach (IWYUIncludeEntry Entry in Info.IncludesSeenInFile)
 				{
-					if (Entry.Full.Contains("/ThirdParty/", StringComparison.Ordinal) || Entry.Full.StartsWith(SystemPath, StringComparison.Ordinal))
+					if (Entry.Full.Contains("/ThirdParty/") || systemPaths.Any(systemPath => Entry.InPath(systemPath)))
 					{
 						Info.Includes.Add(Entry);
 					}
@@ -996,8 +1087,19 @@ namespace UnrealBuildTool
 						IncludedInsideDecl.TryAdd(Entry.Resolved, Info);
 						Info.Includes.Add(Entry);
 					}
+
+					// We don't want to remove these because they enable includes of windows headers.
+					if (Entry.Printable.StartsWith("Windows/AllowWindows") || Entry.Printable.StartsWith("Windows/HideWindows"))
+					{
+						Info.Includes.Add(Entry);
+					}
 				}
 
+				// This is workaround for iwyu bug. If you include STableRow.h you also need SListViewEntry.h
+				// Need to fix iwyu bug but also I want the dependency in the code to go away
+				IWYUIncludeEntry? STableRowEntry = null;
+				IWYUIncludeEntry? SListViewEntry = null;
+
 				foreach (IWYUIncludeEntry Include in Info.Includes)
 				{
 					if (Include.Resolved == null)
@@ -1018,9 +1120,19 @@ namespace UnrealBuildTool
 								}
 							}
 						}
+
+						STableRowEntry ??= Include.Printable == "Widgets/Views/STableRow.h" ? Include : null;
+						SListViewEntry ??= Include.Printable == "Widgets/Views/SListView.h" ? Include : null;
 					}
 				}
 
+				if (STableRowEntry != null && SListViewEntry == null)
+				{
+					SListViewEntry = new IWYUIncludeEntry() { Full = STableRowEntry.Full.Replace("STableRow.h", "SListView.h"), Printable = "Widgets/Views/SListView.h" };
+					Info.Includes.Add(SListViewEntry);
+					Infos.TryGetValue(SListViewEntry.Full, out SListViewEntry.Resolved);
+				}
+
 				foreach (IWYUMissingInclude MissingInclude in Info.MissingIncludes)
 				{
 					if (Info.TransitiveIncludes.ContainsKey(MissingInclude.Full))
@@ -1043,6 +1155,16 @@ namespace UnrealBuildTool
 					}
 				}
 			});
+
+			// Resolving includes inside generated files.
+			Parallel2.ForEach(GeneratedCppInfos, entry =>
+			{
+				foreach (IWYUIncludeEntry include in entry.Item1.Includes)
+				{
+					Infos.TryGetValue(include.Full, out include.Resolved);
+				}
+				
+			});
 			resolvingIncludes.Finish();
 
 			using ITimelineEvent parsingSkipped = Timeline.ScopeEvent("ParsingSkippedHeaders");
@@ -1078,31 +1200,21 @@ namespace UnrealBuildTool
 			// If we have built .gen.cpp it means that it is not inlined in another cpp file
 			// And we need to promote the includes needed to the header since we can never modify the .gen.cpp file
 			Logger.LogInformation($"Transferring needed includes from .gen.cpp to owning file...");
-			Parallel2.ForEach(GeneratedCppInfos, GeneratedCpp =>
+			Parallel2.ForEach(GeneratedCppInfos, entry =>
 			{
-				// First, check which files .gen.cpp will see
-				HashSet<string> SeenTransitiveIncludes = new();
-				foreach (IWYUIncludeEntry SeenInclude in GeneratedCpp.IncludesSeenInFile)
-				{
-					SeenTransitiveIncludes.Add(SeenInclude.Full);
-					foreach (IWYUIncludeEntry Include in GeneratedCpp.Includes)
-					{
-						IWYUInfo? IncludeInfo;
-						if (SeenInclude.Printable == Include.Printable && Infos.TryGetValue(Include.Full, out IncludeInfo))
-						{
-							foreach (string I in IncludeInfo.TransitiveIncludes.Keys)
-							{
-								SeenTransitiveIncludes.Add(I);
-							}
+				IWYUInfo GeneratedCpp = entry.Item1;
 
-							break;
-						}
-					}
-				}
+				// Calculate transitive includes for .gen.cpp/ipp files because they are not in Infos.Values list
+				StringBuilder stack = stacks.Get();
+				GeneratedCpp.TransitiveIncludes.EnsureCapacity(10);
+				GeneratedCpp.TransitiveForwardDeclarations.EnsureCapacity(10);
+				CalculateTransitive(GeneratedCpp, GeneratedCpp, stack, GeneratedCpp.TransitiveIncludes, GeneratedCpp.TransitiveForwardDeclarations, false);
+				stacks.Return(stack);
 
 				IWYUInfo? IncluderInfo = null;
 
 				// If there is only one file in .iwyu it means that .gen.cpp was compiled separately
+				// And we need to make sure IncluderInfo has all the includes the generated cpp file needs
 				if (GeneratedCpp.Source!.Files.Count == 1)
 				{
 					int NameStart = GeneratedCpp.File.LastIndexOf('/');
@@ -1117,41 +1229,51 @@ namespace UnrealBuildTool
 							break;
 						}
 					}
+
+					if (IncluderInfo == null)
+					{
+						return;
+					}
 				}
 				else // this .gen.cpp is inlined in a cpp..
 				{
 					IncluderInfo = GeneratedCpp.Source.Files.FirstOrDefault(I => I.File.Contains(".cpp") && !I.File.Contains(".gen"));
-				}
 
-				if (IncluderInfo == null)
-				{
-					return;
+					if (IncluderInfo == null)
+					{
+						return;
+					}
 				}
 
-				foreach (IWYUIncludeEntry Include in GeneratedCpp.Includes)
+
+				foreach (IWYUIncludeEntry Include in entry.Item2)
 				{
-					if (SeenTransitiveIncludes.Contains(Include.Full))
-					{
-						continue;
-					}
+					//if (neededIncludes.Contains(Include.Full))
+					//{
+					//	continue;
+					//}
 
 					// TODO: Remove UObject check once we've added "IWYU pragma: keep" around the includes in ScriptMacros and ObjectMacros
 
-					if (Include.Full.Contains(".generated.h") || Include.Printable.StartsWith("UObject/"))
+					//if (Include.Full.Contains(".generated.h") || Include.Printable.StartsWith("UObject/"))
+					//{
+					//	continue;
+					//}
+
+					// If file included by .gen.cpp (or file .gen.cpp is inlined into) does not see needed include, then we need to add it
+					if (IncluderInfo!.TransitiveIncludes.ContainsKey(Include.Full))
 					{
 						continue;
 					}
 
-					if (!IncluderInfo!.TransitiveIncludes.ContainsKey(Include.Full))
+					bool isExcluded = Include.Full.Contains("/ThirdParty/") || systemPaths.Any(systemPath => Include.InPath(systemPath));
+					if (!isExcluded)
 					{
-						if (!Include.Full.Contains("/ThirdParty/") && !Include.Full.StartsWith(SystemPath))
+						if (Include.Resolved == null)
 						{
-							if (Include.Resolved == null)
-							{
-								Infos.TryGetValue(Include.Full, out Include.Resolved);
-							}
-							IncluderInfo.Includes.Add(Include);
+							Infos.TryGetValue(Include.Full, out Include.Resolved);
 						}
+						IncluderInfo.Includes.Add(Include);
 					}
 				}
 			});
@@ -1385,6 +1507,12 @@ namespace UnrealBuildTool
 				return false;
 			}
 
+			// Skip ThirdParty files if there is a thirdparty sub folder inside a module
+			if (Info.File.Contains("/ThirdParty/"))
+			{
+				return false;
+			}
+
 			// Filter out files
 			foreach (string ValidPath in ValidPaths)
 			{
@@ -1414,7 +1542,7 @@ namespace UnrealBuildTool
 			uint FilesParseCount = 0;
 			int ProcessSuccess = 1;
 
-			Parallel2.ForEach(Infos.Values, Info =>
+			void UpdateFile(IWYUInfo Info)
 			{
 				if (!IsValidForUpdate(Info, ValidPaths, true))
 				{
@@ -1440,9 +1568,9 @@ namespace UnrealBuildTool
 					MatchingH = Info.File.Substring(LastSlash, Info.File.Length - LastSlash - 4) + ".h";
 				}
 
-				Dictionary<string, string> TransitivelyIncluded = new();
-				SortedSet<string> CleanedupIncludes = new();
-				SortedSet<string> ForwardDeclarationsToAdd = new();
+				Dictionary<string, string> TransitivelyIncluded = [];
+				SortedSet<string> CleanedupIncludes = [];
+				SortedSet<string> ForwardDeclarationsToAdd = [];
 
 				foreach (IWYUIncludeEntry Include in Info.Includes)
 				{
@@ -1454,6 +1582,11 @@ namespace UnrealBuildTool
 						NameWithoutPath = NameWithoutPath.Substring(LastSlash + 1);
 					}
 
+					if (IsCpp && Include.Full.EndsWith(".generated.h"))
+					{
+						continue;
+					}
+
 					string QuotedPrintable = Include.System ? $"<{Include.Printable}>" : $"\"{Include.Printable}\"";
 
 					bool Keep = true;
@@ -1533,6 +1666,11 @@ namespace UnrealBuildTool
 				// Read all lines of the header/source file
 				string[] ExistingLines = File.ReadAllLines(Info.File);
 
+				if (ExistingLines.Length > 1 && ExistingLines[0].Contains("automatically generated"))
+				{
+					return;
+				}
+
 				SortedDictionary<string, string> LinesToRemove = new();
 				SortedSet<string> IncludesToAdd = new(CleanedupIncludes);
 
@@ -1573,6 +1711,7 @@ namespace UnrealBuildTool
 					}
 				}
 
+				int includeCounter = 0;
 				// Traverse all lines in file and figure out which includes that should be added or removed
 				foreach (string Line in ExistingLines)
 				{
@@ -1626,6 +1765,8 @@ namespace UnrealBuildTool
 						return;
 					}
 
+					++includeCounter;
+
 					HasIncludes = true;
 
 					bool ForceKeep = false;
@@ -1645,7 +1786,7 @@ namespace UnrealBuildTool
 						{
 							int Open = IncludeSpan.IndexOf('(') + 1;
 							int Close = IncludeSpan.IndexOf(')');
-							string ActualInclude = $"\"Linux/Linux{IncludeSpan.Slice(Open, Close - Open).ToString()}\"";
+							string ActualInclude = $"\"Windows/Windows{IncludeSpan.Slice(Open, Close - Open).ToString()}\"";
 							IncludesToAdd.Remove(ActualInclude);
 						}
 						else
@@ -1689,7 +1830,13 @@ namespace UnrealBuildTool
 					}
 					else
 					{
-						IncludesToAdd.Remove(Include);
+						if (!IncludesToAdd.Remove(Include))
+						{
+							// Include can be relative path of cpp file
+							if (Info.IsCpp && includeCounter == 1) // First include can have a different path
+							{
+							}
+						}
 					}
 				}
 
@@ -1874,8 +2021,12 @@ namespace UnrealBuildTool
 
 					if (LinesToAdd.Count > 0)
 					{
+						if (LastSeenIncludeBeforeCode == -1)
+						{
+							throw new IWYUUnsupportedFormattingException($"Code structure to complicated for IWYU to update {Info.File}. Definitions before includes?");
+						}
 						int InsertPos = LastSeenIncludeBeforeCode + 1;
-						if (NewLines[LastSeenIncludeBeforeCode].Contains(".generated.h", StringComparison.Ordinal))
+						if (NewLines[LastSeenIncludeBeforeCode].Contains(".generated.h", StringComparison.Ordinal) || NewLines[LastSeenIncludeBeforeCode].Contains(".gen.h", StringComparison.Ordinal))
 						{
 							--InsertPos;
 						}
@@ -1978,6 +2129,24 @@ namespace UnrealBuildTool
 				{
 					UpdatedFiles.Add(new(Info, NewLines));
 				}
+			}
+
+			Parallel2.ForEach(Infos.Values, Info =>
+			{
+				try
+				{
+					UpdateFile(Info);
+				}
+				catch (IWYUUnsupportedFormattingException ex)
+				{
+					Logger.LogError(ex, ex.Message);
+					Interlocked.Exchange(ref ProcessSuccess, 0);
+				}
+				catch (Exception ex)
+				{
+					Logger.LogError(ex, "Exception while processing {File} {Exception}", Info.File, ex.ToString());
+					Interlocked.Exchange(ref ProcessSuccess, 0);
+				}
 			});
 
 			// Something went wrong processing code files.
@@ -2233,17 +2402,6 @@ namespace UnrealBuildTool
 				{
 					AllSeenIncludes.TryAdd(Include.Key, 0);
 				}
-
-				/*
-				if (Info.TransitiveIncludes.Count < SeenTransitiveIncludes.Count)
-				{
-					float Ratio = (float)Info.TransitiveIncludes.Count / SeenTransitiveIncludes.Count;
-					lock (RatioList)
-						RatioList.Add(new(Ratio, Info));
-				}
-				if (SeenTransitiveIncludes.Count < Info.TransitiveIncludes.Count)
-					Console.WriteLine($"{Info.File} - Now: {SeenTransitiveIncludes.Count}  Optimized: {Info.TransitiveIncludes.Count}");
-				*/
 			});
 
 			Dictionary<string, long> FileToSize = new Dictionary<string, long>(AllSeenIncludes);
diff --git a/Engine/Source/Programs/UnrealBuildTool/ToolChain/IWYUToolChain.cs b/Engine/Source/Programs/UnrealBuildTool/ToolChain/IWYUToolChain.cs
--- a/Engine/Source/Programs/UnrealBuildTool/ToolChain/IWYUToolChain.cs
+++ b/Engine/Source/Programs/UnrealBuildTool/ToolChain/IWYUToolChain.cs
@@ -4,6 +4,7 @@ using System;
 using System.Collections.Generic;
 using System.IO;
 using System.Linq;
+using System.Runtime.InteropServices;
 using EpicGames.Core;
 using Microsoft.Extensions.Logging;
 using UnrealBuildBase;
@@ -30,6 +31,15 @@ namespace UnrealBuildTool
 
 		public IWYUToolChain(ReadOnlyTargetRules target, ISPCToolChain internalToolChain, ILogger InLogger) : base(ClangToolChainOptions.None, InLogger)
 		{
+			if (!OperatingSystem.IsWindows())
+			{
+				throw new BuildException("IWYU only works on win64 x64 platform");
+			}
+			if (target.Platform != UnrealTargetPlatform.Win64)
+			{
+				throw new BuildException("IWYU only works on win64 targets. Epic is using UnrealEditor win64 development when running iwyu");
+			}
+
 			TargetPlatform = target.Platform;
 			InnerToolChain = internalToolChain;
 
@@ -246,7 +256,18 @@ namespace UnrealBuildTool
 				{
 					foreach (FileItem InlinedFile in InlinedFiles)
 					{
-						Arguments.Add($"-Xiwyu --check_also={InlinedFile.FullName.Replace('\\', '/')}");
+						Arguments.Add($"-Xiwyu --check_also=\"{InlinedFile.FullName.Replace('\\', '/')}\"");
+					}
+				}
+
+				if (SourceFile.Name != "SingleFile.cpp" && !SourceFile.Name.EndsWith(".gen.cpp"))
+				{
+					foreach (string include in CompileEnvironment.MetadataCache.GetHeaderIncludes(SourceFile))
+					{
+						if (include.EndsWith(".ipp"))
+						{
+							Arguments.Add($"-Xiwyu --check_also=*/{include.Replace('\\', '/')}");
+						}
 					}
 				}
 			}
