Unshelved from pending changelist '47833442':

[UBT]
* Fixed static analysis warning in IWYU mode

diff --git a/Engine/Source/Programs/UnrealBuildTool/Modes/IWYUMode.cs b/Engine/Source/Programs/UnrealBuildTool/Modes/IWYUMode.cs
--- a/Engine/Source/Programs/UnrealBuildTool/Modes/IWYUMode.cs
+++ b/Engine/Source/Programs/UnrealBuildTool/Modes/IWYUMode.cs
@@ -207,7 +207,10 @@ namespace UnrealBuildTool
 		public string? Name;
 	}
 
-	sealed class IWYUUnsupportedFormattingException : Exception
+	/// <summary>
+	/// Exception when formatting of file not supported
+	/// </summary>
+	public class IWYUUnsupportedFormattingException : Exception
 	{
 		internal IWYUUnsupportedFormattingException(string? message, Exception? exception = null) : base(message, exception)
 		{
