In UBT's Analyze mode, make sure to include circular dependencies, as it's possible that the only reference to them is a dependency marked as circular, so failing to include them can cause UBT to crash with a KeyNotFoundException.

diff --git a/Engine/Source/Programs/UnrealBuildTool/Modes/AnalyzeMode.cs b/Engine/Source/Programs/UnrealBuildTool/Modes/AnalyzeMode.cs
--- a/Engine/Source/Programs/UnrealBuildTool/Modes/AnalyzeMode.cs
+++ b/Engine/Source/Programs/UnrealBuildTool/Modes/AnalyzeMode.cs
@@ -93,7 +93,7 @@ namespace UnrealBuildTool.Modes
 			{
 				// Now we need to recheck all downstream dependencies again because the chain may be shorter
 				List<UEBuildModule> RecheckModules = [];
-				Module.GetAllDependencyModules(RecheckModules, [], true, false, true);
+				Module.GetAllDependencyModules(RecheckModules, [], bIncludeDynamicallyLoaded: true, bForceCircular: true, bOnlyDirectDependencies: true);
 				Visited.ExceptWith(RecheckModules);
 				Logger.LogDebug("Found shorter chain for {Module} {Prev} -> {New}, rechecking {Count} already visited dependencies", Module.Name, moduleInfo.IncludeChain.Length, CurrentChain.Count, RecheckModules.Count);
 				moduleInfo.IncludeChain = [.. CurrentChain];
@@ -106,7 +106,7 @@ namespace UnrealBuildTool.Modes
 			Visited.Add(Module);
 
 			List<UEBuildModule> TargetModules = [];
-			Module.GetAllDependencyModules(TargetModules, [], true, false, true);
+			Module.GetAllDependencyModules(TargetModules, [], bIncludeDynamicallyLoaded: true, bForceCircular: true, bOnlyDirectDependencies: true);
 			TargetModules.ForEach(x => AnalyzeModuleChains(x.Name, CurrentChain, Target, ModuleToInfo, Visited, Logger));
 		}
 
@@ -171,11 +171,11 @@ namespace UnrealBuildTool.Modes
 			foreach ((UEBuildModule SourceModule, ModuleInfo SourceModuleInfo) in ModuleToInfo)
 			{
 				SourceModuleInfo.OutwardRefs.Add(SourceModule);
-				SourceModule.GetAllDependencyModules([], SourceModuleInfo.OutwardRefs, false, false, false);
+				SourceModule.GetAllDependencyModules([], SourceModuleInfo.OutwardRefs, bIncludeDynamicallyLoaded: false, bForceCircular: true, bOnlyDirectDependencies: false);
 				SourceModuleInfo.OutwardRefs.Remove(SourceModule);
 			}
 
-			// Find the direct output dependencies of each module
+			// Find the unique output dependencies of each module
 			foreach ((UEBuildModule SourceModule, ModuleInfo SourceModuleInfo) in ModuleToInfo)
 			{
 				SourceModuleInfo.UniqueOutwardRefs = [.. SourceModuleInfo.OutwardRefs];
@@ -185,7 +185,7 @@ namespace UnrealBuildTool.Modes
 					HashSet<UEBuildModule> VisitedTargetModules = [SourceModule];
 
 					List<UEBuildModule> DependencyModules = [];
-					TargetModule.GetAllDependencyModules(DependencyModules, VisitedTargetModules, false, false, false);
+					TargetModule.GetAllDependencyModules(DependencyModules, VisitedTargetModules, bIncludeDynamicallyLoaded: false, bForceCircular: true, bOnlyDirectDependencies: false);
 					DependencyModules.Remove(TargetModule);
 
 					SourceModuleInfo.UniqueOutwardRefs.ExceptWith(DependencyModules);
@@ -419,16 +419,6 @@ namespace UnrealBuildTool.Modes
 
 			GraphVisualization.WriteGraphFile(FileName, $"Module dependency graph for {Target.TargetName}", [.. NameToNode.Values], Edges);
 		}
-
-		private static HashSet<UEBuildModule> GetDirectDependencyModules(UEBuildModule Module)
-		{
-			HashSet<UEBuildModule> ReferencedModules = [];
-			Module.GetAllDependencyModules([], ReferencedModules, true, false, false);
-
-			HashSet<UEBuildModule> Modules = [.. Module.GetDirectDependencyModules()];
-			Modules.ExceptWith(ReferencedModules);
-			return Modules;
-		}
 	}
 }
 
