Remove references to the now-removed bundled LibCxx, replacing with references to ones from the in-use compiler toolchain where appropriate.

diff --git a/Engine/Source/Programs/UnrealBuildTool/Platform/Linux/LinuxPlatformSDK.cs b/Engine/Source/Programs/UnrealBuildTool/Platform/Linux/LinuxPlatformSDK.cs
--- a/Engine/Source/Programs/UnrealBuildTool/Platform/Linux/LinuxPlatformSDK.cs
+++ b/Engine/Source/Programs/UnrealBuildTool/Platform/Linux/LinuxPlatformSDK.cs
@@ -92,11 +92,6 @@ namespace UnrealBuildTool
 		/// </summary>
 		private int bForceUseSystemCompiler = -1;
 
-		/// <summary>
-		/// Force using legacy LibCxx in Engine/Source/ThirdParty/Unix
-		/// </summary>
-		private int bForceUseLegacyLibCxx = -1;
-
 		/// <summary>
 		/// Whether to compile with the verbose flag
 		/// </summary>
@@ -186,30 +181,6 @@ namespace UnrealBuildTool
 			return bForceUseSystemCompiler == 1;
 		}
 
-		/// <summary>
-		/// Force use of Legacy LibCxx in Engine/Source/ThirdParty/Unix, which will be deprecated in the future
-		/// </summary>
-		public virtual bool ForceUseLegacyLibCxx()
-		{
-			// by default tools chains don't parse arguments, but we want to be able to check the -bForceUseSystemCompiler flag.
-			if (bForceUseLegacyLibCxx == -1)
-			{
-				bForceUseLegacyLibCxx = 0;
-				string[] CmdLine = Environment.GetCommandLineArgs();
-
-				foreach (string CmdLineArg in CmdLine)
-				{
-					if (CmdLineArg.Equals("-ForceUseLegacyLibCxx", StringComparison.OrdinalIgnoreCase))
-					{
-						bForceUseLegacyLibCxx = 1;
-						break;
-					}
-				}
-			}
-
-			return bForceUseLegacyLibCxx == 1;
-		}
-
 		/// <summary>
 		/// Returns the root SDK path for all architectures
 		/// WARNING: Do not cache this value - it may be changed after sourcing OutputEnvVars.txt
diff --git a/Engine/Source/Programs/UnrealBuildTool/Platform/Linux/LinuxProjectGenerator.cs b/Engine/Source/Programs/UnrealBuildTool/Platform/Linux/LinuxProjectGenerator.cs
--- a/Engine/Source/Programs/UnrealBuildTool/Platform/Linux/LinuxProjectGenerator.cs
+++ b/Engine/Source/Programs/UnrealBuildTool/Platform/Linux/LinuxProjectGenerator.cs
@@ -3,6 +3,7 @@
 using System;
 using System.Collections.Generic;
 using System.IO;
+using System.Linq;
 using System.Text.RegularExpressions;
 using EpicGames.Core;
 using Microsoft.Extensions.Logging;
@@ -43,22 +44,25 @@ namespace UnrealBuildTool
 		/// <inheritdoc/>
 		public override IList<string> GetSystemIncludePaths(UEBuildTarget InTarget)
 		{
-			List<string> Result = new List<string>();
+			List<string> Result = [];
 			string EngineDirectory = Unreal.EngineDirectory.ToString();
 			string? UseLibcxxEnvVarOverride = Environment.GetEnvironmentVariable("UE_LINUX_USE_LIBCXX");
+			bool ShouldUseLibcxx = string.IsNullOrEmpty(UseLibcxxEnvVarOverride) || UseLibcxxEnvVarOverride == "1";
+
 			UnrealArch TargetArchitecture = InTarget.Architectures.SingleArchitecture;
-			if (String.IsNullOrEmpty(UseLibcxxEnvVarOverride) || UseLibcxxEnvVarOverride == "1")
-			{
-				if (TargetArchitecture == UnrealArch.X64 || TargetArchitecture == UnrealArch.Arm64)
-				{
-					// libc++ include directories
-					Result.Add(Path.Combine(EngineDirectory, "Source/ThirdParty/Unix/LibCxx/include/"));
-					Result.Add(Path.Combine(EngineDirectory, "Source/ThirdParty/Unix/LibCxx/include/c++/v1"));
-				}
-			}
 
 			UEBuildPlatformSDK BuildPlatformSdk = UEBuildPlatform.GetSDK(InTarget.Platform)!;
+			LinuxPlatformSDK? LinuxPlatformSdk = BuildPlatformSdk as LinuxPlatformSDK;
 			string? InternalSdkPath = BuildPlatformSdk.GetInternalSDKPath();
+			string? BaseLinuxPath = InternalSdkPath ?? LinuxPlatformSdk?.GetBaseLinuxPathForArchitecture(TargetArchitecture)?.FullName;
+
+			if (ShouldUseLibcxx && BaseLinuxPath is not null)
+			{
+				// libc++ include directories
+				Result.Add(Path.Combine(BaseLinuxPath, "include"));
+				Result.Add(Path.Combine(BaseLinuxPath, "include", "c++", "v1"));
+			}
+
 			if (InternalSdkPath != null)
 			{
 				string PlatformSdkVersionString = BuildPlatformSdk.GetInstalledVersion()!;
@@ -75,7 +79,7 @@ namespace UnrealBuildTool
 				}
 			}
 
-			return Result;
+			return [.. Result.Distinct()];
 		}
 
 		/// <summary>
diff --git a/Engine/Source/Programs/UnrealBuildTool/Platform/Linux/LinuxToolChain.cs b/Engine/Source/Programs/UnrealBuildTool/Platform/Linux/LinuxToolChain.cs
--- a/Engine/Source/Programs/UnrealBuildTool/Platform/Linux/LinuxToolChain.cs
+++ b/Engine/Source/Programs/UnrealBuildTool/Platform/Linux/LinuxToolChain.cs
@@ -453,23 +453,12 @@ namespace UnrealBuildTool
 			if (ShouldUseLibcxx())
 			{
 				Arguments.Add("-nostdinc++");
-				if (PlatformSDK.ForceUseLegacyLibCxx())
+				// LinuxInfo.BaseLinuxPath should only be null in the case that we've been forced to use the system compiler.
+				// In that case, project build.cs files should be adjusted to include the relevant directories to the system include paths
+				if (!PlatformSDK.ForceUseSystemCompiler() && LinuxInfo.BaseLinuxPath != null)
 				{
-					Arguments.Add(GetSystemIncludePathArgument(DirectoryReference.Combine(Unreal.EngineSourceDirectory, "ThirdParty", "Unix", "LibCxx", "include"), CompileEnvironment));
-					Arguments.Add(GetSystemIncludePathArgument(DirectoryReference.Combine(Unreal.EngineSourceDirectory, "ThirdParty", "Unix", "LibCxx", "include", "c++", "v1"), CompileEnvironment));
-				}
-				else
-				{
-					if (!PlatformSDK.ForceUseSystemCompiler())
-					{
-						// LinuxInfo.BaseLinuxPath should only be null in the case that we've been forced to use the system compiler.
-						// In that case, project build.cs files should be adjusted to include the relevant directories to the system include paths
-						if(LinuxInfo.BaseLinuxPath != null)
-						{
-							Arguments.Add(GetSystemIncludePathArgument(DirectoryReference.Combine(LinuxInfo.BaseLinuxPath, "include"), CompileEnvironment));
-							Arguments.Add(GetSystemIncludePathArgument(DirectoryReference.Combine(LinuxInfo.BaseLinuxPath, "include", "c++", "v1"), CompileEnvironment));
-						}
-					}
+					Arguments.Add(GetSystemIncludePathArgument(DirectoryReference.Combine(LinuxInfo.BaseLinuxPath, "include"), CompileEnvironment));
+					Arguments.Add(GetSystemIncludePathArgument(DirectoryReference.Combine(LinuxInfo.BaseLinuxPath, "include", "c++", "v1"), CompileEnvironment));
 				}
 			}
 
@@ -1283,10 +1272,6 @@ namespace UnrealBuildTool
 				// libc++ and its abi lib
 				LinkArguments.Add("-nodefaultlibs");
 				DirectoryReference LibCxxLibDir = DirectoryReference.Combine(LinuxInfo.BaseLinuxPath!, "lib64");
-				if(PlatformSDK.ForceUseLegacyLibCxx())
-				{
-					LibCxxLibDir = DirectoryReference.Combine(Unreal.EngineSourceDirectory, "ThirdParty", "Unix", "LibCxx", "lib", "Unix", LinkEnvironment.Architecture.LinuxName);
-				}
 				LinkArguments.Add($"\"{NormalizeCommandLinePath(FileReference.Combine(LibCxxLibDir, "libc++.a"), LinkEnvironment.RootPaths)}\"");
 				LinkArguments.Add($"\"{NormalizeCommandLinePath(FileReference.Combine(LibCxxLibDir, "libc++abi.a"), LinkEnvironment.RootPaths)}\"");
 				LinkArguments.Add("-lm");
diff --git a/Engine/Source/Programs/UnrealBuildTool/ToolChain/IWYUToolChain.cs b/Engine/Source/Programs/UnrealBuildTool/ToolChain/IWYUToolChain.cs
--- a/Engine/Source/Programs/UnrealBuildTool/ToolChain/IWYUToolChain.cs
+++ b/Engine/Source/Programs/UnrealBuildTool/ToolChain/IWYUToolChain.cs
@@ -155,8 +155,14 @@ namespace UnrealBuildTool
 					if (ShouldUseLibcxx(CompileEnvironment.Architecture))
 					{
 						Arguments.Add("-nostdinc++");
-						Arguments.Add(GetSystemIncludePathArgument(DirectoryReference.Combine(Unreal.EngineSourceDirectory, "ThirdParty", "Unix", "LibCxx", "include"), CompileEnvironment));
-						Arguments.Add(GetSystemIncludePathArgument(DirectoryReference.Combine(Unreal.EngineSourceDirectory, "ThirdParty", "Unix", "LibCxx", "include", "c++", "v1"), CompileEnvironment));
+
+						LinuxPlatformSDK PlatformSDK = new(Logger);
+						DirectoryReference? BaseLinuxPath = PlatformSDK.GetBaseLinuxPathForArchitecture(CompileEnvironment.Architecture);
+						if (BaseLinuxPath != null)
+						{
+							Arguments.Add(GetSystemIncludePathArgument(DirectoryReference.Combine(BaseLinuxPath, "include"), CompileEnvironment));
+							Arguments.Add(GetSystemIncludePathArgument(DirectoryReference.Combine(BaseLinuxPath, "include", "c++", "v1"), CompileEnvironment));
+						}
 					}
 
 					Arguments.Add("-fno-math-errno");               // do not assume that math ops have side effects
