Disable unhelp CA1027 warnings
#rnx

diff --git a/Engine/Source/Programs/UnrealBuildTool/Configuration/TargetRules.cs b/Engine/Source/Programs/UnrealBuildTool/Configuration/TargetRules.cs
--- a/Engine/Source/Programs/UnrealBuildTool/Configuration/TargetRules.cs
+++ b/Engine/Source/Programs/UnrealBuildTool/Configuration/TargetRules.cs
@@ -119,7 +119,9 @@ namespace UnrealBuildTool
 	/// <summary>
 	/// Determines which version of the engine to take default build settings from. This allows for backwards compatibility as new options are enabled by default.
 	/// </summary>
+#pragma warning disable CA1027 // Mark enums with FlagsAttribute
 	public enum BuildSettingsVersion
+#pragma warning restore CA1027 // Mark enums with FlagsAttribute
 	{
 		/// <summary>
 		/// Legacy default build settings for 4.23 and earlier.
@@ -173,7 +175,9 @@ namespace UnrealBuildTool
 	/// Specific versions can be used by licensees to avoid compile errors when integrating new versions of Unreal.
 	/// Using Latest comes with a high risk of introducing compile errors in your code on newer Unreal versions.
 	/// </summary>
+#pragma warning disable CA1027 // Mark enums with FlagsAttribute
 	public enum EngineIncludeOrderVersion
+#pragma warning restore  CA1027 // Mark enums with FlagsAttribute
 	{
 		/// <summary>
 		/// Include order used in Unreal 5.1
diff --git a/Engine/Source/Programs/UnrealBuildTool/System/CppCompileEnvironment.cs b/Engine/Source/Programs/UnrealBuildTool/System/CppCompileEnvironment.cs
--- a/Engine/Source/Programs/UnrealBuildTool/System/CppCompileEnvironment.cs
+++ b/Engine/Source/Programs/UnrealBuildTool/System/CppCompileEnvironment.cs
@@ -24,7 +24,9 @@ namespace UnrealBuildTool
 	/// <summary>
 	/// Specifies which language standard to use. This enum should be kept in order, so that toolchains can check whether the requested setting is >= values that they support.
 	/// </summary>
+#pragma warning disable CA1027 // Mark enums with FlagsAttribute
 	public enum CppStandardVersion
+#pragma warning restore CA1027 // Mark enums with FlagsAttribute
 	{
 		/// <summary>
 		/// Obsolete C++14 standard. No longer allowed, code will not compile with this version.
@@ -72,7 +74,9 @@ namespace UnrealBuildTool
 	/// <summary>
 	/// Specifies which C language standard to use. This enum should be kept in order, so that toolchains can check whether the requested setting is >= values that they support.
 	/// </summary>
+#pragma warning disable CA1027 // Mark enums with FlagsAttribute
 	public enum CStandardVersion
+#pragma warning restore CA1027 // Mark enums with FlagsAttribute
 	{
 		/// <summary>
 		/// Supports no additional standard version flag
@@ -115,7 +119,9 @@ namespace UnrealBuildTool
 	/// Note that by enabling this you are changing the minspec for the PC platform, and the resultant executable will crash on machines without AVX support.
 	/// For more details please see https://learn.microsoft.com/en-us/cpp/build/reference/arch-x64
 	/// </summary>
+#pragma warning disable CA1027 // Mark enums with FlagsAttribute
 	public enum MinimumCpuArchitectureX64
+#pragma warning restore CA1027 // Mark enums with FlagsAttribute
 	{
 		/// <summary>
 		/// No minimum architecture
diff --git a/Engine/Source/Programs/UnrealBuildTool/System/PluginDescriptor.cs b/Engine/Source/Programs/UnrealBuildTool/System/PluginDescriptor.cs
--- a/Engine/Source/Programs/UnrealBuildTool/System/PluginDescriptor.cs
+++ b/Engine/Source/Programs/UnrealBuildTool/System/PluginDescriptor.cs
@@ -13,7 +13,9 @@ namespace UnrealBuildTool
 	/// <summary>
 	/// The version format for .uplugin files. This rarely changes now; plugin descriptors should maintain backwards compatibility automatically.
 	/// </summary>
+#pragma warning disable CA1027 // Mark enums with FlagsAttribute
 	public enum PluginDescriptorVersion
+#pragma warning restore CA1027 // Mark enums with FlagsAttribute
 	{
 		/// <summary>
 		/// Invalid
