Fix CA1064

#rnx

diff --git a/Engine/Source/Programs/UnrealBuildTool/Preprocessor/Preprocessor.cs b/Engine/Source/Programs/UnrealBuildTool/Preprocessor/Preprocessor.cs
--- a/Engine/Source/Programs/UnrealBuildTool/Preprocessor/Preprocessor.cs
+++ b/Engine/Source/Programs/UnrealBuildTool/Preprocessor/Preprocessor.cs
@@ -16,7 +16,7 @@ namespace UnrealBuildTool
 	/// <summary>
 	/// Exception class for the preprocessor, which contains the file and position of the code causing an error
 	/// </summary>
-	class PreprocessorException : Exception
+	public class PreprocessorException : Exception
 	{
 		/// <summary>
 		/// The context when the error was encountered
diff --git a/Engine/Source/Programs/UnrealBuildTool/Preprocessor/PreprocessorContext.cs b/Engine/Source/Programs/UnrealBuildTool/Preprocessor/PreprocessorContext.cs
--- a/Engine/Source/Programs/UnrealBuildTool/Preprocessor/PreprocessorContext.cs
+++ b/Engine/Source/Programs/UnrealBuildTool/Preprocessor/PreprocessorContext.cs
@@ -8,7 +8,7 @@ namespace UnrealBuildTool
 	/// <summary>
 	/// Represents a context that the preprocessor is working in. Used to form error messages, and determine things like the __FILE__ and __LINE__ directives.
 	/// </summary>
-	abstract class PreprocessorContext
+	abstract public class PreprocessorContext
 	{
 		/// <summary>
 		/// The outer context
