Fix warning CA1012: Abstract type 'PreprocessorContext' should not have public constructors
#rnx

diff --git a/Engine/Source/Programs/UnrealBuildTool/Preprocessor/PreprocessorContext.cs b/Engine/Source/Programs/UnrealBuildTool/Preprocessor/PreprocessorContext.cs
--- a/Engine/Source/Programs/UnrealBuildTool/Preprocessor/PreprocessorContext.cs
+++ b/Engine/Source/Programs/UnrealBuildTool/Preprocessor/PreprocessorContext.cs
@@ -8,7 +8,7 @@ namespace UnrealBuildTool
 	/// <summary>
 	/// Represents a context that the preprocessor is working in. Used to form error messages, and determine things like the __FILE__ and __LINE__ directives.
 	/// </summary>
-	abstract public class PreprocessorContext
+	public abstract class PreprocessorContext
 	{
 		/// <summary>
 		/// The outer context
@@ -19,7 +19,7 @@ namespace UnrealBuildTool
 		/// Constructor
 		/// </summary>
 		/// <param name="outer">The outer context</param>
-		public PreprocessorContext(PreprocessorContext? outer)
+		protected PreprocessorContext(PreprocessorContext? outer)
 		{
 			Outer = outer;
 		}
