[IWYU]
* Fixed more edge cases. Resolving dependencies when using for example TArray<FSharedBuffer, TInlineAllocator<1>>

diff --git a/Engine/Source/Programs/UnrealBuildTool/Modes/IWYUMode.cs b/Engine/Source/Programs/UnrealBuildTool/Modes/IWYUMode.cs
--- a/Engine/Source/Programs/UnrealBuildTool/Modes/IWYUMode.cs
+++ b/Engine/Source/Programs/UnrealBuildTool/Modes/IWYUMode.cs
@@ -1092,7 +1092,7 @@ namespace UnrealBuildTool
 					}
 
 					// We don't want to remove these because they enable includes of windows headers.
-					if (Entry.Printable.StartsWith("Windows/AllowWindows") || Entry.Printable.StartsWith("Windows/HideWindows"))
+					if (Entry.Printable.StartsWith("Windows/AllowWindows") || Entry.Printable.StartsWith("Windows/HideWindows") || Entry.Printable.StartsWith("Windows/WindowsHWrapper.h"))
 					{
 						Info.Includes.Add(Entry);
 					}
