[UBT]
* Fixed bug in code creating marketplace rules assembly where it didn't populate rules dirs before finding modules. Fixes plugin discovery leading to RiderLink's BuildId not getting updated when it didn't need to be built

diff --git a/Engine/Source/Programs/UnrealBuildTool/System/RulesCompiler.cs b/Engine/Source/Programs/UnrealBuildTool/System/RulesCompiler.cs
--- a/Engine/Source/Programs/UnrealBuildTool/System/RulesCompiler.cs
+++ b/Engine/Source/Programs/UnrealBuildTool/System/RulesCompiler.cs
@@ -243,6 +243,22 @@ namespace UnrealBuildTool
 		{
 			RulesScope MarketplaceScope = new RulesScope("Marketplace", Parent.Scope);
 
+			// We want to prefetch the directories in one call to make it go as wide as possible against filesystem
+			{
+				List<DirectoryReference> populateDirs = [];
+				foreach (PluginInfo info in Plugins)
+				{
+					populateDirs.Add(DirectoryReference.Combine(info.Directory, "Source"));
+					populateDirs.Add(DirectoryReference.Combine(info.Directory, "Tests"));
+					foreach (FileReference ChildFile in info.ChildFiles)
+					{
+						populateDirs.Add(DirectoryReference.Combine(ChildFile.Directory, "Source"));
+						populateDirs.Add(DirectoryReference.Combine(ChildFile.Directory, "Tests"));
+					}
+				}
+				Rules.Populate(populateDirs);
+			}
+
 			// Add all the plugin modules too (don't need to loop over RootDirectories since the plugins come in already found
 			Dictionary<FileReference, ModuleRulesContext> ModuleFileToContext = new Dictionary<FileReference, ModuleRulesContext>();
 			using (GlobalTracer.Instance.BuildSpan("Finding marketplace plugin modules").StartActive())
