Fix missing documentation

#rnx

diff --git a/Engine/Source/Programs/UnrealBuildTool/Configuration/BuildConfiguration.cs b/Engine/Source/Programs/UnrealBuildTool/Configuration/BuildConfiguration.cs
--- a/Engine/Source/Programs/UnrealBuildTool/Configuration/BuildConfiguration.cs
+++ b/Engine/Source/Programs/UnrealBuildTool/Configuration/BuildConfiguration.cs
@@ -133,6 +133,9 @@ namespace UnrealBuildTool
 		[XmlConfigFile(Category = "UEBuildConfiguration")]
 		public bool bForceHeaderGeneration { get; set; } = false;
 
+		/// <summary>
+		/// If true, enable the UHT input cache for both reading and writing
+		/// </summary>
 		[CommandLine("-EnableUHTInputCache")]
 		[XmlConfigFile(Category = "UEBuildConfiguration")]
 		public bool bEnableUHTInputCache { get; set; } = false;
