[UBT]
* Fixed so prebuildtargets work again in UBT

diff --git a/Engine/Source/Programs/UnrealBuildTool/Modes/BuildMode.cs b/Engine/Source/Programs/UnrealBuildTool/Modes/BuildMode.cs
--- a/Engine/Source/Programs/UnrealBuildTool/Modes/BuildMode.cs
+++ b/Engine/Source/Programs/UnrealBuildTool/Modes/BuildMode.cs
@@ -262,12 +262,16 @@ namespace UnrealBuildTool
 					// Create the working set provider per group.
 					using (ISourceFileWorkingSet WorkingSet = SourceFileWorkingSet.Create(Unreal.RootDirectory, ProjectDirs, Logger))
 					{
-						List<(TargetDescriptor, TargetMakefile)> Targets = CreateMakefiles(BuildConfiguration, TargetDescriptors, WorkingSet, Logger, true, true, tasks, cppDependencyCache, fileHashCache);
+						List<List<(TargetDescriptor, TargetMakefile)>> targetsPasses = CreateMakefiles(BuildConfiguration, TargetDescriptors, WorkingSet, Logger, true, false, tasks, cppDependencyCache, fileHashCache);
 						if (fileHashCache.HasErrors)
 						{
 							throw new BuildException("FileHashCache had errors");
 						}
-						await BuildAsync(Targets, BuildConfiguration, Options, WriteOutdatedActionsFile, cppDependencyCache, Logger, ActionTypeFilter, executor);
+
+						foreach (List<(TargetDescriptor, TargetMakefile)> pass in targetsPasses)
+						{
+							await BuildAsync(pass, BuildConfiguration, Options, WriteOutdatedActionsFile, cppDependencyCache, Logger, ActionTypeFilter, executor);
+						}
 					}
 				}
 			}
@@ -411,8 +415,11 @@ namespace UnrealBuildTool
 			List<Task> tasks = [];
 			try
 			{
-				List<(TargetDescriptor, TargetMakefile)> Targets = CreateMakefiles(BuildConfiguration, TargetDescriptors, WorkingSet, Logger, true, true, tasks, cppDependencyCache, new(false, true));
-				await BuildAsync(Targets, BuildConfiguration, Options, WriteOutdatedActionsFile, cppDependencyCache, Logger, ActionTypeFilter, null);
+				List<List<(TargetDescriptor, TargetMakefile)>> targetsPasses = CreateMakefiles(BuildConfiguration, TargetDescriptors, WorkingSet, Logger, true, true, tasks, cppDependencyCache, new(false, true));
+				foreach (List<(TargetDescriptor, TargetMakefile)> targets in targetsPasses)
+				{
+					await BuildAsync(targets, BuildConfiguration, Options, WriteOutdatedActionsFile, cppDependencyCache, Logger, ActionTypeFilter, null);
+				}
 			}
 			finally
 			{
@@ -1048,8 +1055,8 @@ namespace UnrealBuildTool
 			List<Task> tasks = [];
 			try
 			{
-				List<(TargetDescriptor, TargetMakefile Makefile)> Targets = CreateMakefiles(BuildConfiguration, [TargetDescriptor], WorkingSet, Logger, false, false, tasks, null, new(false, true));
-				return Targets[0].Makefile;
+				List<List<(TargetDescriptor, TargetMakefile Makefile)>> targetsPasses = CreateMakefiles(BuildConfiguration, [TargetDescriptor], WorkingSet, Logger, false, false, tasks, null, new(false, true));
+				return targetsPasses[0][0].Makefile;
 			}
 			finally
 			{
@@ -1082,9 +1089,9 @@ namespace UnrealBuildTool
 			return (makefile, true);
 		}
 
-		private static List<(TargetDescriptor, TargetMakefile)> CreateMakefiles(BuildConfiguration buildConfiguration, List<TargetDescriptor> descriptors, ISourceFileWorkingSet workingSet, ILogger logger, bool runInitScripts, bool skipPreBuildTargets, List<Task> outTasks, CppDependencyCache? cppDependencyCache, FileHashCache fileHashCache)
+		private static List<List<(TargetDescriptor, TargetMakefile)>> CreateMakefiles(BuildConfiguration buildConfiguration, List<TargetDescriptor> descriptors, ISourceFileWorkingSet workingSet, ILogger logger, bool runInitScripts, bool skipPreBuildTargets, List<Task> outTasks, CppDependencyCache? cppDependencyCache, FileHashCache fileHashCache)
 		{
-			List<(TargetDescriptor descriptor, TargetMakefile makefile)> outTargets = new();
+			List<List<(TargetDescriptor descriptor, TargetMakefile makefile)>> outTargetsPasses = new();
 
 			int helperCount = buildConfiguration.bParallelMakefileGeneration ? -1 : 0;
 
@@ -1098,6 +1105,8 @@ namespace UnrealBuildTool
 
 			while (descriptors.Count != 0) // descriptors list will be reallocated for PrebuildTargets. hence the while-loop
 			{
+				List<(TargetDescriptor descriptor, TargetMakefile makefile)> outTargets = new();
+
 				int descriptorCount = descriptors.Count;
 				bool skipRulesCompile = buildConfiguration.bSkipRulesCompile;
 				bool forceRulesCompile = buildConfiguration.bForceRulesCompile;
@@ -1452,14 +1461,23 @@ namespace UnrealBuildTool
 
 				descriptors = newDescriptors;
 
+				outTargetsPasses.Insert(0, outTargets);
+
 				// In the very rare scenario we preload the wrong makefiles we want to make sure to wait for them in case there are prebuildtargets that use them
 				foreach (Task<(TargetMakefile?, bool)> task in preLoadedMakefiles.Values)
 				{
 					task?.Wait();
 				}
+
+				// Make sure to finish all tasks if we have prebuild targets
+				if (descriptors.Count != 0)
+				{
+					Task.WaitAll(outTasks.ToArray());
+					outTasks.Clear();
+				}
 			}
 
-			return outTargets;
+			return outTargetsPasses;
 		}
 
 		/// <summary>
