UnrealBuildTool: Update PCH related classes to records

#rnx

diff --git a/Engine/Source/Programs/UnrealBuildTool/Configuration/UEBuildModuleCPP.cs b/Engine/Source/Programs/UnrealBuildTool/Configuration/UEBuildModuleCPP.cs
--- a/Engine/Source/Programs/UnrealBuildTool/Configuration/UEBuildModuleCPP.cs
+++ b/Engine/Source/Programs/UnrealBuildTool/Configuration/UEBuildModuleCPP.cs
@@ -1742,10 +1742,7 @@ namespace UnrealBuildTool
 					{
 						Output = ToolChain.CompileAllCPPFiles(CompileEnvironment, new List<FileItem>() { WrapperFile }, Template.OutputDir, "Shared", Graph);
 					}
-					Instance = new PrecompiledHeaderInstance(WrapperFile, SharedDefinitionsFileItem, CompileEnvironment, Output, GetImmutableDefinitions(Template.BaseCompileEnvironment.Definitions))
-					{
-						ParentPCHInstance = ParentPCHInstance
-					};
+					Instance = new PrecompiledHeaderInstance(WrapperFile, SharedDefinitionsFileItem, CompileEnvironment, Output, GetImmutableDefinitions(Template.BaseCompileEnvironment.Definitions), ParentPCHInstance);
 					Template.Instances.Add(Instance);
 				}
 
diff --git a/Engine/Source/Programs/UnrealBuildTool/System/PrecompiledHeaderInstance.cs b/Engine/Source/Programs/UnrealBuildTool/System/PrecompiledHeaderInstance.cs
--- a/Engine/Source/Programs/UnrealBuildTool/System/PrecompiledHeaderInstance.cs
+++ b/Engine/Source/Programs/UnrealBuildTool/System/PrecompiledHeaderInstance.cs
@@ -8,62 +8,20 @@ namespace UnrealBuildTool
 	/// <summary>
 	/// Information about a PCH instance
 	/// </summary>
-	class PrecompiledHeaderInstance
+	/// <param name="HeaderFile">The file to include to use this shared PCH</param>
+	/// <param name="DefinitionsFile">The definitions file</param>
+	/// <param name="CompileEnvironment">The compile environment for this shared PCH</param>
+	/// <param name="Output">The output files for the shared PCH</param>
+	/// <param name="ImmutableDefinitions">These are definitions that are immutable and should never be #undef. There are a few exceptions and we make sure those are not ending up in this list</param>
+	/// <param name="ParentPCHInstance">Parent PCH instance used in PCH chaining</param>
+	record class PrecompiledHeaderInstance(FileItem HeaderFile, FileItem DefinitionsFile, CppCompileEnvironment CompileEnvironment, CPPOutput Output, IReadOnlySet<string> ImmutableDefinitions, PrecompiledHeaderInstance? ParentPCHInstance = null)
 	{
 		/// <summary>
-		/// The file to include to use this shared PCH
+		/// Modules using this PCH instance
 		/// </summary>
-		public FileItem HeaderFile;
+		public HashSet<UEBuildModuleCPP> Modules { get; } = [];
 
-		/// <summary>
-		/// The definitions file
-		/// </summary>
-		public FileItem DefinitionsFile;
-
-		/// <summary>
-		/// The compile environment for this shared PCH
-		/// </summary>
-		public CppCompileEnvironment CompileEnvironment;
-
-		/// <summary>
-		/// The output files for the shared PCH
-		/// </summary>
-		public CPPOutput Output;
-
-		/// <summary>
-		/// List of modules using this instance
-		/// </summary>
-		public HashSet<UEBuildModuleCPP> Modules = new HashSet<UEBuildModuleCPP>();
-
-		/// <summary>
-		/// These are definitions that are immutable and should never be #undef. There are a few exceptions and we make sure those are not ending up in this list
-		/// </summary>
-		public HashSet<string> ImmutableDefinitions;
-
-		/// <summary>
-		/// Parent PCH instance used in PCH chaining
-		/// </summary>
-		public PrecompiledHeaderInstance? ParentPCHInstance;
-
-		/// <summary>
-		/// Constructor
-		/// </summary>
-		public PrecompiledHeaderInstance(FileItem HeaderFile, FileItem DefinitionsFile, CppCompileEnvironment CompileEnvironment, CPPOutput Output, HashSet<string> ImmutableDefinitions)
-		{
-			this.HeaderFile = HeaderFile;
-			this.DefinitionsFile = DefinitionsFile;
-			this.CompileEnvironment = CompileEnvironment;
-			this.Output = Output;
-			this.ImmutableDefinitions = ImmutableDefinitions;
-		}
-
-		/// <summary>
-		/// Return a string representation of this object for debugging
-		/// </summary>
-		/// <returns>String representation of the object</returns>
-		public override string ToString()
-		{
-			return HeaderFile.Location.GetFileName();
-		}
+		/// <inheritdoc/>
+		public override string ToString() => HeaderFile.Location.GetFileName();
 	}
 }
diff --git a/Engine/Source/Programs/UnrealBuildTool/System/PrecompiledHeaderTemplate.cs b/Engine/Source/Programs/UnrealBuildTool/System/PrecompiledHeaderTemplate.cs
--- a/Engine/Source/Programs/UnrealBuildTool/System/PrecompiledHeaderTemplate.cs
+++ b/Engine/Source/Programs/UnrealBuildTool/System/PrecompiledHeaderTemplate.cs
@@ -9,80 +9,28 @@ namespace UnrealBuildTool
 	/// <summary>
 	/// A template for creating a shared PCH. Instances of it are created depending on the configurations required.
 	/// </summary>
-	class PrecompiledHeaderTemplate
+	/// <param name="Module">The module providing valid shared PCH</param>
+	/// <param name="BaseCompileEnvironment">The base compile environment, including all the public compile environment that all consuming modules inherit</param>
+	/// <param name="HeaderFile">The header file to generate a PCH from</param>
+	/// <param name="OutputDir">Output directory for instances of this PCH</param>
+	/// <param name="ModuleDependencies">All the module dependencies this template has</param>
+	record class PrecompiledHeaderTemplate(UEBuildModuleCPP Module, CppCompileEnvironment BaseCompileEnvironment, FileItem HeaderFile, DirectoryReference OutputDir, IReadOnlySet<UEBuildModule> ModuleDependencies)
 	{
-		/// <summary>
-		/// Module providing this PCH.
-		/// </summary>
-		public UEBuildModuleCPP Module;
-
-		/// <summary>
-		/// The base compile environment, including all the public compile environment that all consuming modules inherit.
-		/// </summary>
-		public CppCompileEnvironment BaseCompileEnvironment;
-
-		/// <summary>
-		/// The header file 
-		/// </summary>
-		public FileItem HeaderFile;
-
-		/// <summary>
-		/// Output directory for instances of this PCH.
-		/// </summary>
-		public DirectoryReference OutputDir;
-
 		/// <summary>
 		/// Instances of this PCH
 		/// </summary>
-		public List<PrecompiledHeaderInstance> Instances = [];
-
-		/// <summary>
-		/// All the module dependencies this template has
-		/// </summary>
-		public HashSet<UEBuildModule> ModuleDependencies;
-
-		/// <summary>
-		/// Constructor
-		/// </summary>
-		/// <param name="Module">The module with a valid shared PCH</param>
-		/// <param name="BaseCompileEnvironment">The compile environment to use</param>
-		/// <param name="HeaderFile">The header file to generate a PCH from</param>
-		/// <param name="OutputDir">Output directory for instances of this PCH</param>
-		/// <param name="ModuleDependencies">All the module dependencies this template has</param>
-		public PrecompiledHeaderTemplate(UEBuildModuleCPP Module, CppCompileEnvironment BaseCompileEnvironment, FileItem HeaderFile, DirectoryReference OutputDir, HashSet<UEBuildModule> ModuleDependencies)
-		{
-			this.Module = Module;
-			this.BaseCompileEnvironment = BaseCompileEnvironment;
-			this.HeaderFile = HeaderFile;
-			this.OutputDir = OutputDir;
-			this.ModuleDependencies = ModuleDependencies;
-		}
+		public List<PrecompiledHeaderInstance> Instances { get; } = [];
 
 		/// <summary>
 		/// Checks whether this template is valid for the given compile environment
 		/// </summary>
-		/// <param name="CompileEnvironment">Compile environment to check with</param>
+		/// <param name="compileEnvironment">Compile environment to check with</param>
 		/// <returns>True if the template is compatible with the given compile environment</returns>
-		public bool IsValidFor(CppCompileEnvironment CompileEnvironment)
-		{
-			if (CompileEnvironment.bIsBuildingDLL != BaseCompileEnvironment.bIsBuildingDLL)
-			{
-				return false;
-			}
-			if (CompileEnvironment.bIsBuildingLibrary != BaseCompileEnvironment.bIsBuildingLibrary)
-			{
-				return false;
-			}
-			return true;
-		}
+		public bool IsValidFor(CppCompileEnvironment compileEnvironment) =>
+			compileEnvironment.bIsBuildingDLL == BaseCompileEnvironment.bIsBuildingDLL &&
+			compileEnvironment.bIsBuildingLibrary == BaseCompileEnvironment.bIsBuildingLibrary;
 
-		/// <summary>
-		/// Return a string representation of this object for debugging
-		/// </summary>
-		/// <returns>String representation of the object</returns>
-		public override string ToString()
-		{
-			return HeaderFile.AbsolutePath;
-		}
+		/// <inheritdoc/>
+		public override string ToString() => HeaderFile.AbsolutePath;
 	}
 }
