UnrealBuildTool: Use explicit type instead of 'var' (ide0008)

#rnx

diff --git a/Engine/Source/Programs/Horde/.editorconfig b/Engine/Source/Programs/Horde/.editorconfig
--- a/Engine/Source/Programs/Horde/.editorconfig
+++ b/Engine/Source/Programs/Horde/.editorconfig
@@ -218,6 +218,7 @@ csharp_style_allow_blank_line_after_colon_in_constructor_initializer_experimenta
 
 # Prefer "var" everywhere
 dotnet_diagnostic.IDE0007.severity = error
+dotnet_diagnostic.IDE0008.severity = error
 csharp_style_var_for_built_in_types = false:error
 csharp_style_var_when_type_is_apparent = false:error
 csharp_style_var_elsewhere = false:error
diff --git a/Engine/Source/Programs/Shared/.editorconfig b/Engine/Source/Programs/Shared/.editorconfig
--- a/Engine/Source/Programs/Shared/.editorconfig
+++ b/Engine/Source/Programs/Shared/.editorconfig
@@ -214,6 +214,7 @@ csharp_style_allow_blank_line_after_colon_in_constructor_initializer_experimenta
 
 # Prefer "var" everywhere
 dotnet_diagnostic.IDE0007.severity = error
+dotnet_diagnostic.IDE0008.severity = error
 csharp_style_var_for_built_in_types = false:error
 csharp_style_var_when_type_is_apparent = false:error
 csharp_style_var_elsewhere = false:error
diff --git a/Engine/Source/Programs/UnrealBuildTool/.editorconfig b/Engine/Source/Programs/UnrealBuildTool/.editorconfig
--- a/Engine/Source/Programs/UnrealBuildTool/.editorconfig
+++ b/Engine/Source/Programs/UnrealBuildTool/.editorconfig
@@ -213,6 +213,7 @@ csharp_style_allow_blank_line_after_colon_in_constructor_initializer_experimenta
 
 # Prefer "var" everywhere
 dotnet_diagnostic.IDE0007.severity = error
+dotnet_diagnostic.IDE0008.severity = error
 csharp_style_var_for_built_in_types = false:error
 csharp_style_var_when_type_is_apparent = false:error
 csharp_style_var_elsewhere = false:error
diff --git a/Engine/Source/Programs/UnrealBuildTool/Configuration/UEBuildModule.cs b/Engine/Source/Programs/UnrealBuildTool/Configuration/UEBuildModule.cs
--- a/Engine/Source/Programs/UnrealBuildTool/Configuration/UEBuildModule.cs
+++ b/Engine/Source/Programs/UnrealBuildTool/Configuration/UEBuildModule.cs
@@ -368,10 +368,10 @@ namespace UnrealBuildTool
 			// We precalculate these because they are reused over and over again and causing tons of allocations
 			for (int i=0; i!=3; ++i)
 			{
-				string ApiDefinition = ApiDefinitions[i]; 
+				string ApiDefinition = ApiDefinitions[i];
 
 				// Add a spacer if and only if we've actually used import/export above.
-				String Spacer = ApiDefinition.Length == 0 ? String.Empty : " ";
+				string Spacer = ApiDefinition.Length == 0 ? String.Empty : " ";
 				// Add the import or export declaration for the module
 				ModuleApiDefinitionStrings[i] = $"{ModuleApiDefine}={ApiDefinition}{Spacer}{AutoRTFMDisable}";
 				ModuleNonAttributedApiDefinitionStrings[i] = $"{ModuleNonAttributedDefine}={ApiDefinition}";
diff --git a/Engine/Source/Programs/UnrealBuildTool/Executors/UnrealBuildAccelerator/UBAAgentCoordinatorHorde.cs b/Engine/Source/Programs/UnrealBuildTool/Executors/UnrealBuildAccelerator/UBAAgentCoordinatorHorde.cs
--- a/Engine/Source/Programs/UnrealBuildTool/Executors/UnrealBuildAccelerator/UBAAgentCoordinatorHorde.cs
+++ b/Engine/Source/Programs/UnrealBuildTool/Executors/UnrealBuildAccelerator/UBAAgentCoordinatorHorde.cs
@@ -914,7 +914,7 @@ namespace UnrealBuildTool
 					return;
 				}
 
-				var (hordeSession, reason) = await _hordeSessionTask;
+				(UBAHordeSession? hordeSession, string? reason) = await _hordeSessionTask;
 
 				if (hordeSession == null)
 				{
@@ -1018,7 +1018,7 @@ namespace UnrealBuildTool
 				return;
 			}
 
-			var (hordeSession, reason) = await _hordeSessionTask;
+			(UBAHordeSession? hordeSession, string? reason) = await _hordeSessionTask;
 			_hordeSessionTask = null;
 			if (hordeSession != null)
 			{
diff --git a/Engine/Source/Programs/UnrealBuildTool/Executors/UnrealBuildAccelerator/UBAExecutor.cs b/Engine/Source/Programs/UnrealBuildTool/Executors/UnrealBuildAccelerator/UBAExecutor.cs
--- a/Engine/Source/Programs/UnrealBuildTool/Executors/UnrealBuildAccelerator/UBAExecutor.cs
+++ b/Engine/Source/Programs/UnrealBuildTool/Executors/UnrealBuildAccelerator/UBAExecutor.cs
@@ -946,7 +946,7 @@ namespace UnrealBuildTool
 				bool enableDetour = !ForceLocalNoDetour(action) && action.bCanExecuteInUBA;
 				bool canRunRemotely = CanRunRemotely(action);
 
-				var (knownInputs, knownInputsCount) = BuildKnownInputs(action);
+				(byte[]? knownInputs, uint knownInputsCount) = BuildKnownInputs(action);
 
 				uint bucket = GetCacheBucket(logger, action);
 				uint memoryGroup = GetMemoryGroup(action);
diff --git a/Engine/Source/Programs/UnrealBuildTool/Modes/IWYUMode.cs b/Engine/Source/Programs/UnrealBuildTool/Modes/IWYUMode.cs
--- a/Engine/Source/Programs/UnrealBuildTool/Modes/IWYUMode.cs
+++ b/Engine/Source/Programs/UnrealBuildTool/Modes/IWYUMode.cs
@@ -458,7 +458,7 @@ namespace UnrealBuildTool
 				Logger.LogInformation($"Creating BuildTarget for {TargetName}...");
 				UEBuildTarget Target = UEBuildTarget.Create(TargetDescriptors[0], BuildConfiguration, Logger);
 
-				var env = Target.Rules.WindowsPlatform.Environment!;
+				VCEnvironment env = Target.Rules.WindowsPlatform.Environment!;
 				systemPaths.Add(env.ToolChainDir.FullName.Replace('\\', '/'));
 				systemPaths.Add(env.WindowsSdkDir.FullName.Replace('\\', '/'));
 
@@ -707,14 +707,21 @@ namespace UnrealBuildTool
 						IWYUFile? IWYUFile = JsonSerializer.Deserialize<IWYUFile>(JsonContent);
 						if (IWYUFile == null)
 						{
-							Logger.LogError($"Failed to deserialize file {IWYUFilePath}");
+							Logger.LogError("Failed to deserialize file {IWYUFilePath}", IWYUFilePath);
+							Interlocked.Exchange(ref ReadSuccess, 0);
+							return;
+						}
+
+						if (IWYUFile.Files.Count == 0)
+						{
+							Logger.LogError("No files found in {IWYUFilePath}", IWYUFilePath);
 							Interlocked.Exchange(ref ReadSuccess, 0);
 							return;
 						}
 
 						IWYUFile.Name = IWYUFilePath;
 
-						var mainInfo = IWYUFile.Files[IWYUFile.Files.Count - 1];
+						IWYUInfo mainInfo = IWYUFile.Files.Last();
 						mainInfo.Source = IWYUFile;
 
 						lock (Infos)
@@ -741,7 +748,7 @@ namespace UnrealBuildTool
 
 				// Now traverse main infos and find additional infos and fix them up
 				Logger.LogInformation($"Adding additional iwyu infos...");
-				var mainInfos = Infos.Values.ToArray();
+				IWYUInfo[] mainInfos = Infos.Values.ToArray();
 				Parallel2.ForEach(mainInfos, (mainInfo) =>
 				{
 					IWYUFile IWYUFile = mainInfo.Source!;
@@ -877,9 +884,9 @@ namespace UnrealBuildTool
 					{
 						if (headerInCpp.Includes.Count != headerAsMain!.Includes.Count)
 						{
-							var exists = new HashSet<IWYUIncludeEntry>(mainInfo.Includes, new IWYUIncludeEntryFullPathComparer());
+							HashSet<IWYUIncludeEntry> exists = new HashSet<IWYUIncludeEntry>(mainInfo.Includes, new IWYUIncludeEntryFullPathComparer());
 							exists.UnionWith(headerAsMain!.Includes);
-							foreach (var i in headerInCpp.Includes)
+							foreach (IWYUIncludeEntry i in headerInCpp.Includes)
 							{
 								if (!exists.Contains(i))
 								{
@@ -897,15 +904,15 @@ namespace UnrealBuildTool
 				{
 					IWYUInfo mainInfo = entry.mainInfo;
 					List<IWYUInfo> generatedHeaders = entry.generatedHeaders;
-					var exists = new HashSet<IWYUIncludeEntry>(mainInfo.Includes, new IWYUIncludeEntryFullPathComparer());
-					foreach (var generatedHeader in generatedHeaders)
+					HashSet<IWYUIncludeEntry> exists = new HashSet<IWYUIncludeEntry>(mainInfo.Includes, new IWYUIncludeEntryFullPathComparer());
+					foreach (IWYUInfo generatedHeader in generatedHeaders)
 					{
 						exists.UnionWith(generatedHeader.IncludesSeenInFile);
 					}
 
-					foreach (var generatedHeader in generatedHeaders)
+					foreach (IWYUInfo generatedHeader in generatedHeaders)
 					{
-						foreach (var i in generatedHeader.Includes)
+						foreach (IWYUIncludeEntry i in generatedHeader.Includes)
 						{
 							if (!exists.Contains(i))
 							{
@@ -1282,13 +1289,13 @@ namespace UnrealBuildTool
 			{
 				PrintInfo = PrintInfo.Replace('\\', '/');
 				bool found = false;
-				foreach (var info in Infos.Values)
+				foreach (IWYUInfo info in Infos.Values)
 				{
 					if (info.File.Contains(PrintInfo, StringComparison.OrdinalIgnoreCase))
 					{
 						Logger.LogInformation($"--- {info.File} ---");
 						Logger.LogInformation($"  Transitive includes");
-						foreach (var i in info.TransitiveIncludes)
+						foreach (KeyValuePair<string, string> i in info.TransitiveIncludes)
 						{
 							Logger.LogInformation($"    {i.Key}");
 						}
@@ -2167,7 +2174,7 @@ namespace UnrealBuildTool
 			// If we have been logging we can exit now since we don't want to write any files to disk
 			if (sortedLog != null)
 			{
-				foreach (var kv in sortedLog)
+				foreach (KeyValuePair<string, List<string>> kv in sortedLog)
 				{
 					Console.WriteLine(kv.Key);
 					foreach (string line in kv.Value)
