UnrealBuildTool: Use framework type (ide0049)

#rnx

diff --git a/Engine/Source/Programs/UnrealBuildTool/Executors/UnrealBuildAccelerator/UBAAgentCoordinatorHorde.cs b/Engine/Source/Programs/UnrealBuildTool/Executors/UnrealBuildAccelerator/UBAAgentCoordinatorHorde.cs
--- a/Engine/Source/Programs/UnrealBuildTool/Executors/UnrealBuildAccelerator/UBAAgentCoordinatorHorde.cs
+++ b/Engine/Source/Programs/UnrealBuildTool/Executors/UnrealBuildAccelerator/UBAAgentCoordinatorHorde.cs
@@ -330,8 +330,8 @@ namespace UnrealBuildTool
 				string agentName = String.Empty;
 				int numLogicalCores = 24; // Assume 24 if something goes wrong here and property is not found
 				string computeIp = String.Empty;
-				string leaseId = string.Empty;
-				string instanceType = string.Empty;
+				string leaseId = String.Empty;
+				string instanceType = String.Empty;
 				foreach (string property in lease.Properties)
 				{
 					int equalsIdx = property.IndexOf('=', StringComparison.OrdinalIgnoreCase);
@@ -837,7 +837,7 @@ namespace UnrealBuildTool
 			if (_hordeConfig.HordeServer != null)
 			{
 				link = $"{_hordeConfig.HordeServer!.TrimEnd('/')}/agents";
-				if (!string.IsNullOrEmpty(_hordeConfig.HordePool))
+				if (!String.IsNullOrEmpty(_hordeConfig.HordePool))
 				{
 					link += $"?agent={_hordeConfig.HordePool}";
 				}
@@ -1055,7 +1055,7 @@ namespace UnrealBuildTool
 		{
 			Requirements requirements = new() { Exclusive = true };
 			
-			if (!string.IsNullOrEmpty(config.HordeCluster))
+			if (!String.IsNullOrEmpty(config.HordeCluster))
 			{
 				// Apply filtering options when a cluster, either explicit or auto is specified, ensuring the cluster options take precidence.
 				string condition = "";
diff --git a/Engine/Source/Programs/UnrealBuildTool/Executors/UnrealBuildAccelerator/UBAExecutor.cs b/Engine/Source/Programs/UnrealBuildTool/Executors/UnrealBuildAccelerator/UBAExecutor.cs
--- a/Engine/Source/Programs/UnrealBuildTool/Executors/UnrealBuildAccelerator/UBAExecutor.cs
+++ b/Engine/Source/Programs/UnrealBuildTool/Executors/UnrealBuildAccelerator/UBAExecutor.cs
@@ -755,7 +755,7 @@ namespace UnrealBuildTool
 
 				LinkedAction action = (LinkedAction)info.UserData;
 
-				string additionalDescription = string.Empty;
+				string additionalDescription = String.Empty;
 
 				ProcessExecutionType executionType = info.ExecutionType;
 
diff --git a/Engine/Source/Programs/UnrealBuildTool/Modes/UnrealHeaderToolMode.cs b/Engine/Source/Programs/UnrealBuildTool/Modes/UnrealHeaderToolMode.cs
--- a/Engine/Source/Programs/UnrealBuildTool/Modes/UnrealHeaderToolMode.cs
+++ b/Engine/Source/Programs/UnrealBuildTool/Modes/UnrealHeaderToolMode.cs
@@ -303,7 +303,7 @@ namespace UnrealBuildTool.Modes
 		/// <param name="ProjectPath"></param>
 		public void ProjectSpecificConfigs(string? ProjectPath)
 		{
-			if (string.IsNullOrEmpty(ProjectPath))
+			if (String.IsNullOrEmpty(ProjectPath))
 			{
 				return;
 			}
diff --git a/Engine/Source/Programs/UnrealBuildTool/Platform/Android/AndroidPlatformSDK.cs b/Engine/Source/Programs/UnrealBuildTool/Platform/Android/AndroidPlatformSDK.cs
--- a/Engine/Source/Programs/UnrealBuildTool/Platform/Android/AndroidPlatformSDK.cs
+++ b/Engine/Source/Programs/UnrealBuildTool/Platform/Android/AndroidPlatformSDK.cs
@@ -209,7 +209,7 @@ namespace UnrealBuildTool
 				Debug.Assert(OperatingSystem.IsWindows());
 
 				string? Directory = Microsoft.Win32.Registry.GetValue(@"HKEY_LOCAL_MACHINE\SOFTWARE\Android Studio", "Path", null) as string;
-				return !string.IsNullOrEmpty(Directory) ? Directory : Path.Combine(Environment.GetFolderPath(Environment.SpecialFolder.LocalApplicationData), @"Programs\Android Studio");
+				return !String.IsNullOrEmpty(Directory) ? Directory : Path.Combine(Environment.GetFolderPath(Environment.SpecialFolder.LocalApplicationData), @"Programs\Android Studio");
 			}
 		}
 		
@@ -228,7 +228,7 @@ namespace UnrealBuildTool
 				Debug.Assert(OperatingSystem.IsWindows());
 
 				string? Directory = Microsoft.Win32.Registry.GetValue(@"HKEY_LOCAL_MACHINE\SOFTWARE\Android Studio", "SdkPath", null) as string;
-				return !string.IsNullOrEmpty(Directory) ? Directory : Path.Combine(Environment.GetFolderPath(Environment.SpecialFolder.LocalApplicationData), @"Android\Sdk");
+				return !String.IsNullOrEmpty(Directory) ? Directory : Path.Combine(Environment.GetFolderPath(Environment.SpecialFolder.LocalApplicationData), @"Android\Sdk");
 			}
 		}
 
@@ -255,7 +255,7 @@ namespace UnrealBuildTool
 			
 			string? CurrentSdkDir = Environment.GetEnvironmentVariable("ANDROID_HOME");
 			
-			if (string.IsNullOrEmpty(CurrentSdkDir) || !Directory.Exists(Path.Combine(CurrentSdkDir, "platforms", GetPlatformSpecificVersion("platforms"))) || GetInstalledSDKVersion() != SelectedSDKVersion)
+			if (String.IsNullOrEmpty(CurrentSdkDir) || !Directory.Exists(Path.Combine(CurrentSdkDir, "platforms", GetPlatformSpecificVersion("platforms"))) || GetInstalledSDKVersion() != SelectedSDKVersion)
 			{
 				try
 				{
diff --git a/Engine/Source/Programs/UnrealBuildTool/Platform/Android/AndroidToolChain.cs b/Engine/Source/Programs/UnrealBuildTool/Platform/Android/AndroidToolChain.cs
--- a/Engine/Source/Programs/UnrealBuildTool/Platform/Android/AndroidToolChain.cs
+++ b/Engine/Source/Programs/UnrealBuildTool/Platform/Android/AndroidToolChain.cs
@@ -892,7 +892,7 @@ namespace UnrealBuildTool
 			}
 
 			int i = 4, nDigitsApi = 0;
-			while (i < suffix.Length && char.IsDigit(suffix[i]))
+			while (i < suffix.Length && Char.IsDigit(suffix[i]))
 			{
 				i++;
 				nDigitsApi++;
@@ -911,7 +911,7 @@ namespace UnrealBuildTool
 			i += 4;
 
 			int nDigitsNdk = 0;
-			while (i < suffix.Length && char.IsDigit(suffix[i]))
+			while (i < suffix.Length && Char.IsDigit(suffix[i]))
 			{
 				i++;
 				nDigitsNdk++;
diff --git a/Engine/Source/Programs/UnrealBuildTool/Platform/Android/UEDeployAndroid.cs b/Engine/Source/Programs/UnrealBuildTool/Platform/Android/UEDeployAndroid.cs
--- a/Engine/Source/Programs/UnrealBuildTool/Platform/Android/UEDeployAndroid.cs
+++ b/Engine/Source/Programs/UnrealBuildTool/Platform/Android/UEDeployAndroid.cs
@@ -244,7 +244,7 @@ namespace UnrealBuildTool
 		private int GetApiLevelInt(string ApiString)
 		{
 			int DashIndex = ApiString.LastIndexOf('-');
-			if (DashIndex != -1 && int.TryParse(ApiString.Substring(DashIndex + 1), out int Version))
+			if (DashIndex != -1 && Int32.TryParse(ApiString.Substring(DashIndex + 1), out int Version))
 			{
 				return Version;
 			}
@@ -1441,10 +1441,10 @@ namespace UnrealBuildTool
 			[JsonPropertyName("Value")]
 			public object MaybeValue { get; set; }
 
-			private string MaybeValueAsString => MaybeValue?.ToString() ?? string.Empty;
+			private string MaybeValueAsString => MaybeValue?.ToString() ?? String.Empty;
 
 			[JsonIgnore]
-			public string Name => !string.IsNullOrEmpty(MaybeName) ? MaybeName : MaybeValueAsString;
+			public string Name => !String.IsNullOrEmpty(MaybeName) ? MaybeName : MaybeValueAsString;
 		}
 
 		private void CheckSOInFolderAre16KbCompatible(string FolderPath, UnrealArch UnrealArch)
@@ -6516,7 +6516,7 @@ import java.util.Collection;
 								int TraceStartIndex = KVP.Value.IndexOf("Trace.");
 								if (TraceStartIndex == 0 || (TraceStartIndex > 0 && KVP.Value[TraceStartIndex - 1] == '\n'))
 								{
-									String Addition = KVP.Value;
+									string Addition = KVP.Value;
 									do
 									{
 										int TraceEndIndex = Addition.IndexOf("\n", TraceStartIndex);
@@ -6639,7 +6639,7 @@ import java.util.Collection;
 								int TraceStartIndex = KVP.Value.IndexOf("Trace.");
 								if (TraceStartIndex == 0 || (TraceStartIndex > 0 && KVP.Value[TraceStartIndex - 1] == '\n'))
 								{
-									String Addition = KVP.Value;
+									string Addition = KVP.Value;
 									do
 									{
 										int TraceEndIndex = Addition.IndexOf("\n", TraceStartIndex);
diff --git a/Engine/Source/Programs/UnrealBuildTool/Platform/Apple/UEBuildApple.cs b/Engine/Source/Programs/UnrealBuildTool/Platform/Apple/UEBuildApple.cs
--- a/Engine/Source/Programs/UnrealBuildTool/Platform/Apple/UEBuildApple.cs
+++ b/Engine/Source/Programs/UnrealBuildTool/Platform/Apple/UEBuildApple.cs
@@ -156,11 +156,11 @@ namespace UnrealBuildTool
 		{
 			base.ValidateTarget(Target);
 			
-			if (!string.IsNullOrWhiteSpace(Environment.GetEnvironmentVariable("CLANG_STATIC_ANALYZER_MODE")))
+			if (!String.IsNullOrWhiteSpace(Environment.GetEnvironmentVariable("CLANG_STATIC_ANALYZER_MODE")))
 			{
 				Target.StaticAnalyzer = StaticAnalyzer.Default;
 				Target.StaticAnalyzerOutputType = (Environment.GetEnvironmentVariable("CLANG_ANALYZER_OUTPUT")?.Contains("html", StringComparison.OrdinalIgnoreCase) == true) ? StaticAnalyzerOutputType.Html : StaticAnalyzerOutputType.Text;
-				Target.StaticAnalyzerMode = string.Equals(Environment.GetEnvironmentVariable("CLANG_STATIC_ANALYZER_MODE"), "shallow", StringComparison.OrdinalIgnoreCase) ? StaticAnalyzerMode.Shallow : StaticAnalyzerMode.Deep;
+				Target.StaticAnalyzerMode = String.Equals(Environment.GetEnvironmentVariable("CLANG_STATIC_ANALYZER_MODE"), "shallow", StringComparison.OrdinalIgnoreCase) ? StaticAnalyzerMode.Shallow : StaticAnalyzerMode.Deep;
 			}
 			else if (Target.StaticAnalyzer == StaticAnalyzer.Clang)
 			{
diff --git a/Engine/Source/Programs/UnrealBuildTool/Platform/Linux/LinuxProjectGenerator.cs b/Engine/Source/Programs/UnrealBuildTool/Platform/Linux/LinuxProjectGenerator.cs
--- a/Engine/Source/Programs/UnrealBuildTool/Platform/Linux/LinuxProjectGenerator.cs
+++ b/Engine/Source/Programs/UnrealBuildTool/Platform/Linux/LinuxProjectGenerator.cs
@@ -47,7 +47,7 @@ namespace UnrealBuildTool
 			List<string> Result = [];
 			string EngineDirectory = Unreal.EngineDirectory.ToString();
 			string? UseLibcxxEnvVarOverride = Environment.GetEnvironmentVariable("UE_LINUX_USE_LIBCXX");
-			bool ShouldUseLibcxx = string.IsNullOrEmpty(UseLibcxxEnvVarOverride) || UseLibcxxEnvVarOverride == "1";
+			bool ShouldUseLibcxx = String.IsNullOrEmpty(UseLibcxxEnvVarOverride) || UseLibcxxEnvVarOverride == "1";
 
 			UnrealArch TargetArchitecture = InTarget.Architectures.SingleArchitecture;
 
diff --git a/Engine/Source/Programs/UnrealBuildTool/Platform/Mac/MacToolChain.cs b/Engine/Source/Programs/UnrealBuildTool/Platform/Mac/MacToolChain.cs
--- a/Engine/Source/Programs/UnrealBuildTool/Platform/Mac/MacToolChain.cs
+++ b/Engine/Source/Programs/UnrealBuildTool/Platform/Mac/MacToolChain.cs
@@ -652,7 +652,7 @@ namespace UnrealBuildTool
 			string IntermediateDirectory = (ProjectFile == null || LinkEnvironment.OutputFilePaths.All(x => x.IsUnderDirectory(Unreal.EngineDirectory)) || bIsEngineDylib ? Unreal.EngineDirectory : ProjectFile.Directory) + "/Intermediate/Build/Mac";
 			string DylibFilename = Path.GetFileName(DylibPath);
 
-			string AppName = string.Empty;
+			string AppName = String.Empty;
 			int firstDash = DylibFilename.IndexOf('-');
 			if (firstDash != -1)
 			{
@@ -759,7 +759,7 @@ namespace UnrealBuildTool
 					HashSet<FileReference> UsedSourceLibs = new();
 					void ProcessLib(FileReference SourceLib, FileReference DestLib, FileReference ReferencingFile)
 					{
-						if (string.Equals(SourceLib.GetExtension(), ".dylib", StringComparison.OrdinalIgnoreCase))
+						if (String.Equals(SourceLib.GetExtension(), ".dylib", StringComparison.OrdinalIgnoreCase))
 						{
 							DirectoryReference DestDir = DestLib.Directory;
 							if (UsedSourceLibs.Contains(SourceLib) || UsedDestDirs.Contains(DestDir))
diff --git a/Engine/Source/Programs/UnrealBuildTool/Platform/Windows/PVSToolChain.cs b/Engine/Source/Programs/UnrealBuildTool/Platform/Windows/PVSToolChain.cs
--- a/Engine/Source/Programs/UnrealBuildTool/Platform/Windows/PVSToolChain.cs
+++ b/Engine/Source/Programs/UnrealBuildTool/Platform/Windows/PVSToolChain.cs
@@ -1009,7 +1009,7 @@ namespace UnrealBuildTool
 				FileItem output = FileItem.GetItemByPath(match.Groups[2].Value);
 				FileItem ignored = FileItem.GetItemByPath(match.Groups[3].Value);
 				FileItem rootPaths = FileItem.GetItemByPath(match.Groups[4].Value);
-				int printLevel = int.Parse(match.Groups[5].Value);
+				int printLevel = Int32.Parse(match.Groups[5].Value);
 
 				FileItem newInput = FileItem.GetItemByFileReference(input.Location.ChangeExtension($".single{input.Location.GetExtension()}"));
 				FileItem newOutput = FileItem.GetItemByFileReference(output.Location.ChangeExtension($".single{output.Location.GetExtension()}"));
diff --git a/Engine/Source/Programs/UnrealBuildTool/Platform/Windows/WindowsProjectGenerator.cs b/Engine/Source/Programs/UnrealBuildTool/Platform/Windows/WindowsProjectGenerator.cs
--- a/Engine/Source/Programs/UnrealBuildTool/Platform/Windows/WindowsProjectGenerator.cs
+++ b/Engine/Source/Programs/UnrealBuildTool/Platform/Windows/WindowsProjectGenerator.cs
@@ -240,7 +240,7 @@ namespace UnrealBuildTool
 				}
 			}
 
-			return !string.IsNullOrEmpty(RemoteDebugMachine) && !string.IsNullOrEmpty(RemoteDebugDirectory);
+			return !System.String.IsNullOrEmpty(RemoteDebugMachine) && !System.String.IsNullOrEmpty(RemoteDebugDirectory);
 		}
 	}
 }
diff --git a/Engine/Source/Programs/UnrealBuildTool/ProjectFiles/AliasedFile.cs b/Engine/Source/Programs/UnrealBuildTool/ProjectFiles/AliasedFile.cs
--- a/Engine/Source/Programs/UnrealBuildTool/ProjectFiles/AliasedFile.cs
+++ b/Engine/Source/Programs/UnrealBuildTool/ProjectFiles/AliasedFile.cs
@@ -17,7 +17,7 @@ namespace UnrealBuildTool
 			this.ProjectPath = ProjectPath;
 		}
 
-		public int CompareTo(AliasedFile other) => string.CompareOrdinal(Location.FullName, other.Location.FullName);
+		public int CompareTo(AliasedFile other) => String.CompareOrdinal(Location.FullName, other.Location.FullName);
 
 		// Full location on disk.
 		public readonly FileReference Location;
diff --git a/Engine/Source/Programs/UnrealBuildTool/ProjectFiles/VisualStudioWorkspace/AvailableConfiguration.cs b/Engine/Source/Programs/UnrealBuildTool/ProjectFiles/VisualStudioWorkspace/AvailableConfiguration.cs
--- a/Engine/Source/Programs/UnrealBuildTool/ProjectFiles/VisualStudioWorkspace/AvailableConfiguration.cs
+++ b/Engine/Source/Programs/UnrealBuildTool/ProjectFiles/VisualStudioWorkspace/AvailableConfiguration.cs
@@ -1,19 +1,20 @@
 // Copyright Epic Games, Inc. All Rights Reserved.
 
+using System;
 using System.Collections.Generic;
 
 namespace UnrealBuildTool
 {
 	internal class AvailableConfiguration
 	{
-		public HashSet<string> Configurations { get; set; } = new();
+		public HashSet<string> Configurations { get; set; } = [];
 
-		public HashSet<string> Platforms { get; set; } = new();
+		public HashSet<string> Platforms { get; set; } = [];
 
-		public string TargetPath { get; set; } = string.Empty;
+		public string TargetPath { get; set; } = String.Empty;
 
-		public string ProjectPath { get; set; } = string.Empty;
+		public string ProjectPath { get; set; } = String.Empty;
 
-		public string TargetType { get; set; } = string.Empty;
+		public string TargetType { get; set; } = String.Empty;
 	}
 }
diff --git a/Engine/Source/Programs/UnrealBuildTool/ProjectFiles/VisualStudioWorkspace/VSWorkspaceProjectFile.cs b/Engine/Source/Programs/UnrealBuildTool/ProjectFiles/VisualStudioWorkspace/VSWorkspaceProjectFile.cs
--- a/Engine/Source/Programs/UnrealBuildTool/ProjectFiles/VisualStudioWorkspace/VSWorkspaceProjectFile.cs
+++ b/Engine/Source/Programs/UnrealBuildTool/ProjectFiles/VisualStudioWorkspace/VSWorkspaceProjectFile.cs
@@ -67,11 +67,11 @@ namespace UnrealBuildTool
 						availableConfiguration = new AvailableConfiguration();
 						availableConfigurations.Add(ProjectTarget.Name, availableConfiguration);
 					}
-					availableConfiguration.TargetType = ProjectTarget.TargetRules?.Type.ToString() ?? string.Empty;
+					availableConfiguration.TargetType = ProjectTarget.TargetRules?.Type.ToString() ?? String.Empty;
 					availableConfiguration.Configurations.Add(Configuration.ToString());
 					availableConfiguration.Platforms.Add(Platform.ToString());
 					availableConfiguration.TargetPath = ProjectTarget.TargetFilePath.FullName;
-					availableConfiguration.ProjectPath = ProjectTarget.UnrealProjectFilePath?.FullName ?? string.Empty;
+					availableConfiguration.ProjectPath = ProjectTarget.UnrealProjectFilePath?.FullName ?? String.Empty;
 
 					return Task.CompletedTask;
 				});
@@ -137,7 +137,7 @@ namespace UnrealBuildTool
 							sb.Append('_');
 							sb.Append(TargetInfo.Platform);
 
-							if (!string.IsNullOrEmpty(TargetInfo.Architecture))
+							if (!String.IsNullOrEmpty(TargetInfo.Architecture))
 							{
 								sb.Append('_');
 								sb.Append(TargetInfo.Architecture);
@@ -432,7 +432,7 @@ namespace UnrealBuildTool
 				return DefaultPlatformName;
 			}
 
-			return string.Empty;
+			return String.Empty;
 		}
 
 		private string GetVCConfiguration(UEBuildTarget Target, UnrealArch? Architecture, PlatformProjectGeneratorCollection Generators, VCProjectFileFormat VCFormat)
@@ -449,10 +449,10 @@ namespace UnrealBuildTool
 					}
 				}
 
-				return string.Format("{0}{1}_{2}", Target.Platform.ToString(), Architecture != null ? "_" + Architecture.ToString() : string.Empty, Target.Configuration.ToString());
+				return String.Format("{0}{1}_{2}", Target.Platform.ToString(), Architecture != null ? "_" + Architecture.ToString() : String.Empty, Target.Configuration.ToString());
 			}
 
-			return string.Empty;
+			return String.Empty;
 		}
 
 		private string GetArchitecture(UEBuildTarget Target, UnrealArch? Architecture, PlatformProjectGeneratorCollection Generators, VCProjectFileFormat VCFormat)
@@ -465,18 +465,18 @@ namespace UnrealBuildTool
 				{
 					if (!Generator.RequiresDistinctVisualStudioConfigurationName(VSSettings))
 					{
-						return string.Empty;
+						return String.Empty;
 					}
 				}
 
 				if (Architecture != null)
 				{
 					string? result = Architecture.ToString();
-					return result != null ? result : string.Empty;
+					return result != null ? result : String.Empty;
 				}
 			}
 
-			return string.Empty;
+			return String.Empty;
 		}
 
 		internal class ExportedTargetInfo
diff --git a/Engine/Source/Programs/UnrealBuildTool/ProjectFiles/VisualStudioWorkspace/VSWorkspaceProjectFileGenerator.cs b/Engine/Source/Programs/UnrealBuildTool/ProjectFiles/VisualStudioWorkspace/VSWorkspaceProjectFileGenerator.cs
--- a/Engine/Source/Programs/UnrealBuildTool/ProjectFiles/VisualStudioWorkspace/VSWorkspaceProjectFileGenerator.cs
+++ b/Engine/Source/Programs/UnrealBuildTool/ProjectFiles/VisualStudioWorkspace/VSWorkspaceProjectFileGenerator.cs
@@ -64,10 +64,10 @@ namespace UnrealBuildTool
 		private JsonWriterStyle Minimize = JsonWriterStyle.Readable;
 
 		[CommandLine("-Query=", Description = "Outputs all the available configurations, platforms and targets to a file provided by the user.")]
-		private string Query = string.Empty;
+		private string Query = String.Empty;
 
 		[CommandLine("-TargetDescriptionFile=", Description = "")]
-		private string TargetDescriptionFile = string.Empty;
+		private string TargetDescriptionFile = String.Empty;
 
 		[CommandLine("-WorkspaceGeneratorVersion", Description = "Show version information and exit.")]
 		private bool ShowVersionInfo = false;
@@ -276,7 +276,7 @@ namespace UnrealBuildTool
 				AddAllGameProjects(GameProjects);
 			}
 
-			if (!string.IsNullOrEmpty(Query))
+			if (!String.IsNullOrEmpty(Query))
 			{
 				return WriteAvailableConfigurations(PlatformProjectGenerators, Logger);
 			}
@@ -311,8 +311,8 @@ namespace UnrealBuildTool
 						ProjectPath = g.Select(i => i.ProjectPath).Distinct().Single(),
 						Configurations = g.Select(i => i.Configuration).Distinct().ToList(),
 						Platforms = g.Select(i => i.Platform).Distinct().ToList(),
-						Architectures = g.Select(i => i.Architecture).Where(i => !string.IsNullOrEmpty(i)).Distinct().ToList(),
-						ConfigurationsAndPlatforms = g.Select(i => string.Format("{0}-{1}{2}", i.Configuration, i.Platform, !string.IsNullOrEmpty(i.Architecture) ? "-" + i.Architecture : string.Empty)).Distinct().ToList(),
+						Architectures = g.Select(i => i.Architecture).Where(i => !String.IsNullOrEmpty(i)).Distinct().ToList(),
+						ConfigurationsAndPlatforms = g.Select(i => String.Format("{0}-{1}{2}", i.Configuration, i.Platform, !String.IsNullOrEmpty(i.Architecture) ? "-" + i.Architecture : String.Empty)).Distinct().ToList(),
 					}));
 
 				// The inner Targets object is needed for schema compatibility with the Query Mode API.
@@ -345,7 +345,7 @@ namespace UnrealBuildTool
 
 		private AvailableTargetsConfigurations? ReadProjectGenerateDescriptionFilePath()
 		{
-			if (string.IsNullOrEmpty(TargetDescriptionFile))
+			if (String.IsNullOrEmpty(TargetDescriptionFile))
 			{
 				return null;
 			}
@@ -581,7 +581,7 @@ namespace UnrealBuildTool
 									}
 									else
 									{
-										if (NodesToUpdate.TryGetValue(Node.Name, out string? NewNodeValue) && !string.IsNullOrEmpty(NewNodeValue))
+										if (NodesToUpdate.TryGetValue(Node.Name, out string? NewNodeValue) && !String.IsNullOrEmpty(NewNodeValue))
 										{
 											Node.InnerText = NewNodeValue;
 										}
diff --git a/Engine/Source/Programs/UnrealBuildTool/ProjectFiles/Xcode/XcodeProject.cs b/Engine/Source/Programs/UnrealBuildTool/ProjectFiles/Xcode/XcodeProject.cs
--- a/Engine/Source/Programs/UnrealBuildTool/ProjectFiles/Xcode/XcodeProject.cs
+++ b/Engine/Source/Programs/UnrealBuildTool/ProjectFiles/Xcode/XcodeProject.cs
@@ -2209,14 +2209,14 @@ namespace UnrealBuildTool.XcodeProjectXcconfig
 			if (PlistMetadata.Mode == MetadataMode.UpdateTemplate)
 			{
 				bool bAssetCatalogAllIcons = false;
-				String? AssetCatalogAlternates = null;
+				string? AssetCatalogAlternates = null;
 				PlatformIni.TryGetValue("/Script/IOSRuntimeSettings.IOSRuntimeSettings", "bAssetCatalogAllIcons", out bAssetCatalogAllIcons);
 				PlatformIni.TryGetValue("/Script/IOSRuntimeSettings.IOSRuntimeSettings", "AssetCatalogAlternates", out AssetCatalogAlternates);
 
 				// allow Xcode to generate the final plist file from our input, some INFOPLIST settings and other settings 
 				Xcconfig.AppendLine("GENERATE_INFOPLIST_FILE = YES");
 				Xcconfig.AppendLine($"ASSETCATALOG_COMPILER_APPICON_NAME = AppIcon");
-				String AllIcons = bAssetCatalogAllIcons ? "YES" : "NO";
+				string AllIcons = bAssetCatalogAllIcons ? "YES" : "NO";
 				Xcconfig.AppendLine($"ASSETCATALOG_COMPILER_INCLUDE_ALL_APPICON_ASSETS = {AllIcons}");
 				if (!String.IsNullOrEmpty(AssetCatalogAlternates))
 				{
diff --git a/Engine/Source/Programs/UnrealBuildTool/System/CppDependencyCache.cs b/Engine/Source/Programs/UnrealBuildTool/System/CppDependencyCache.cs
--- a/Engine/Source/Programs/UnrealBuildTool/System/CppDependencyCache.cs
+++ b/Engine/Source/Programs/UnrealBuildTool/System/CppDependencyCache.cs
@@ -533,7 +533,7 @@ namespace UnrealBuildTool
 				TokenIdx += 2;
 
 				// TODO: Move this to the provided from the outside somewhere.
-				string CurrentDir = string.Empty;
+				string CurrentDir = String.Empty;
 
 				List<FileItem> NewDependencyFiles = new List<FileItem>();
 				for (; TokenIdx < Tokens.Count && Tokens[TokenIdx] != "\n"; TokenIdx++)
diff --git a/Engine/Source/Programs/UnrealBuildTool/System/FileHashCache.cs b/Engine/Source/Programs/UnrealBuildTool/System/FileHashCache.cs
--- a/Engine/Source/Programs/UnrealBuildTool/System/FileHashCache.cs
+++ b/Engine/Source/Programs/UnrealBuildTool/System/FileHashCache.cs
@@ -600,7 +600,7 @@ namespace UnrealBuildTool
 
 				for (int Idx = 0; Idx < count; Idx++)
 				{
-					UInt16 matchLength = reader.ReadUInt16();
+					ushort matchLength = reader.ReadUInt16();
 					string str = reader.ReadString();
 					sb.Clear().Append(lastFilename, 0, matchLength).Append(str);
 
diff --git a/Engine/Source/Programs/UnrealBuildTool/System/UHTExecution.cs b/Engine/Source/Programs/UnrealBuildTool/System/UHTExecution.cs
--- a/Engine/Source/Programs/UnrealBuildTool/System/UHTExecution.cs
+++ b/Engine/Source/Programs/UnrealBuildTool/System/UHTExecution.cs
@@ -500,9 +500,9 @@ namespace UnrealBuildTool
 
 					DirectoryReference[] ModuleIncludePaths = [.. Module.PublicIncludePaths.Union(Module.InternalIncludePaths).Union(Module.PrivateIncludePaths)];
 
-					string VerseMountPoint = string.Empty;
-					string VersePackageName = string.Empty;
-					string VerseDirectoryPath = string.Empty;
+					string VerseMountPoint = String.Empty;
+					string VersePackageName = String.Empty;
+					string VerseDirectoryPath = String.Empty;
 					if (Module.bHasVerse)
 					{
 						DirectoryReference SourceDirectory = DirectoryReference.Combine(Module.Rules.Directory, "Verse");
