UnrealBuildTool: Remove and sort usings

#rnx

diff --git a/Engine/Source/Programs/UnrealBuildTool/Platform/Apple/AppleExports.cs b/Engine/Source/Programs/UnrealBuildTool/Platform/Apple/AppleExports.cs
--- a/Engine/Source/Programs/UnrealBuildTool/Platform/Apple/AppleExports.cs
+++ b/Engine/Source/Programs/UnrealBuildTool/Platform/Apple/AppleExports.cs
@@ -1,10 +1,10 @@
 ﻿// Copyright Epic Games, Inc. All Rights Reserved.
 
 using System;
-using EpicGames.Core;
 using System.Collections.Generic;
-using System.Linq;
 using System.IO;
+using System.Linq;
+using EpicGames.Core;
 using Microsoft.Extensions.Logging;
 using UnrealBuildBase;
 
diff --git a/Engine/Source/Programs/UnrealBuildTool/System/CppCompileEnvironment.cs b/Engine/Source/Programs/UnrealBuildTool/System/CppCompileEnvironment.cs
--- a/Engine/Source/Programs/UnrealBuildTool/System/CppCompileEnvironment.cs
+++ b/Engine/Source/Programs/UnrealBuildTool/System/CppCompileEnvironment.cs
@@ -1,11 +1,11 @@
 // Copyright Epic Games, Inc. All Rights Reserved.
 
+using System;
 using System.Collections.Generic;
-using System.Linq;
 using System.IO.Enumeration;
+using System.Linq;
 using EpicGames.Core;
 using UnrealBuildBase;
-using System;
 
 namespace UnrealBuildTool
 {
