Fix or selectively disable issue IDE0044 (Make field readonly) in UBT.

diff --git a/Engine/Source/Programs/UnrealBuildTool/Configuration/CompileWarnings/WarningLevelResolverRegistry.cs b/Engine/Source/Programs/UnrealBuildTool/Configuration/CompileWarnings/WarningLevelResolverRegistry.cs
--- a/Engine/Source/Programs/UnrealBuildTool/Configuration/CompileWarnings/WarningLevelResolverRegistry.cs
+++ b/Engine/Source/Programs/UnrealBuildTool/Configuration/CompileWarnings/WarningLevelResolverRegistry.cs
@@ -32,7 +32,7 @@ namespace UnrealBuildTool
 	internal class WarningLevelResolverRegistry
 	{
 		private readonly ILogger? _logger = Log.Logger;
-		private static Lazy<WarningLevelResolverRegistry> s_resolverRegistry = new(() => new());
+		private static readonly Lazy<WarningLevelResolverRegistry> s_resolverRegistry = new(() => new());
 		private readonly Dictionary<string, IWarningLevelResolver> _attributeNameToBehaviour = [];
 
 		/// <summary>
diff --git a/Engine/Source/Programs/UnrealBuildTool/Configuration/Rules/CppCompileWarnings.cs b/Engine/Source/Programs/UnrealBuildTool/Configuration/Rules/CppCompileWarnings.cs
--- a/Engine/Source/Programs/UnrealBuildTool/Configuration/Rules/CppCompileWarnings.cs
+++ b/Engine/Source/Programs/UnrealBuildTool/Configuration/Rules/CppCompileWarnings.cs
@@ -1487,7 +1487,7 @@ namespace UnrealBuildTool
 
 		private static readonly Lazy<List<(MemberInfo, List<ApplyWarningsAttribute>)>> s_membersWithApplyWarningsAttribute = new(() => GetMembersWithAttributes<ApplyWarningsAttribute>(typeof(CppCompileWarnings), typeof(WarningLevel)));
 		private static readonly Lazy<List<(MemberInfo, List<WarningLevelDefaultAttribute>)>> s_membersWithWarningLevelDefaultAttribute = new(() => GetMembersWithAttributes<WarningLevelDefaultAttribute>(typeof(CppCompileWarnings), typeof(WarningLevel)).ToList());
-		private static ConcurrentDictionary<MemberInfo, Lazy<Dictionary<Type, List<WarningLevelDefaultAttribute>>>>[] s_groupedAttributes = { [], [], [] };
+		private static readonly ConcurrentDictionary<MemberInfo, Lazy<Dictionary<Type, List<WarningLevelDefaultAttribute>>>>[] s_groupedAttributes = { [], [], [] };
 
 		/// <summary>
 		/// Generates the compiler warnings command line arguments for the requested <see cref="UEToolChain"/> type, and version.
diff --git a/Engine/Source/Programs/UnrealBuildTool/Configuration/Rules/ModuleRules.cs b/Engine/Source/Programs/UnrealBuildTool/Configuration/Rules/ModuleRules.cs
--- a/Engine/Source/Programs/UnrealBuildTool/Configuration/Rules/ModuleRules.cs
+++ b/Engine/Source/Programs/UnrealBuildTool/Configuration/Rules/ModuleRules.cs
@@ -1018,7 +1018,7 @@ namespace UnrealBuildTool
 		/// </summary>
 		public HashSet<string> StaticAnalyzerDisabledCheckers { get; } = []; // This is set to DefaultStaticAnalyzerDisabledCheckers if static analyzer is set on target
 
-		private static HashSet<string> DefaultStaticAnalyzerDisabledCheckers = [
+		private static readonly HashSet<string> DefaultStaticAnalyzerDisabledCheckers = [
 			// Work in progress, remove disabled checkers as all warnings are fixed
 			"core.BitwiseShift",
 			"core.CallAndMessage",
diff --git a/Engine/Source/Programs/UnrealBuildTool/Configuration/Rules/RulesAssembly.cs b/Engine/Source/Programs/UnrealBuildTool/Configuration/Rules/RulesAssembly.cs
--- a/Engine/Source/Programs/UnrealBuildTool/Configuration/Rules/RulesAssembly.cs
+++ b/Engine/Source/Programs/UnrealBuildTool/Configuration/Rules/RulesAssembly.cs
@@ -428,7 +428,7 @@ namespace UnrealBuildTool
 			}
 		}
 
-		private ConcurrentDictionary<string, Type?> _moduleNameToType = [];
+		private readonly ConcurrentDictionary<string, Type?> _moduleNameToType = [];
 
 		/// <summary>
 		/// Gets the type defining rules for the given module within this assembly
diff --git a/Engine/Source/Programs/UnrealBuildTool/Configuration/Rules/RulesCompiler.cs b/Engine/Source/Programs/UnrealBuildTool/Configuration/Rules/RulesCompiler.cs
--- a/Engine/Source/Programs/UnrealBuildTool/Configuration/Rules/RulesCompiler.cs
+++ b/Engine/Source/Programs/UnrealBuildTool/Configuration/Rules/RulesCompiler.cs
@@ -45,7 +45,7 @@ namespace UnrealBuildTool
 		/// <summary>
 		/// Lock for creating engine rule assembly
 		/// </summary>
-		private static object EngineRulesAssemblyLock = new();
+		private static readonly object EngineRulesAssemblyLock = new();
 
 		/// <summary>
 		/// The cached rules assembly for engine modules and targets.
@@ -56,7 +56,7 @@ namespace UnrealBuildTool
 		/// Map of assembly names we've already compiled and loaded to their Assembly and list of game folders.  This is used to prevent
 		/// trying to recompile the same assembly when ping-ponging between different types of targets
 		/// </summary>
-		private static ConcurrentDictionary<FileReference, Lazy<RulesAssembly>> LoadedAssemblyMap = new();
+		private static readonly ConcurrentDictionary<FileReference, Lazy<RulesAssembly>> LoadedAssemblyMap = new();
 
 		/// <summary>
 		/// Creates the engine rules assembly
diff --git a/Engine/Source/Programs/UnrealBuildTool/Configuration/Rules/TestModuleRules.cs b/Engine/Source/Programs/UnrealBuildTool/Configuration/Rules/TestModuleRules.cs
--- a/Engine/Source/Programs/UnrealBuildTool/Configuration/Rules/TestModuleRules.cs
+++ b/Engine/Source/Programs/UnrealBuildTool/Configuration/Rules/TestModuleRules.cs
@@ -27,8 +27,6 @@ namespace UnrealBuildTool
 			RestrictedFolder.CarefullyRedist.ToString()
 		};
 
-		private bool bUsesCatch2 = true;
-
 		/// <summary>
 		/// Check if running in test mode.
 		/// </summary>
@@ -55,12 +53,11 @@ namespace UnrealBuildTool
 
 		/// <summary>
 		/// Constructs a TestModuleRules object as its own test module.
-		/// Sets value of bUsesCatch2.
+		/// Sets common properties if using Catch2.
 		/// </summary>
 		public TestModuleRules(ReadOnlyTargetRules Target, bool InUsesCatch2) : base(Target)
 		{
-			bUsesCatch2 = InUsesCatch2;
-			if (bUsesCatch2)
+			if (InUsesCatch2)
 			{
 				SetupCommonProperties(Target);
 			}
diff --git a/Engine/Source/Programs/UnrealBuildTool/Configuration/UEBuildBinary.cs b/Engine/Source/Programs/UnrealBuildTool/Configuration/UEBuildBinary.cs
--- a/Engine/Source/Programs/UnrealBuildTool/Configuration/UEBuildBinary.cs
+++ b/Engine/Source/Programs/UnrealBuildTool/Configuration/UEBuildBinary.cs
@@ -120,7 +120,7 @@ namespace UnrealBuildTool
 		/// <summary>
 		/// Cached list of dependent link libraries.
 		/// </summary>
-		private Lazy<List<FileReference>> DependentLinkLibraries;
+		private readonly Lazy<List<FileReference>> DependentLinkLibraries;
 
 		/// <summary>
 		/// Expanded list of runtime dependencies generated by calling PrepareRuntimeDependencies.
diff --git a/Engine/Source/Programs/UnrealBuildTool/Configuration/UEBuildModule.cs b/Engine/Source/Programs/UnrealBuildTool/Configuration/UEBuildModule.cs
--- a/Engine/Source/Programs/UnrealBuildTool/Configuration/UEBuildModule.cs
+++ b/Engine/Source/Programs/UnrealBuildTool/Configuration/UEBuildModule.cs
@@ -223,8 +223,8 @@ namespace UnrealBuildTool
 
 		static readonly string[] ApiDefinitions = ["", "DLLEXPORT", "DLLIMPORT"];
 
-		private string[] ModuleApiDefinitionStrings;
-		private string[] ModuleNonAttributedApiDefinitionStrings;
+		private readonly string[] ModuleApiDefinitionStrings;
+		private readonly string[] ModuleNonAttributedApiDefinitionStrings;
 
 		/// <summary>
 		/// Constructor
diff --git a/Engine/Source/Programs/UnrealBuildTool/Configuration/UEBuildModuleCPP.cs b/Engine/Source/Programs/UnrealBuildTool/Configuration/UEBuildModuleCPP.cs
--- a/Engine/Source/Programs/UnrealBuildTool/Configuration/UEBuildModuleCPP.cs
+++ b/Engine/Source/Programs/UnrealBuildTool/Configuration/UEBuildModuleCPP.cs
@@ -112,7 +112,7 @@ namespace UnrealBuildTool
 		/// <summary>
 		/// Verse source code directory associated with this module
 		/// </summary>
-		DirectoryItem? AssociatedVerseDirectory;
+		readonly DirectoryItem? AssociatedVerseDirectory;
 
 		/// <summary>
 		/// Create PerModuleInline source file. Defaults to PrimaryModule in a binary but could be another one if primary module does not depend on core
@@ -1323,7 +1323,7 @@ namespace UnrealBuildTool
 			return HeaderFileItems;
 		}
 
-		static ConcurrentDictionary<string, Lazy<HashSet<string>>> s_sharedPchModuleDependenciesCache = [];
+		static readonly ConcurrentDictionary<string, Lazy<HashSet<string>>> s_sharedPchModuleDependenciesCache = [];
 
 		/// <summary>
 		/// Gets the list of module dependencies that this module declares and then tries to optimize the list to what is actually used by the header.
@@ -2934,7 +2934,7 @@ namespace UnrealBuildTool
 
 		// List of known extensions that are checked in FindInputFilesFromDirectory, so we can log out only if a hidden file
 		// is skipped, only if it has an extension that would normally be compiled
-		static string[] KnownInputFileExtensions =
+		static readonly string[] KnownInputFileExtensions =
 		{
 			".h", ".ipsh",
 			".cpp", ".cxx", ".ixx", ".c", ".cc",
diff --git a/Engine/Source/Programs/UnrealBuildTool/Configuration/UEBuildPlatform.cs b/Engine/Source/Programs/UnrealBuildTool/Configuration/UEBuildPlatform.cs
--- a/Engine/Source/Programs/UnrealBuildTool/Configuration/UEBuildPlatform.cs
+++ b/Engine/Source/Programs/UnrealBuildTool/Configuration/UEBuildPlatform.cs
@@ -190,10 +190,10 @@ namespace UnrealBuildTool
 
 	abstract class UEBuildPlatform
 	{
-		private static Dictionary<UnrealTargetPlatform, UEBuildPlatform> BuildPlatformDictionary = new Dictionary<UnrealTargetPlatform, UEBuildPlatform>();
+		private static readonly Dictionary<UnrealTargetPlatform, UEBuildPlatform> BuildPlatformDictionary = new Dictionary<UnrealTargetPlatform, UEBuildPlatform>();
 
 		// a mapping of a group to the platforms in the group (ie, Microsoft contains Win32 and Win64)
-		static Dictionary<UnrealPlatformGroup, List<UnrealTargetPlatform>> PlatformGroupDictionary = new Dictionary<UnrealPlatformGroup, List<UnrealTargetPlatform>>();
+		static readonly Dictionary<UnrealPlatformGroup, List<UnrealTargetPlatform>> PlatformGroupDictionary = new Dictionary<UnrealPlatformGroup, List<UnrealTargetPlatform>>();
 
 		/// <summary>
 		/// The corresponding target platform enum
@@ -255,7 +255,7 @@ namespace UnrealBuildTool
 			}
 		}
 
-		private static string[] UATProjectParams = { "-project=", "-scriptsforproject=" };
+		private static readonly string[] UATProjectParams = { "-project=", "-scriptsforproject=" };
 		private static void InitializePerPlatformSDKs(string[] Args, bool bArgumentsAreForUBT, ILogger Logger)
 		{
 			// Before we setup AutoSDK, we check to see if any projects need to override the Main version so that AutoSDK
diff --git a/Engine/Source/Programs/UnrealBuildTool/Configuration/UEBuildTarget.cs b/Engine/Source/Programs/UnrealBuildTool/Configuration/UEBuildTarget.cs
--- a/Engine/Source/Programs/UnrealBuildTool/Configuration/UEBuildTarget.cs
+++ b/Engine/Source/Programs/UnrealBuildTool/Configuration/UEBuildTarget.cs
@@ -627,7 +627,7 @@ namespace UnrealBuildTool
 		private static UniqueStringRegistry? StringRegistry;
 
 		// tracks if non-default architectures are an x64 arch
-		private static Dictionary<int, bool> IsX64Map = new();
+		private static readonly Dictionary<int, bool> IsX64Map = new();
 
 		// #jira UE-88908 if parts of a partial struct have each static member variables, their initialization order does not appear guaranteed
 		// here this means initializing "StringRegistry" directly to "new UniqueStringRegistry()" may not be executed before FindOrAddByName() has been called as part of initializing a static member variable of another part of the partial struct
@@ -1839,9 +1839,9 @@ namespace UnrealBuildTool
 		/// <summary>
 		/// Used to keep track of all modules by name.
 		/// </summary>
-		private Dictionary<string, UEBuildModule> Modules = new(StringComparer.InvariantCultureIgnoreCase);
+		private readonly Dictionary<string, UEBuildModule> Modules = new(StringComparer.InvariantCultureIgnoreCase);
 
-		private Dictionary<string, UEBuildModule> UsedShortNames = new(StringComparer.InvariantCultureIgnoreCase);
+		private readonly Dictionary<string, UEBuildModule> UsedShortNames = new(StringComparer.InvariantCultureIgnoreCase);
 
 		/// <summary>
 		/// Filename for the receipt for this target.
diff --git a/Engine/Source/Programs/UnrealBuildTool/Executors/Experimental/FASTBuild.cs b/Engine/Source/Programs/UnrealBuildTool/Executors/Experimental/FASTBuild.cs
--- a/Engine/Source/Programs/UnrealBuildTool/Executors/Experimental/FASTBuild.cs
+++ b/Engine/Source/Programs/UnrealBuildTool/Executors/Experimental/FASTBuild.cs
@@ -73,7 +73,7 @@ namespace UnrealBuildTool
 		/// <summary>
 		/// Executor to use for local actions
 		/// </summary>
-		ActionExecutor LocalExecutor;
+		readonly ActionExecutor LocalExecutor;
 
 		public static readonly string DefaultExecutableBasePath = Path.Combine(Unreal.EngineDirectory.FullName, "Extras", "ThirdPartyNotUE", "FASTBuild");
 
diff --git a/Engine/Source/Programs/UnrealBuildTool/Executors/ParallelExecutor.cs b/Engine/Source/Programs/UnrealBuildTool/Executors/ParallelExecutor.cs
--- a/Engine/Source/Programs/UnrealBuildTool/Executors/ParallelExecutor.cs
+++ b/Engine/Source/Programs/UnrealBuildTool/Executors/ParallelExecutor.cs
@@ -20,6 +20,7 @@ namespace UnrealBuildTool
 	/// </summary>
 	class ParallelExecutor : ActionExecutor
 	{
+#pragma warning disable IDE0044 // Make field readonly - these private static fields are set by XML configuration.
 		/// <summary>
 		/// Maximum processor count for local execution. 
 		/// </summary>
@@ -46,6 +47,7 @@ namespace UnrealBuildTool
 		/// </summary>
 		[XmlConfigFile]
 		private static double MemoryPerActionBytes = 1.5 * 1024 * 1024 * 1024;
+#pragma warning restore IDE0044
 
 		/// <summary>
 		/// The priority to set for spawned processes.
diff --git a/Engine/Source/Programs/UnrealBuildTool/Executors/SNDBS.cs b/Engine/Source/Programs/UnrealBuildTool/Executors/SNDBS.cs
--- a/Engine/Source/Programs/UnrealBuildTool/Executors/SNDBS.cs
+++ b/Engine/Source/Programs/UnrealBuildTool/Executors/SNDBS.cs
@@ -49,8 +49,9 @@ namespace UnrealBuildTool
 		private static readonly FileReference IncludeRewriteRulesFile = FileReference.Combine(IntermediateDir, "include-rewrite-rules.ini");
 		private static readonly FileReference ScriptFile = FileReference.Combine(IntermediateDir, "sndbs.json");
 
-		private Dictionary<string, string> ActiveTemplates = BuiltInTemplates.ToDictionary(p => p.Key, p => p.Value);
+		private readonly Dictionary<string, string> ActiveTemplates = BuiltInTemplates.ToDictionary(p => p.Key, p => p.Value);
 
+#pragma warning disable IDE0044 // Make field readonly - these private static fields are set by XML configuration.
 		/// <summary>
 		/// When enabled, dependency scanning is prioritized to avoid being slowed down by running local jobs.
 		/// </summary>
@@ -74,10 +75,11 @@ namespace UnrealBuildTool
 		/// </summary>
 		[XmlConfigFile(Category = "SNDBS")]
 		static string[]? VpnSubnets = null;
+#pragma warning restore IDE0044
 
 		private const string ProgressMarkupPrefix = "@action:";
 
-		private List<TargetDescriptor> TargetDescriptors;
+		private readonly List<TargetDescriptor> TargetDescriptors;
 
 		public SNDBS(List<TargetDescriptor> InTargetDescriptors, ILogger Logger)
 			: base(Logger)
diff --git a/Engine/Source/Programs/UnrealBuildTool/Executors/UnrealBuildAccelerator/UBAExecutor.cs b/Engine/Source/Programs/UnrealBuildTool/Executors/UnrealBuildAccelerator/UBAExecutor.cs
--- a/Engine/Source/Programs/UnrealBuildTool/Executors/UnrealBuildAccelerator/UBAExecutor.cs
+++ b/Engine/Source/Programs/UnrealBuildTool/Executors/UnrealBuildAccelerator/UBAExecutor.cs
@@ -1111,7 +1111,7 @@ namespace UnrealBuildTool
 			}
 		}
 
-		HashSet<uint> _usedBuckets = new();
+		readonly HashSet<uint> _usedBuckets = new();
 
 		uint GetCacheBucket(ILogger logger, LinkedAction action)
 		{
diff --git a/Engine/Source/Programs/UnrealBuildTool/Executors/XGE.cs b/Engine/Source/Programs/UnrealBuildTool/Executors/XGE.cs
--- a/Engine/Source/Programs/UnrealBuildTool/Executors/XGE.cs
+++ b/Engine/Source/Programs/UnrealBuildTool/Executors/XGE.cs
@@ -23,6 +23,7 @@ namespace UnrealBuildTool
 {
 	class XGE : ActionExecutor
 	{
+#pragma warning disable IDE0044 // Make field readonly - these private static fields are set by XML configuration.
 		/// <summary>
 		/// Whether to use the no_watchdog_thread option to prevent VS2015 toolchain stalls.
 		/// </summary>
@@ -77,10 +78,11 @@ namespace UnrealBuildTool
 		/// </summary>
 		[XmlConfigFile(Category = "XGE")]
 		static bool bUnavailableIfInUse = false;
+#pragma warning restore IDE0044
 
 		private const string ProgressMarkupPrefix = "@action";
 
-		private static List<string> CompileAutoRecover = new List<string> {
+		private static readonly List<string> CompileAutoRecover = new List<string> {
 			"C1060", // C1060: compiler is out of heap space
 			"C1076", // C1076: compiler limit: internal heap limit reached
 			"C2855", // C2855: command-line option 'X' inconsistent with precompiled header
@@ -88,7 +90,7 @@ namespace UnrealBuildTool
 			"C3859", // C3859: Failed to create virtual memory for PCH
 		};
 
-		private static List<string> LinkAutoRecover = new List<string> {
+		private static readonly List<string> LinkAutoRecover = new List<string> {
 			"Unexpected PDB error; OK (0)"
 		};
 
diff --git a/Engine/Source/Programs/UnrealBuildTool/Modes/AggregateClangTimingInfo.cs b/Engine/Source/Programs/UnrealBuildTool/Modes/AggregateClangTimingInfo.cs
--- a/Engine/Source/Programs/UnrealBuildTool/Modes/AggregateClangTimingInfo.cs
+++ b/Engine/Source/Programs/UnrealBuildTool/Modes/AggregateClangTimingInfo.cs
@@ -167,7 +167,7 @@ namespace UnrealBuildTool
 			}
 		}
 
-		private ConcurrentDictionary<FileReference, ClangTrace> ClangTraceCache = new();
+		private readonly ConcurrentDictionary<FileReference, ClangTrace> ClangTraceCache = new();
 
 		public override Task<int> ExecuteAsync(CommandLineArguments Arguments, ILogger Logger)
 		{
diff --git a/Engine/Source/Programs/UnrealBuildTool/Modes/AggregateParsedTimingInfo.cs b/Engine/Source/Programs/UnrealBuildTool/Modes/AggregateParsedTimingInfo.cs
--- a/Engine/Source/Programs/UnrealBuildTool/Modes/AggregateParsedTimingInfo.cs
+++ b/Engine/Source/Programs/UnrealBuildTool/Modes/AggregateParsedTimingInfo.cs
@@ -22,12 +22,12 @@ namespace UnrealBuildTool
 	class AggregateParsedTimingInfo : ToolMode
 	{
 		private TimingData? FileTimingData;
-		private ConcurrentDictionary<string, int> DecompressedFileSizes = new ConcurrentDictionary<string, int>();
-		private ConcurrentDictionary<string, byte[]> CompressedFiles = new ConcurrentDictionary<string, byte[]>();
-		private ConcurrentBag<TimingData> AggregateIncludes = new ConcurrentBag<TimingData>();
-		private ConcurrentBag<TimingData> AggregateClasses = new ConcurrentBag<TimingData>();
-		private ConcurrentBag<TimingData> AggregateFunctions = new ConcurrentBag<TimingData>();
-		private object AddFileLock = new object();
+		private readonly ConcurrentDictionary<string, int> DecompressedFileSizes = new ConcurrentDictionary<string, int>();
+		private readonly ConcurrentDictionary<string, byte[]> CompressedFiles = new ConcurrentDictionary<string, byte[]>();
+		private readonly ConcurrentBag<TimingData> AggregateIncludes = new ConcurrentBag<TimingData>();
+		private readonly ConcurrentBag<TimingData> AggregateClasses = new ConcurrentBag<TimingData>();
+		private readonly ConcurrentBag<TimingData> AggregateFunctions = new ConcurrentBag<TimingData>();
+		private readonly object AddFileLock = new object();
 
 		public override Task<int> ExecuteAsync(CommandLineArguments Arguments, ILogger Logger)
 		{
@@ -182,7 +182,7 @@ namespace UnrealBuildTool
 		}
 
 		private int HeaderFilesFound = 0;
-		private object CounterLock = new object();
+		private readonly object CounterLock = new object();
 		private void FlattenIncludes(Dictionary<string, TimingData> Includes, IEnumerable<TimingData> IncludeData)
 		{
 			foreach (TimingData Include in IncludeData)
diff --git a/Engine/Source/Programs/UnrealBuildTool/Modes/CleanMode.cs b/Engine/Source/Programs/UnrealBuildTool/Modes/CleanMode.cs
--- a/Engine/Source/Programs/UnrealBuildTool/Modes/CleanMode.cs
+++ b/Engine/Source/Programs/UnrealBuildTool/Modes/CleanMode.cs
@@ -20,11 +20,13 @@ namespace UnrealBuildTool
 	[ToolMode("Clean", ToolModeOptions.XmlConfig | ToolModeOptions.BuildPlatforms | ToolModeOptions.SingleInstance)]
 	class CleanMode : ToolMode
 	{
+#pragma warning disable IDE0044 // Make field readonly - these private static fields are set by command-line parsing.
 		/// <summary>
 		/// Whether to avoid cleaning targets
 		/// </summary>
 		[CommandLine("-SkipRulesCompile")]
 		bool bSkipRulesCompile = false;
+#pragma warning restore IDE0044
 
 		/// <summary>
 		/// Skip pre build targets; just do the main target.
diff --git a/Engine/Source/Programs/UnrealBuildTool/Modes/FixIncludePathsMode.cs b/Engine/Source/Programs/UnrealBuildTool/Modes/FixIncludePathsMode.cs
--- a/Engine/Source/Programs/UnrealBuildTool/Modes/FixIncludePathsMode.cs
+++ b/Engine/Source/Programs/UnrealBuildTool/Modes/FixIncludePathsMode.cs
@@ -27,11 +27,13 @@ namespace UnrealBuildTool
 		static readonly string[] PreferredPaths = { "/Public/", "/Private/", "/Classes/", "/Internal/", "/UHT/", "/VNI/" };
 		static readonly string[] PublicDirectories = { "Public", "Classes", };
 
+#pragma warning disable IDE0044 // Make field readonly - these private static fields are set by command-line parsing.
 		[CommandLine("-Filter=", Description = "Set of filters for files to include in the database. Relative to the root directory, or to the project file.")]
 		List<string> FilterRules = new List<string>();
 
 		[CommandLine("-IncludeFilter=", Description = "Set of filters for #include'd lines to allow updating. Relative to the root directory, or to the project file.")]
 		List<string> IncludeFilterRules = new List<string>();
+#pragma warning restore IDE0044
 
 		[CommandLine("-CheckoutWithP4", Description = "Flags that this task should use p4 to check out the file before updating it.")]
 		public bool bCheckoutWithP4 = false;
@@ -349,7 +351,7 @@ namespace UnrealBuildTool
 
 		class HeaderSortComparison : IComparer<string>
 		{
-			private string IWYUFileName;
+			private readonly string IWYUFileName;
 
 			public HeaderSortComparison(string IWYUFileName)
 			{
diff --git a/Engine/Source/Programs/UnrealBuildTool/Modes/GenerateClangDatabase.cs b/Engine/Source/Programs/UnrealBuildTool/Modes/GenerateClangDatabase.cs
--- a/Engine/Source/Programs/UnrealBuildTool/Modes/GenerateClangDatabase.cs
+++ b/Engine/Source/Programs/UnrealBuildTool/Modes/GenerateClangDatabase.cs
@@ -17,6 +17,7 @@ namespace UnrealBuildTool
 	[ToolMode("GenerateClangDatabase", ToolModeOptions.XmlConfig | ToolModeOptions.BuildPlatforms | ToolModeOptions.SingleInstance | ToolModeOptions.StartPrefetchingEngine | ToolModeOptions.ShowExecutionTime)]
 	class GenerateClangDatabase : ToolMode
 	{
+#pragma warning disable IDE0044 // Make field readonly - these private static fields are set by command-line parsing.
 		/// <summary>
 		/// Optional set of source files to include in the compile database. If this is empty, all files will be included by default and -Exclude can be used to exclude some.
 		/// Relative to the root directory, or to the project file.
@@ -30,6 +31,7 @@ namespace UnrealBuildTool
 		/// </summary>
 		[CommandLine("-Exclude=")]
 		List<string> ExcludeRules = new List<string>();
+#pragma warning restore IDE0044
 
 		/// <summary>
 		/// Execute any actions which result in code generation (eg. ISPC compilation)
diff --git a/Engine/Source/Programs/UnrealBuildTool/Modes/GenerateProjectFilesMode.cs b/Engine/Source/Programs/UnrealBuildTool/Modes/GenerateProjectFilesMode.cs
--- a/Engine/Source/Programs/UnrealBuildTool/Modes/GenerateProjectFilesMode.cs
+++ b/Engine/Source/Programs/UnrealBuildTool/Modes/GenerateProjectFilesMode.cs
@@ -18,6 +18,7 @@ namespace UnrealBuildTool
 	[ToolMode("GenerateProjectFiles", ToolModeOptions.XmlConfig | ToolModeOptions.BuildPlatforms | ToolModeOptions.SingleInstance | ToolModeOptions.UseStartupTraceListener | ToolModeOptions.StartPrefetchingEngine | ToolModeOptions.ShowExecutionTime)]
 	class GenerateProjectFilesMode : ToolMode
 	{
+#pragma warning disable IDE0044 // Make field readonly - these private static fields are set by command-line parsing.
 		/// <summary>
 		/// Types of project files to generate
 		/// </summary>
@@ -52,6 +53,7 @@ namespace UnrealBuildTool
 		/// </summary>
 		[CommandLine("-Automated")]
 		bool bAutomated = false;
+#pragma warning restore IDE0044
 
 		/// <summary>
 		/// Execute the tool mode
diff --git a/Engine/Source/Programs/UnrealBuildTool/Modes/PipInstallMode.cs b/Engine/Source/Programs/UnrealBuildTool/Modes/PipInstallMode.cs
--- a/Engine/Source/Programs/UnrealBuildTool/Modes/PipInstallMode.cs
+++ b/Engine/Source/Programs/UnrealBuildTool/Modes/PipInstallMode.cs
@@ -200,14 +200,14 @@ namespace UnrealBuildTool.Modes
 		private const string ExtraUrlsFilename = "extra_urls.txt";
 		private const string MergedRequirementsFilename = "merged_requirements.txt";
 
-		private UnrealArch TargetArch;
-		private UnrealTargetPlatform TargetPlatform;
-		private DirectoryReference InstallDir;
-		private FileReference PythonVenvExe;
-		private string? PythonVenvVer;
+		private readonly UnrealArch TargetArch;
+		private readonly UnrealTargetPlatform TargetPlatform;
+		private readonly DirectoryReference InstallDir;
+		private readonly FileReference PythonVenvExe;
+		private readonly string? PythonVenvVer;
 
-		private ILogger Logger;
-		private IBaseProgressLogFactory LoggerFactory;
+		private readonly ILogger Logger;
+		private readonly IBaseProgressLogFactory LoggerFactory;
 
 		public PipEnv(DirectoryReference InInstallDir, TargetDescriptor InTargetDesc, ILogger InLogger, bool UseProgressWriter = false)
 		{
@@ -780,7 +780,7 @@ namespace UnrealBuildTool.Modes
 	/// </summary>
 	class SimpleCmdLogCreator : IBaseProgressLogFactory
 	{
-		private ILogger Logger;
+		private readonly ILogger Logger;
 
 		public SimpleCmdLogCreator(ILogger InLogger) { Logger = InLogger; }
 		public IBaseCmdProgressLogger Create(string Message, int GuessSteps)
@@ -791,7 +791,7 @@ namespace UnrealBuildTool.Modes
 
 	class PipProgressLogCreator : IBaseProgressLogFactory
 	{
-		private ILogger Logger;
+		private readonly ILogger Logger;
 
 		public PipProgressLogCreator(ILogger InLogger) { Logger = InLogger; }
 		public IBaseCmdProgressLogger Create(string Message, int GuessSteps)
@@ -805,7 +805,7 @@ namespace UnrealBuildTool.Modes
 	/// </summary>
 	class SimpleCmdLogger : IBaseCmdProgressLogger
 	{
-		private ILogger Logger;
+		private readonly ILogger Logger;
 		public SimpleCmdLogger(ILogger InLogger)
 		{
 			Logger = InLogger;
@@ -839,7 +839,7 @@ namespace UnrealBuildTool.Modes
 	/// </summary>
 	class PythonCmdLogger : IBaseCmdProgressLogger
 	{
-		private ILogger Logger;
+		private readonly ILogger Logger;
 		public PythonCmdLogger(ILogger InLogger)
 		{
 			Logger = InLogger;
@@ -925,8 +925,8 @@ namespace UnrealBuildTool.Modes
 	/// </summary>
 	class FileCmdLogger : IBaseCmdProgressLogger
 	{
-		private ILogger Logger;
-		private TextWriter TextWriter;
+		private readonly ILogger Logger;
+		private readonly TextWriter TextWriter;
 		public FileCmdLogger(FileReference InTargetFile, ILogger InLogger)
 		{
 			Logger = InLogger;
@@ -964,7 +964,7 @@ namespace UnrealBuildTool.Modes
 	/// </summary>
 	class PipProgressLogger : IBaseCmdProgressLogger
 	{
-		private ILogger Logger;
+		private readonly ILogger Logger;
 		private ProgressWriter Writer;
 
 		private int StepsDone;
diff --git a/Engine/Source/Programs/UnrealBuildTool/Modes/ProfileUnitySizesMode.cs b/Engine/Source/Programs/UnrealBuildTool/Modes/ProfileUnitySizesMode.cs
--- a/Engine/Source/Programs/UnrealBuildTool/Modes/ProfileUnitySizesMode.cs
+++ b/Engine/Source/Programs/UnrealBuildTool/Modes/ProfileUnitySizesMode.cs
@@ -17,11 +17,13 @@ namespace UnrealBuildTool
 	[ToolMode("ProfileUnitySizes", ToolModeOptions.XmlConfig | ToolModeOptions.BuildPlatforms | ToolModeOptions.SingleInstance | ToolModeOptions.StartPrefetchingEngine | ToolModeOptions.ShowExecutionTime)]
 	class ProfileUnitySizesMode : ToolMode
 	{
+#pragma warning disable IDE0044 // Make field readonly - these private static fields are set by command-line parsing.
 		/// <summary>
 		/// Set of filters for files to include in the database. Relative to the root directory, or to the project file.
 		/// </summary>
 		[CommandLine("-Filter=")]
 		List<string> FilterRules = new List<string>();
+#pragma warning restore IDE0044
 
 		class TimingData
 		{
@@ -38,9 +40,9 @@ namespace UnrealBuildTool
 
 		class TimingLogger : ILogger
 		{
-			static Regex ExecutorTimingRegex = new Regex(@"executor.*\s(\d+\.*\d*)\sseconds");
-			static Regex NumFilesRegex = new Regex(@"\[\d+/(\d+)\]");
-			static Regex CPUTimingRegex = new Regex(@"CPU Time:\s(\d+\.*\d*)");
+			static readonly Regex ExecutorTimingRegex = new Regex(@"executor.*\s(\d+\.*\d*)\sseconds");
+			static readonly Regex NumFilesRegex = new Regex(@"\[\d+/(\d+)\]");
+			static readonly Regex CPUTimingRegex = new Regex(@"CPU Time:\s(\d+\.*\d*)");
 			private readonly ILogger Inner;
 			private readonly ILogger OldLogger;
 
diff --git a/Engine/Source/Programs/UnrealBuildTool/Modes/QueryMode.cs b/Engine/Source/Programs/UnrealBuildTool/Modes/QueryMode.cs
--- a/Engine/Source/Programs/UnrealBuildTool/Modes/QueryMode.cs
+++ b/Engine/Source/Programs/UnrealBuildTool/Modes/QueryMode.cs
@@ -120,7 +120,7 @@ namespace UnrealBuildTool
 		[CommandLine("-NoIntellisense")]
 		public bool bNoIntellisense = false;
 
-		private BuildConfiguration BuildConfiguration = new();
+		private readonly BuildConfiguration BuildConfiguration = new();
 
 		public override Task<int> ExecuteAsync(CommandLineArguments Arguments, ILogger Logger)
 		{
diff --git a/Engine/Source/Programs/UnrealBuildTool/Modes/QueryTargetsMode.cs b/Engine/Source/Programs/UnrealBuildTool/Modes/QueryTargetsMode.cs
--- a/Engine/Source/Programs/UnrealBuildTool/Modes/QueryTargetsMode.cs
+++ b/Engine/Source/Programs/UnrealBuildTool/Modes/QueryTargetsMode.cs
@@ -16,6 +16,7 @@ namespace UnrealBuildTool
 	[ToolMode("QueryTargets", ToolModeOptions.XmlConfig | ToolModeOptions.BuildPlatformsHostOnly | ToolModeOptions.SingleInstance)]
 	class QueryTargetsMode : ToolMode
 	{
+#pragma warning disable IDE0044 // Make field readonly - these private static fields are set by command-line parsing.
 		/// <summary>
 		/// Path to the project file to query
 		/// </summary>
@@ -40,6 +41,7 @@ namespace UnrealBuildTool
 		[CommandLine("-IncludeParentAssembly", Value ="true")]
 		[CommandLine("-DontIncludeParentAssembly", Value = "false")]
 		bool bIncludeParentAssembly = true;
+#pragma warning restore IDE0044
 
 		/// <summary>
 		/// Execute the mode
diff --git a/Engine/Source/Programs/UnrealBuildTool/Modes/ServerMode.cs b/Engine/Source/Programs/UnrealBuildTool/Modes/ServerMode.cs
--- a/Engine/Source/Programs/UnrealBuildTool/Modes/ServerMode.cs
+++ b/Engine/Source/Programs/UnrealBuildTool/Modes/ServerMode.cs
@@ -134,7 +134,7 @@ namespace UnrealBuildTool
 
 		private bool KeepRunning { get; set; } = true;
 
-		private BuildConfiguration BuildConfiguration = new();
+		private readonly BuildConfiguration BuildConfiguration = new();
 		private UEBuildTarget? CurrentTarget = null;
 		private TargetIntellisenseInfo? CurrentTargetIntellisenseInfo = null;
 		private GetBrowseConfigurationResponse CurrentBrowseConfiguration;
diff --git a/Engine/Source/Programs/UnrealBuildTool/Modes/TestMode.cs b/Engine/Source/Programs/UnrealBuildTool/Modes/TestMode.cs
--- a/Engine/Source/Programs/UnrealBuildTool/Modes/TestMode.cs
+++ b/Engine/Source/Programs/UnrealBuildTool/Modes/TestMode.cs
@@ -23,6 +23,7 @@ namespace UnrealBuildTool
 		private readonly string COREUOBJECT_MODULE = "CoreUObject";
 		private readonly string EDITOR_MODULE = "UnrealEd";
 
+#pragma warning disable IDE0044 // Make field readonly - these private static fields are set by command-line parsing.
 		/// <summary>
 		/// Whether we're building implicit tests (default is explicit)
 		/// </summary>
@@ -52,6 +53,7 @@ namespace UnrealBuildTool
         /// </summary>
         [CommandLine("-Project")]
         FileReference? ProjectFile = null;
+#pragma warning restore IDE0044
 
 		/// <summary>
 		/// Main entry point
diff --git a/Engine/Source/Programs/UnrealBuildTool/Modes/ValidatePlatformsMode.cs b/Engine/Source/Programs/UnrealBuildTool/Modes/ValidatePlatformsMode.cs
--- a/Engine/Source/Programs/UnrealBuildTool/Modes/ValidatePlatformsMode.cs
+++ b/Engine/Source/Programs/UnrealBuildTool/Modes/ValidatePlatformsMode.cs
@@ -14,6 +14,7 @@ namespace UnrealBuildTool
 	[ToolMode("ValidatePlatforms", ToolModeOptions.XmlConfig | ToolModeOptions.BuildPlatformsForValidation | ToolModeOptions.SingleInstance)]
 	class ValidatePlatformsMode : ToolMode
 	{
+#pragma warning disable IDE0044 // Make field readonly - these private static fields are set by command-line parsing.
 		/// <summary>
 		/// Platforms to validate
 		/// </summary>
@@ -31,6 +32,7 @@ namespace UnrealBuildTool
 		/// </summary>
 		[CommandLine("-OutputSDKs")]
 		bool bOutputSDKs = false;
+#pragma warning restore IDE0044
 
 		/// <summary>
 		/// Executes the tool with the given arguments
diff --git a/Engine/Source/Programs/UnrealBuildTool/Modes/WriteDocumentationMode.cs b/Engine/Source/Programs/UnrealBuildTool/Modes/WriteDocumentationMode.cs
--- a/Engine/Source/Programs/UnrealBuildTool/Modes/WriteDocumentationMode.cs
+++ b/Engine/Source/Programs/UnrealBuildTool/Modes/WriteDocumentationMode.cs
@@ -22,6 +22,7 @@ namespace UnrealBuildTool
 			TargetRules,
 		}
 
+#pragma warning disable IDE0044 // Make field readonly - these private static fields are set by command-line parsing.
 		/// <summary>
 		/// Type of documentation to generate
 		/// </summary>
@@ -33,6 +34,7 @@ namespace UnrealBuildTool
 		/// </summary>
 		[CommandLine(Required = true)]
 		FileReference OutputFile = null!;
+#pragma warning restore IDE0044
 
 		/// <summary>
 		/// Entry point for this command
diff --git a/Engine/Source/Programs/UnrealBuildTool/Platform/Android/AndroidAARHandler.cs b/Engine/Source/Programs/UnrealBuildTool/Platform/Android/AndroidAARHandler.cs
--- a/Engine/Source/Programs/UnrealBuildTool/Platform/Android/AndroidAARHandler.cs
+++ b/Engine/Source/Programs/UnrealBuildTool/Platform/Android/AndroidAARHandler.cs
@@ -36,9 +36,9 @@ namespace UnrealBuildTool
 			}
 		}
 
-		public List<string>? Repositories = null;
-		public List<AndroidAAREntry>? AARList = null;
-		private List<AndroidAAREntry>? JARList = null;
+		public readonly List<string> Repositories;
+		public readonly List<AndroidAAREntry> AARList;
+		private readonly List<AndroidAAREntry> JARList;
 
 		/// <summary>
 		/// Handler for AAR and JAR dependency determination and staging
diff --git a/Engine/Source/Programs/UnrealBuildTool/Platform/Android/AndroidProjectGenerator.cs b/Engine/Source/Programs/UnrealBuildTool/Platform/Android/AndroidProjectGenerator.cs
--- a/Engine/Source/Programs/UnrealBuildTool/Platform/Android/AndroidProjectGenerator.cs
+++ b/Engine/Source/Programs/UnrealBuildTool/Platform/Android/AndroidProjectGenerator.cs
@@ -24,7 +24,7 @@ namespace UnrealBuildTool
 		/// </summary>
 		private bool AGDEInstalled => AGDEInstalledTask.Result;
 
-		private Task<bool> AGDEInstalledTask;
+		private readonly Task<bool> AGDEInstalledTask;
 
 		public AndroidProjectGenerator(CommandLineArguments Arguments, ILogger Logger)
 			: base(Arguments, Logger)
diff --git a/Engine/Source/Programs/UnrealBuildTool/Platform/Android/AndroidToolChain.cs b/Engine/Source/Programs/UnrealBuildTool/Platform/Android/AndroidToolChain.cs
--- a/Engine/Source/Programs/UnrealBuildTool/Platform/Android/AndroidToolChain.cs
+++ b/Engine/Source/Programs/UnrealBuildTool/Platform/Android/AndroidToolChain.cs
@@ -350,7 +350,7 @@ namespace UnrealBuildTool
 			internal int MinPlatform = -1;
 			internal int MaxPlatform = -1;
 		}
-		static ConcurrentDictionary<string, Lazy<MinMax>> CachedMinMax = [];
+		static readonly ConcurrentDictionary<string, Lazy<MinMax>> CachedMinMax = [];
 
 		private bool ReadMinMaxPlatforms(string PlatformsFilename, out int MinPlatform, out int MaxPlatform)
 		{
@@ -1091,8 +1091,8 @@ namespace UnrealBuildTool
 			return ObjDumpPath!;
 		}
 
-		private HashSet<UnrealArch> HasHandledLaunchModule = new();
-		private HashSet<UnrealArch> HasHandledCoreModule = new();
+		private readonly HashSet<UnrealArch> HasHandledLaunchModule = new();
+		private readonly HashSet<UnrealArch> HasHandledCoreModule = new();
 
 		protected override CPPOutput CompileCPPFiles(CppCompileEnvironment CompileEnvironment, IEnumerable<FileItem> InputFiles, DirectoryReference OutputDir, string ModuleName, IActionGraphBuilder Graph)
 		{
diff --git a/Engine/Source/Programs/UnrealBuildTool/Platform/Android/UEBuildAndroid.cs b/Engine/Source/Programs/UnrealBuildTool/Platform/Android/UEBuildAndroid.cs
--- a/Engine/Source/Programs/UnrealBuildTool/Platform/Android/UEBuildAndroid.cs
+++ b/Engine/Source/Programs/UnrealBuildTool/Platform/Android/UEBuildAndroid.cs
@@ -72,7 +72,7 @@ namespace UnrealBuildTool
 		/// <summary>
 		/// The private mutable settings object
 		/// </summary>
-		private AndroidTargetRules Inner;
+		private readonly AndroidTargetRules Inner;
 
 		/// <summary>
 		/// Constructor
@@ -170,7 +170,7 @@ namespace UnrealBuildTool
 
 	class AndroidPlatform : UEBuildPlatform
 	{
-		UEBuildPlatformSDK SDK;
+		readonly UEBuildPlatformSDK SDK;
 
 		public AndroidPlatform(UnrealTargetPlatform InTargetPlatform, UEBuildPlatformSDK InSDK, ILogger InLogger)
 			: base(InTargetPlatform, InSDK, new AndroidArchitectureConfig(), InLogger)
diff --git a/Engine/Source/Programs/UnrealBuildTool/Platform/Android/UEDeployAndroid.cs b/Engine/Source/Programs/UnrealBuildTool/Platform/Android/UEDeployAndroid.cs
--- a/Engine/Source/Programs/UnrealBuildTool/Platform/Android/UEDeployAndroid.cs
+++ b/Engine/Source/Programs/UnrealBuildTool/Platform/Android/UEDeployAndroid.cs
@@ -50,7 +50,7 @@ namespace UnrealBuildTool
 		private int MinimumSDKLevelForGradle = 19;
 
 		// Reserved Java keywords not allowed in package names without modification
-		private static string[] JavaReservedKeywords = new string[] {
+		private static readonly string[] JavaReservedKeywords = new string[] {
 			"abstract", "assert", "boolean", "break", "byte", "case", "catch", "char", "class", "const", "continue", "default", "do",
 			"double", "else", "enum", "extends", "final", "finally", "float", "for", "goto", "if", "implements", "import", "instanceof",
 			"int", "interface", "long", "native", "new", "package", "private", "protected", "public", "return", "short", "static",
@@ -58,11 +58,6 @@ namespace UnrealBuildTool
 			"false", "null", "true"
 		};
 
-		/// <summary>
-		/// Internal usage for GetApiLevel
-		/// </summary>
-		private List<string>? PossibleApiLevels = null;
-
 		protected FileReference? ProjectFile;
 
 		/// <summary>
@@ -203,28 +198,6 @@ namespace UnrealBuildTool
 			}
 		}
 
-		/// <summary>
-		/// Simple function to pipe output asynchronously
-		/// </summary>
-		private void ParseApiLevel(object Sender, DataReceivedEventArgs Event)
-		{
-			// DataReceivedEventHandler is fired with a null string when the output stream is closed.  We don't want to
-			// print anything for that event.
-			if (!String.IsNullOrEmpty(Event.Data))
-			{
-				string Line = Event.Data;
-				if (Line.StartsWith("id:"))
-				{
-					// the line should look like: id: 1 or "android-19"
-					string[] Tokens = Line.Split("\"".ToCharArray());
-					if (Tokens.Length >= 2)
-					{
-						PossibleApiLevels!.Add(Tokens[1]);
-					}
-				}
-			}
-		}
-
 		private ConfigHierarchy GetConfigCacheIni(ConfigHierarchyType Type)
 		{
 			return ConfigCache.ReadHierarchy(Type, DirectoryReference.FromFile(ProjectFile), UnrealTargetPlatform.Android);
@@ -3811,7 +3784,7 @@ namespace UnrealBuildTool
 			}
 		}
 
-		private uint[] CRCTablesSB8 = {
+		private readonly uint[] CRCTablesSB8 = {
 			0x00000000, 0x77073096, 0xee0e612c, 0x990951ba, 0x076dc419, 0x706af48f, 0xe963a535, 0x9e6495a3, 0x0edb8832, 0x79dcb8a4, 0xe0d5e91e, 0x97d2d988, 0x09b64c2b, 0x7eb17cbd, 0xe7b82d07, 0x90bf1d91,
 			0x1db71064, 0x6ab020f2, 0xf3b97148, 0x84be41de, 0x1adad47d, 0x6ddde4eb, 0xf4d4b551, 0x83d385c7, 0x136c9856, 0x646ba8c0, 0xfd62f97a, 0x8a65c9ec, 0x14015c4f, 0x63066cd9, 0xfa0f3d63, 0x8d080df5,
 			0x3b6e20c8, 0x4c69105e, 0xd56041e4, 0xa2677172, 0x3c03e4d1, 0x4b04d447, 0xd20d85fd, 0xa50ab56b, 0x35b5a8fa, 0x42b2986c, 0xdbbbc9d6, 0xacbcf940, 0x32d86ce3, 0x45df5c75, 0xdcd60dcf, 0xabd13d59,
diff --git a/Engine/Source/Programs/UnrealBuildTool/Platform/Apple/AppleExports.cs b/Engine/Source/Programs/UnrealBuildTool/Platform/Apple/AppleExports.cs
--- a/Engine/Source/Programs/UnrealBuildTool/Platform/Apple/AppleExports.cs
+++ b/Engine/Source/Programs/UnrealBuildTool/Platform/Apple/AppleExports.cs
@@ -15,9 +15,9 @@ namespace UnrealBuildTool
 	/// </summary>
 	public static class AppleExports
 	{
-		private static bool bForceModernXcode = Environment.CommandLine.Contains("-modernxcode", StringComparison.OrdinalIgnoreCase);
-		private static bool bForceLegacyXcode = Environment.CommandLine.Contains("-legacyxcode", StringComparison.OrdinalIgnoreCase);
-		private static bool bNoEntitlements = Environment.CommandLine.Contains("-noEntitlements", StringComparison.OrdinalIgnoreCase);
+		private static readonly bool bForceModernXcode = Environment.CommandLine.Contains("-modernxcode", StringComparison.OrdinalIgnoreCase);
+		private static readonly bool bForceLegacyXcode = Environment.CommandLine.Contains("-legacyxcode", StringComparison.OrdinalIgnoreCase);
+		private static readonly bool bNoEntitlements = Environment.CommandLine.Contains("-noEntitlements", StringComparison.OrdinalIgnoreCase);
 
 		private static DirectoryReference? _MobileProvisionDir;
 
diff --git a/Engine/Source/Programs/UnrealBuildTool/Platform/Apple/AppleToolChain.cs b/Engine/Source/Programs/UnrealBuildTool/Platform/Apple/AppleToolChain.cs
--- a/Engine/Source/Programs/UnrealBuildTool/Platform/Apple/AppleToolChain.cs
+++ b/Engine/Source/Programs/UnrealBuildTool/Platform/Apple/AppleToolChain.cs
@@ -22,7 +22,7 @@ namespace UnrealBuildTool
 	{
 		protected class AppleToolChainInfo : ClangToolChainInfo
 		{
-			Tuple<Version, Version>[] AppleVersionToLLVMVersion;
+			readonly Tuple<Version, Version>[] AppleVersionToLLVMVersion;
 
 			public AppleToolChainInfo(UnrealTargetPlatform Platform, DirectoryReference DeveloperDir, FileReference Clang, FileReference Archiver, ILogger Logger)
 				: base(DeveloperDir, Clang, Archiver, Logger)
@@ -61,7 +61,7 @@ namespace UnrealBuildTool
 			public VersionNumber AppleClangVersion => VersionNumber.Parse(base.QueryClangVersion().ToString());
 		}
 
-		private Lazy<AppleToolChainSettings> ToolChainSettings;
+		private readonly Lazy<AppleToolChainSettings> ToolChainSettings;
 		
 		protected AppleToolChainSettings Settings => ToolChainSettings.Value;
 
diff --git a/Engine/Source/Programs/UnrealBuildTool/Platform/Apple/AppleToolChainSettings.cs b/Engine/Source/Programs/UnrealBuildTool/Platform/Apple/AppleToolChainSettings.cs
--- a/Engine/Source/Programs/UnrealBuildTool/Platform/Apple/AppleToolChainSettings.cs
+++ b/Engine/Source/Programs/UnrealBuildTool/Platform/Apple/AppleToolChainSettings.cs
@@ -30,12 +30,12 @@ namespace UnrealBuildTool
 		/// <summary>
 		/// A portion of the target "tuple"
 		/// </summary>
-		private string TargetOSName;
+		private readonly string TargetOSName;
 
 		/// <summary>
 		/// The name apple uses for SDK directories
 		/// </summary>
-		private string TargetSDKName;
+		private readonly string TargetSDKName;
 
 		/// <summary>
 		/// The version of the SDK being used to build with
@@ -67,8 +67,8 @@ namespace UnrealBuildTool
 		/// <summary>
 		/// Cache SDK dir for device and simulator (for non-Mac)
 		/// </summary>
-		DirectoryReference SDKDir;
-		DirectoryReference? SimulatorSDKDir = null;
+		readonly DirectoryReference SDKDir;
+		readonly DirectoryReference? SimulatorSDKDir = null;
 
 		/// <summary>
 		/// Dummy UUID for running iOS binary natively on Mac
diff --git a/Engine/Source/Programs/UnrealBuildTool/Platform/Apple/RemoteMac.cs b/Engine/Source/Programs/UnrealBuildTool/Platform/Apple/RemoteMac.cs
--- a/Engine/Source/Programs/UnrealBuildTool/Platform/Apple/RemoteMac.cs
+++ b/Engine/Source/Programs/UnrealBuildTool/Platform/Apple/RemoteMac.cs
@@ -34,17 +34,18 @@ namespace UnrealBuildTool
 	/// </summary>
 	class RemoteMac
 	{
+#pragma warning disable IDE0044 // Make field readonly - these private static fields are set by XML configuration.
 		/// <summary>
 		/// These two variables will be loaded from the XML config file in XmlConfigLoader.Init().
 		/// </summary>
 		[XmlConfigFile]
-		private readonly string? ServerName;
+		private string? ServerName;
 
 		/// <summary>
 		/// The remote username.
 		/// </summary>
 		[XmlConfigFile]
-		private readonly string? UserName;
+		private string? UserName;
 
 		/// <summary>
 		/// If set, instead of looking for RemoteToolChainPrivate.key in the usual places (Documents/Unreal, Engine/UnrealBuildTool/SSHKeys or Engine/Build/SSHKeys), this private key will be used.
@@ -63,6 +64,7 @@ namespace UnrealBuildTool
 		/// </summary>
 		[XmlConfigFile]
 		private string SshAuthentication = "-i '${CYGWIN_SSH_PRIVATE_KEY}'";
+#pragma warning restore IDE0044
 
 		/// <summary>
 		/// Save the specified port so that RemoteServerName is the machine address only
@@ -97,27 +99,27 @@ namespace UnrealBuildTool
 		/// <summary>
 		/// The base directory on the remote machine
 		/// </summary>
-		private string RemoteBaseDir;
+		private readonly string RemoteBaseDir;
 
 		/// <summary>
 		/// Mappings from local directories to remote directories
 		/// </summary>
-		private List<RemoteMapping> Mappings;
+		private readonly List<RemoteMapping> Mappings;
 
 		/// <summary>
 		/// Arguments that are used by every Ssh call
 		/// </summary>
-		private List<string> CommonSshArguments;
+		private readonly List<string> CommonSshArguments;
 
 		/// <summary>
 		/// Arguments that are used by every Rsync call
 		/// </summary>
-		private List<string> BasicRsyncArguments;
+		private readonly List<string> BasicRsyncArguments;
 
 		/// <summary>
 		/// Arguments that are used by directory Rsync call
 		/// </summary>
-		private List<string> CommonRsyncArguments;
+		private readonly List<string> CommonRsyncArguments;
 
 		private List<string> RsyncArgumentsForDownload(List<string> BaseArguments)
 		{
@@ -133,7 +135,7 @@ namespace UnrealBuildTool
 			return Arguments;
 		}
 
-		private string? IniBundleIdentifier = "";
+		private readonly string? IniBundleIdentifier = "";
 
 		/// <summary>
 		/// Constructor
diff --git a/Engine/Source/Programs/UnrealBuildTool/Platform/Clang/ClangSpecificFileAction.cs b/Engine/Source/Programs/UnrealBuildTool/Platform/Clang/ClangSpecificFileAction.cs
--- a/Engine/Source/Programs/UnrealBuildTool/Platform/Clang/ClangSpecificFileAction.cs
+++ b/Engine/Source/Programs/UnrealBuildTool/Platform/Clang/ClangSpecificFileAction.cs
@@ -14,11 +14,11 @@ namespace UnrealBuildTool
 	/// </summary>
 	internal class ClangSpecificFileAction : Action, ISpecificFileAction
 	{
-		DirectoryReference SourceDir;
-		DirectoryReference OutputDir;
-		IEnumerable<string> RspLines;
+		readonly DirectoryReference SourceDir;
+		readonly DirectoryReference OutputDir;
+		readonly IEnumerable<string> RspLines;
 
-		Dictionary<string, List<FileItem>> SingleFiles = new();
+		readonly Dictionary<string, List<FileItem>> SingleFiles = new();
 
 		internal ClangSpecificFileAction(DirectoryReference Source, DirectoryReference Output, Action Action, IEnumerable<string> ContentLines) : base(Action)
 		{
diff --git a/Engine/Source/Programs/UnrealBuildTool/Platform/Clang/ClangWarnings.cs b/Engine/Source/Programs/UnrealBuildTool/Platform/Clang/ClangWarnings.cs
--- a/Engine/Source/Programs/UnrealBuildTool/Platform/Clang/ClangWarnings.cs
+++ b/Engine/Source/Programs/UnrealBuildTool/Platform/Clang/ClangWarnings.cs
@@ -22,7 +22,7 @@ namespace UnrealBuildTool
 		}
 
 		// TODO: Get valid checkers by calling clang -cc1 -analyzer-checker-help
-		static Lazy<Dictionary<string, int>> CheckerAddedVersion = new Lazy<Dictionary<string, int>>(() => new()
+		static readonly Lazy<Dictionary<string, int>> CheckerAddedVersion = new Lazy<Dictionary<string, int>>(() => new()
 			{
 				{ "core.BitwiseShift", 18 },
 				{ "optin.core.EnumCastOutOfRange", 18 },
diff --git a/Engine/Source/Programs/UnrealBuildTool/Platform/IOS/IOSToolChain.cs b/Engine/Source/Programs/UnrealBuildTool/Platform/IOS/IOSToolChain.cs
--- a/Engine/Source/Programs/UnrealBuildTool/Platform/IOS/IOSToolChain.cs
+++ b/Engine/Source/Programs/UnrealBuildTool/Platform/IOS/IOSToolChain.cs
@@ -18,7 +18,8 @@ namespace UnrealBuildTool
 {
 	class IOSToolChain : AppleToolChain
 	{
-		private static List<FileItem> BundleDependencies = new List<FileItem>();
+		// TODO: Should this really be static? If yes, does it need to be thread-safe?
+		private static readonly List<FileItem> BundleDependencies = new List<FileItem>();
 
 		protected IOSProjectSettings ProjectSettings;
 
diff --git a/Engine/Source/Programs/UnrealBuildTool/Platform/IOS/MobileProvisionContents.cs b/Engine/Source/Programs/UnrealBuildTool/Platform/IOS/MobileProvisionContents.cs
--- a/Engine/Source/Programs/UnrealBuildTool/Platform/IOS/MobileProvisionContents.cs
+++ b/Engine/Source/Programs/UnrealBuildTool/Platform/IOS/MobileProvisionContents.cs
@@ -17,12 +17,12 @@ namespace UnrealBuildTool
 		/// <summary>
 		/// The contents of the provision
 		/// </summary>
-		XmlDocument Document;
+		readonly XmlDocument Document;
 
 		/// <summary>
 		/// Map of key names to XML elements holding their values
 		/// </summary>
-		Dictionary<string, XmlElement> NameToValue = new Dictionary<string, XmlElement>();
+		readonly Dictionary<string, XmlElement> NameToValue = new Dictionary<string, XmlElement>();
 
 		/// <summary>
 		/// Constructor
diff --git a/Engine/Source/Programs/UnrealBuildTool/Platform/IOS/UEBuildIOS.cs b/Engine/Source/Programs/UnrealBuildTool/Platform/IOS/UEBuildIOS.cs
--- a/Engine/Source/Programs/UnrealBuildTool/Platform/IOS/UEBuildIOS.cs
+++ b/Engine/Source/Programs/UnrealBuildTool/Platform/IOS/UEBuildIOS.cs
@@ -25,7 +25,7 @@ namespace UnrealBuildTool
 		/// </summary>
 		public static UnrealArch TVOSSimulator { get; } = FindOrAddByName("tvossimulator", bIsX64: false);
 
-		private static IReadOnlyDictionary<UnrealArch, string> AppleToolchainArchitectures = new Dictionary<UnrealArch, string>()
+		private static readonly IReadOnlyDictionary<UnrealArch, string> AppleToolchainArchitectures = new Dictionary<UnrealArch, string>()
 		{
 			{ UnrealArch.Arm64,         "arm64" },
 			{ UnrealArch.X64,           "x86_64" },
@@ -653,9 +653,9 @@ namespace UnrealBuildTool
 
 	class IOSPlatform : AppleBuildPlatform
 	{
-		List<IOSProjectSettings> CachedProjectSettings = new List<IOSProjectSettings>();
-		List<IOSProjectSettings> CachedProjectSettingsByBundle = new List<IOSProjectSettings>();
-		Dictionary<string, IOSProvisioningData> ProvisionCache = new Dictionary<string, IOSProvisioningData>();
+		readonly List<IOSProjectSettings> CachedProjectSettings = new List<IOSProjectSettings>();
+		readonly List<IOSProjectSettings> CachedProjectSettingsByBundle = new List<IOSProjectSettings>();
+		readonly Dictionary<string, IOSProvisioningData> ProvisionCache = new Dictionary<string, IOSProvisioningData>();
 
 		public IOSPlatform(UEBuildPlatformSDK InSDK, ILogger Logger)
 			: this(InSDK, UnrealTargetPlatform.IOS, Logger)
diff --git a/Engine/Source/Programs/UnrealBuildTool/Platform/IOS/UEDeployIOS.cs b/Engine/Source/Programs/UnrealBuildTool/Platform/IOS/UEDeployIOS.cs
--- a/Engine/Source/Programs/UnrealBuildTool/Platform/IOS/UEDeployIOS.cs
+++ b/Engine/Source/Programs/UnrealBuildTool/Platform/IOS/UEDeployIOS.cs
@@ -507,7 +507,7 @@ namespace UnrealBuildTool
 			File.WriteAllText(PlistFile.FullName, Text.ToString());
 		}
 
-		private static object GenerateIOSPListLock = new();
+		private static readonly object GenerateIOSPListLock = new();
 
 		public static bool GenerateIOSPList(FileReference? ProjectFile, UnrealTargetConfiguration Config, string ProjectDirectory, bool bIsUnrealGame, string GameName, bool bIsClient, string ProjectName, string InEngineDir, string AppDirectory, UnrealPluginLanguage? UPL, string? BundleID, bool bBuildAsFramework, ILogger Logger)
 		{
diff --git a/Engine/Source/Programs/UnrealBuildTool/Platform/Linux/UEBuildLinux.cs b/Engine/Source/Programs/UnrealBuildTool/Platform/Linux/UEBuildLinux.cs
--- a/Engine/Source/Programs/UnrealBuildTool/Platform/Linux/UEBuildLinux.cs
+++ b/Engine/Source/Programs/UnrealBuildTool/Platform/Linux/UEBuildLinux.cs
@@ -11,7 +11,7 @@ namespace UnrealBuildTool
 {
 	partial struct UnrealArch
 	{
-		private static IReadOnlyDictionary<UnrealArch, string> LinuxToolchainArchitectures = new Dictionary<UnrealArch, string>()
+		private static readonly IReadOnlyDictionary<UnrealArch, string> LinuxToolchainArchitectures = new Dictionary<UnrealArch, string>()
 		{
 			{ UnrealArch.Arm64,         "aarch64-unknown-linux-gnueabi" },
 			{ UnrealArch.X64,           "x86_64-unknown-linux-gnu" },
@@ -142,7 +142,7 @@ namespace UnrealBuildTool
 	/// </summary>
 	public class ReadOnlyLinuxTargetRules
 	{
-		private LinuxTargetRules Inner;
+		private readonly LinuxTargetRules Inner;
 
 		/// <summary>
 		/// Constructor
diff --git a/Engine/Source/Programs/UnrealBuildTool/Platform/Mac/MacToolChain.cs b/Engine/Source/Programs/UnrealBuildTool/Platform/Mac/MacToolChain.cs
--- a/Engine/Source/Programs/UnrealBuildTool/Platform/Mac/MacToolChain.cs
+++ b/Engine/Source/Programs/UnrealBuildTool/Platform/Mac/MacToolChain.cs
@@ -1106,7 +1106,7 @@ namespace UnrealBuildTool
 			return TargetItem;
 		}
 
-		private static Dictionary<ReadOnlyTargetRules, DirectoryReference> BundleContentsDirectories = new();
+		private static readonly Dictionary<ReadOnlyTargetRules, DirectoryReference> BundleContentsDirectories = new();
 
 		public override void ModifyBuildProducts(ReadOnlyTargetRules Target, UEBuildBinary Binary, IEnumerable<string> Libraries, IEnumerable<UEBuildBundleResource> BundleResources, Dictionary<FileReference, BuildProductType> BuildProducts)
 		{
@@ -1227,7 +1227,7 @@ namespace UnrealBuildTool
 			}
 		}
 
-		private List<FileItem> DebugInfoFiles = new List<FileItem>();
+		private readonly List<FileItem> DebugInfoFiles = new List<FileItem>();
 
 		public override void FinalizeOutput(ReadOnlyTargetRules Target, TargetMakefileBuilder MakefileBuilder)
 		{
diff --git a/Engine/Source/Programs/UnrealBuildTool/Platform/Mac/UEBuildMac.cs b/Engine/Source/Programs/UnrealBuildTool/Platform/Mac/UEBuildMac.cs
--- a/Engine/Source/Programs/UnrealBuildTool/Platform/Mac/UEBuildMac.cs
+++ b/Engine/Source/Programs/UnrealBuildTool/Platform/Mac/UEBuildMac.cs
@@ -132,7 +132,7 @@ namespace UnrealBuildTool
 			return GetProjectArchitectures(ProjectFile, TargetName, true, false);
 		}
 
-		private static Dictionary<string, UnrealArchitectures> ProjectArchitectureCache = new();
+		private static readonly Dictionary<string, UnrealArchitectures> ProjectArchitectureCache = new();
 		private UnrealArchitectures GetProjectArchitectures(FileReference? ProjectFile, string? TargetName, bool bGetAllSupported, bool bIsDistributionMode)
 		{
 			string Key = $"{ProjectFile}{TargetName}{bGetAllSupported}{bIsDistributionMode}";
diff --git a/Engine/Source/Programs/UnrealBuildTool/Platform/Windows/MicrosoftPlatformSDK.Versions.cs b/Engine/Source/Programs/UnrealBuildTool/Platform/Windows/MicrosoftPlatformSDK.Versions.cs
--- a/Engine/Source/Programs/UnrealBuildTool/Platform/Windows/MicrosoftPlatformSDK.Versions.cs
+++ b/Engine/Source/Programs/UnrealBuildTool/Platform/Windows/MicrosoftPlatformSDK.Versions.cs
@@ -18,7 +18,7 @@ namespace UnrealBuildTool
 		/// <summary>
 		/// The default compiler version to be used, if installed. 
 		/// </summary>
-		static VersionNumberRange[] PreferredClangVersions = SDK.GetVersionNumberRangeArrayFromConfig("PreferredClangVersions");
+		static readonly VersionNumberRange[] PreferredClangVersions = SDK.GetVersionNumberRangeArrayFromConfig("PreferredClangVersions");
 
 		/// <summary>
 		/// The minimum supported Clang compiler
diff --git a/Engine/Source/Programs/UnrealBuildTool/Platform/Windows/MicrosoftPlatformSDK.cs b/Engine/Source/Programs/UnrealBuildTool/Platform/Windows/MicrosoftPlatformSDK.cs
--- a/Engine/Source/Programs/UnrealBuildTool/Platform/Windows/MicrosoftPlatformSDK.cs
+++ b/Engine/Source/Programs/UnrealBuildTool/Platform/Windows/MicrosoftPlatformSDK.cs
@@ -90,7 +90,7 @@ namespace UnrealBuildTool
 		/// <summary>
 		/// Cache of Visual C++ installation directories
 		/// </summary>
-		private static ConcurrentDictionary<WindowsCompiler, List<ToolChainInstallation>> CachedToolChainInstallations = new ConcurrentDictionary<WindowsCompiler, List<ToolChainInstallation>>();
+		private static readonly ConcurrentDictionary<WindowsCompiler, List<ToolChainInstallation>> CachedToolChainInstallations = new ConcurrentDictionary<WindowsCompiler, List<ToolChainInstallation>>();
 
 		/// <summary>
 		/// Cache of Windows SDK installation directories
@@ -107,12 +107,12 @@ namespace UnrealBuildTool
 		/// </summary>
 		private static IReadOnlyList<VisualStudioInstallation>? CachedVisualStudioInstallations;
 
-		private static object CachedVisualStudioInstallationsLock = new();
+		private static readonly object CachedVisualStudioInstallationsLock = new();
 
 		/// <summary>
 		/// Cache of DIA SDK installation directories
 		/// </summary>
-		private static ConcurrentDictionary<WindowsCompiler, List<DirectoryReference>> CachedDiaSdkDirs = new ConcurrentDictionary<WindowsCompiler, List<DirectoryReference>>();
+		private static readonly ConcurrentDictionary<WindowsCompiler, List<DirectoryReference>> CachedDiaSdkDirs = new ConcurrentDictionary<WindowsCompiler, List<DirectoryReference>>();
 
 		#endregion
 
diff --git a/Engine/Source/Programs/UnrealBuildTool/Platform/Windows/PVSToolChain.cs b/Engine/Source/Programs/UnrealBuildTool/Platform/Windows/PVSToolChain.cs
--- a/Engine/Source/Programs/UnrealBuildTool/Platform/Windows/PVSToolChain.cs
+++ b/Engine/Source/Programs/UnrealBuildTool/Platform/Windows/PVSToolChain.cs
@@ -371,6 +371,7 @@ namespace UnrealBuildTool
 	[ToolMode("PVSGather", ToolModeOptions.None)]
 	class PVSGatherMode : ToolMode
 	{
+#pragma warning disable IDE0044 // Make field readonly - these private static fields are set by command-line parsing.
 		/// <summary>
 		/// Path to the input file list
 		/// </summary>
@@ -400,11 +401,12 @@ namespace UnrealBuildTool
 		/// </summary>
 		[CommandLine("-PrintLevel")]
 		int _printLevel = 1;
+#pragma warning restore IDE0044
 
 		/// <summary>
 		/// If all ThirdParty code should be ignored
 		/// </summary>
-		bool _ignoreThirdParty = true;
+		const bool _ignoreThirdParty = true;
 
 		readonly CaptureLogger _parseLogger = new();
 
@@ -556,7 +558,7 @@ namespace UnrealBuildTool
 		static readonly Version _minAnalyzerVersion = new Version("7.30");
 		static readonly Version _analysisPathsSkipVersion = new Version("7.34");
 
-		static string OutputFileExtension = ".PVS-Studio.log";
+		const string OutputFileExtension = ".PVS-Studio.log";
 
 		public PVSToolChain(ReadOnlyTargetRules Target, VCToolChain InInnerToolchain, ILogger Logger)
 			: base(Logger)
diff --git a/Engine/Source/Programs/UnrealBuildTool/Platform/Windows/ResXResource/UEResXWriter.cs b/Engine/Source/Programs/UnrealBuildTool/Platform/Windows/ResXResource/UEResXWriter.cs
--- a/Engine/Source/Programs/UnrealBuildTool/Platform/Windows/ResXResource/UEResXWriter.cs
+++ b/Engine/Source/Programs/UnrealBuildTool/Platform/Windows/ResXResource/UEResXWriter.cs
@@ -68,8 +68,8 @@ namespace UnrealBuildTool
 			return Data;
 		}
 
-		private string Filename;
-		private XmlDocument Document;
-		private XmlElement RootElement;
+		private readonly string Filename;
+		private readonly XmlDocument Document;
+		private readonly XmlElement RootElement;
 	}
 }
diff --git a/Engine/Source/Programs/UnrealBuildTool/Platform/Windows/UEBuildWindows.cs b/Engine/Source/Programs/UnrealBuildTool/Platform/Windows/UEBuildWindows.cs
--- a/Engine/Source/Programs/UnrealBuildTool/Platform/Windows/UEBuildWindows.cs
+++ b/Engine/Source/Programs/UnrealBuildTool/Platform/Windows/UEBuildWindows.cs
@@ -35,7 +35,7 @@ namespace UnrealBuildTool
 			}
 		}
 
-		private static IReadOnlyDictionary<UnrealArch, WindowsArchInfo> WindowsToolchainArchitectures = new Dictionary<UnrealArch, WindowsArchInfo>()
+		private static readonly IReadOnlyDictionary<UnrealArch, WindowsArchInfo> WindowsToolchainArchitectures = new Dictionary<UnrealArch, WindowsArchInfo>()
 		{
 			{ UnrealArch.Arm64,         new WindowsArchInfo("arm64", "arm64", "arm64") },
 			{ UnrealArch.Arm64ec,       new WindowsArchInfo("arm64", "x64", "arm64") },
@@ -236,7 +236,7 @@ namespace UnrealBuildTool
 		/// <summary>
 		/// The target rules which owns this object. Used to resolve some properties.
 		/// </summary>
-		TargetRules Target;
+		readonly TargetRules Target;
 
 		/// <summary>
 		/// If enabled will set the ProductVersion embeded in windows executables and dlls to contain BUILT_FROM_CHANGELIST and BuildVersion
@@ -348,11 +348,13 @@ namespace UnrealBuildTool
 		}
 		private string[]? _preferredCompilersPrivate = null;
 
+#pragma warning disable IDE0044 // Make field readonly - these private static fields are set by XML configuration.
 		/// <summary>
 		/// Default order of preference when Compiler is set to Default. If populated and none listed are available this will cause a warning.
 		/// </summary>
 		[XmlConfigFile(Category = "WindowsPlatform")]
 		private string[] DefaultPreferredCompilers = [];
+#pragma warning restore IDE0044
 
 		/// <summary>
 		/// Architecture of Target.
@@ -951,7 +953,7 @@ namespace UnrealBuildTool
 		/// <summary>
 		/// The private mutable settings object
 		/// </summary>
-		private WindowsTargetRules Inner;
+		private readonly WindowsTargetRules Inner;
 
 		/// <summary>
 		/// Constructor
@@ -1222,7 +1224,7 @@ namespace UnrealBuildTool
 
 	class WindowsPlatform : UEBuildPlatform
 	{
-		MicrosoftPlatformSDK SDK;
+		readonly MicrosoftPlatformSDK SDK;
 
 		/// <summary>
 		/// Constructor
diff --git a/Engine/Source/Programs/UnrealBuildTool/Platform/Windows/VCPreludeCache.cs b/Engine/Source/Programs/UnrealBuildTool/Platform/Windows/VCPreludeCache.cs
--- a/Engine/Source/Programs/UnrealBuildTool/Platform/Windows/VCPreludeCache.cs
+++ b/Engine/Source/Programs/UnrealBuildTool/Platform/Windows/VCPreludeCache.cs
@@ -18,7 +18,7 @@ namespace UnrealBuildTool
 		/// <summary>
 		/// Set of options which can influence the prelude state. Any options matching entries in this set will be used to generate a separate prelude file.
 		/// </summary>
-		static HashSet<string> FilterArguments = new HashSet<string>(StringComparer.OrdinalIgnoreCase)
+		static readonly HashSet<string> FilterArguments = new HashSet<string>(StringComparer.OrdinalIgnoreCase)
 		{
 			"/EH",  // Exception handling mode
 			"/GR",  // RTTI settings
@@ -31,7 +31,7 @@ namespace UnrealBuildTool
 		/// <summary>
 		/// Object used for ensuring only one thread can create items in the cache at once
 		/// </summary>
-		static object LockObject = new object();
+		static readonly object LockObject = new object();
 
 		/// <summary>
 		/// File containing a list of macros to be used to generate a prelude file
@@ -46,7 +46,7 @@ namespace UnrealBuildTool
 		/// <summary>
 		/// The parsed list of macros
 		/// </summary>
-		static List<string> MacroNames = new List<string>();
+		static readonly List<string> MacroNames = new List<string>();
 
 		/// <summary>
 		/// Path to the file used to generate prelude headers
@@ -56,7 +56,7 @@ namespace UnrealBuildTool
 		/// <summary>
 		/// Cached mapping from compiler filename to its version string
 		/// </summary>
-		static Dictionary<FileReference, string> CompilerExeToVersion = new Dictionary<FileReference, string>();
+		static readonly Dictionary<FileReference, string> CompilerExeToVersion = new Dictionary<FileReference, string>();
 
 		/// <summary>
 		/// Get the prelude file (or create it) for the given compiler executable and set of options
diff --git a/Engine/Source/Programs/UnrealBuildTool/Platform/Windows/VCSpecificFileAction.cs b/Engine/Source/Programs/UnrealBuildTool/Platform/Windows/VCSpecificFileAction.cs
--- a/Engine/Source/Programs/UnrealBuildTool/Platform/Windows/VCSpecificFileAction.cs
+++ b/Engine/Source/Programs/UnrealBuildTool/Platform/Windows/VCSpecificFileAction.cs
@@ -10,15 +10,15 @@ namespace UnrealBuildTool
 {
 	internal class VcSpecificFileAction : Action, ISpecificFileAction
 	{
-		DirectoryReference SourceDir;
-		DirectoryReference OutputDir;
-		VCCompileAction BaseAction;
+		readonly DirectoryReference SourceDir;
+		readonly DirectoryReference OutputDir;
+		readonly VCCompileAction BaseAction;
 
 		// TODO: Entire CppCompileEnvironment needs to be saved with BinaryArchiveWriter, some options may still be unused
-		bool bPreprocessOnly;
-		bool bWithAssembly;
+		readonly bool bPreprocessOnly;
+		readonly bool bWithAssembly;
 
-		Dictionary<string, List<FileItem>> SingleFiles = new();
+		readonly Dictionary<string, List<FileItem>> SingleFiles = new();
 
 		internal VcSpecificFileAction(DirectoryReference Source, DirectoryReference Output, VCCompileAction Action, CppCompileEnvironment? CompileEnvironment = null) : base(ActionType.Compile)
 		{
diff --git a/Engine/Source/Programs/UnrealBuildTool/ProjectFiles/CodeLite/CodeLiteGenerator.cs b/Engine/Source/Programs/UnrealBuildTool/ProjectFiles/CodeLite/CodeLiteGenerator.cs
--- a/Engine/Source/Programs/UnrealBuildTool/ProjectFiles/CodeLite/CodeLiteGenerator.cs
+++ b/Engine/Source/Programs/UnrealBuildTool/ProjectFiles/CodeLite/CodeLiteGenerator.cs
@@ -22,7 +22,7 @@ namespace UnrealBuildTool
 		public string SolutionExtension = ".workspace";
 		public string CodeCompletionFileName = "CodeCompletionFolders.txt";
 		public string CodeCompletionPreProcessorFileName = "CodeLitePreProcessor.txt";
-		CodeliteProjectFileFormat ProjectFileFormat = CodeliteProjectFileFormat.CodeLite10;
+		readonly CodeliteProjectFileFormat ProjectFileFormat = CodeliteProjectFileFormat.CodeLite10;
 
 		public CodeLiteGenerator(FileReference? InOnlyGameProject, CommandLineArguments CommandLine)
 			: base(InOnlyGameProject)
diff --git a/Engine/Source/Programs/UnrealBuildTool/ProjectFiles/CodeLite/CodeLiteProject.cs b/Engine/Source/Programs/UnrealBuildTool/ProjectFiles/CodeLite/CodeLiteProject.cs
--- a/Engine/Source/Programs/UnrealBuildTool/ProjectFiles/CodeLite/CodeLiteProject.cs
+++ b/Engine/Source/Programs/UnrealBuildTool/ProjectFiles/CodeLite/CodeLiteProject.cs
@@ -12,7 +12,7 @@ namespace UnrealBuildTool
 {
 	class CodeLiteProject : ProjectFile
 	{
-		FileReference? OnlyGameProject;
+		readonly FileReference? OnlyGameProject;
 
 		public CodeLiteProject(FileReference InitFilePath, DirectoryReference BaseDir, FileReference? InOnlyGameProject) : base(InitFilePath, BaseDir)
 		{
diff --git a/Engine/Source/Programs/UnrealBuildTool/ProjectFiles/Make/MakefileGenerator.cs b/Engine/Source/Programs/UnrealBuildTool/ProjectFiles/Make/MakefileGenerator.cs
--- a/Engine/Source/Programs/UnrealBuildTool/ProjectFiles/Make/MakefileGenerator.cs
+++ b/Engine/Source/Programs/UnrealBuildTool/ProjectFiles/Make/MakefileGenerator.cs
@@ -25,7 +25,7 @@ namespace UnrealBuildTool
 	{
 		/// True if intellisense data should be generated (takes a while longer)
 		/// Now this is needed for project target generation.
-		bool bGenerateIntelliSenseData = true;
+		const bool bGenerateIntelliSenseData = true;
 
 		/// Default constructor
 		public MakefileGenerator(FileReference? InOnlyGameProject)
diff --git a/Engine/Source/Programs/UnrealBuildTool/ProjectFiles/PlatformProjectGenerator.cs b/Engine/Source/Programs/UnrealBuildTool/ProjectFiles/PlatformProjectGenerator.cs
--- a/Engine/Source/Programs/UnrealBuildTool/ProjectFiles/PlatformProjectGenerator.cs
+++ b/Engine/Source/Programs/UnrealBuildTool/ProjectFiles/PlatformProjectGenerator.cs
@@ -48,7 +48,7 @@ namespace UnrealBuildTool
 			return false;
 		}
 
-		ConcurrentDictionary<int, bool> FoundVSISupportDict = new();
+		readonly ConcurrentDictionary<int, bool> FoundVSISupportDict = new();
 
 		public struct VSSettings
 		{
@@ -309,8 +309,8 @@ namespace UnrealBuildTool
 				}
 			}
 
-			private HashSet<string> PropertiesToPatchContainer = new();
-			private HashSet<string> PropertiesToPatchOrderButPreserveValueContainer = new();
+			private readonly HashSet<string> PropertiesToPatchContainer = new();
+			private readonly HashSet<string> PropertiesToPatchOrderButPreserveValueContainer = new();
 		}
 
 		/// <summary>
diff --git a/Engine/Source/Programs/UnrealBuildTool/ProjectFiles/PlatformProjectGeneratorCollection.cs b/Engine/Source/Programs/UnrealBuildTool/ProjectFiles/PlatformProjectGeneratorCollection.cs
--- a/Engine/Source/Programs/UnrealBuildTool/ProjectFiles/PlatformProjectGeneratorCollection.cs
+++ b/Engine/Source/Programs/UnrealBuildTool/ProjectFiles/PlatformProjectGeneratorCollection.cs
@@ -13,7 +13,7 @@ namespace UnrealBuildTool
 	/// </summary>
 	class PlatformProjectGeneratorCollection
 	{
-		Dictionary<UnrealTargetPlatform, PlatformProjectGenerator> ProjectGeneratorDictionary = new Dictionary<UnrealTargetPlatform, PlatformProjectGenerator>();
+		readonly Dictionary<UnrealTargetPlatform, PlatformProjectGenerator> ProjectGeneratorDictionary = new Dictionary<UnrealTargetPlatform, PlatformProjectGenerator>();
 
 		/// <summary>
 		/// Returns the list of platforms that have been registered in this collection
diff --git a/Engine/Source/Programs/UnrealBuildTool/ProjectFiles/ProjectFile.cs b/Engine/Source/Programs/UnrealBuildTool/ProjectFiles/ProjectFile.cs
--- a/Engine/Source/Programs/UnrealBuildTool/ProjectFiles/ProjectFile.cs
+++ b/Engine/Source/Programs/UnrealBuildTool/ProjectFiles/ProjectFile.cs
@@ -182,7 +182,7 @@ namespace UnrealBuildTool
 		/// <summary>
 		/// Set of all files that have been added
 		/// </summary>
-		Dictionary<FileReference, AliasedFile> AliasedFilesLookup = [];
+		readonly Dictionary<FileReference, AliasedFile> AliasedFilesLookup = [];
 
 		public Dictionary<FileReference, AliasedFile>.ValueCollection AliasedFiles => AliasedFilesLookup.Values;
 
@@ -534,12 +534,12 @@ namespace UnrealBuildTool
 		/// <summary>
 		/// Merged list of include paths for the project
 		/// </summary>
-		IncludePathsCollection UserIncludePaths = new IncludePathsCollection();
+		readonly IncludePathsCollection UserIncludePaths = new IncludePathsCollection();
 
 		/// <summary>
 		/// Merged list of include paths for the project
 		/// </summary>
-		IncludePathsCollection SystemIncludePaths = new IncludePathsCollection();
+		readonly IncludePathsCollection SystemIncludePaths = new IncludePathsCollection();
 
 		/// <summary>
 		/// Map of base directory to user include paths
@@ -564,7 +564,7 @@ namespace UnrealBuildTool
 		/// <summary>
 		/// Set of unique preprocessor definitions
 		/// </summary>
-		HashSet<string> KnownIntelliSensePreprocessorDefinitions = new HashSet<string>(StringComparer.InvariantCultureIgnoreCase);
+		readonly HashSet<string> KnownIntelliSensePreprocessorDefinitions = new HashSet<string>(StringComparer.InvariantCultureIgnoreCase);
 
 		/// <summary>
 		/// Projects that this project is dependent on
diff --git a/Engine/Source/Programs/UnrealBuildTool/ProjectFiles/ProjectFileGenerator.cs b/Engine/Source/Programs/UnrealBuildTool/ProjectFiles/ProjectFileGenerator.cs
--- a/Engine/Source/Programs/UnrealBuildTool/ProjectFiles/ProjectFileGenerator.cs
+++ b/Engine/Source/Programs/UnrealBuildTool/ProjectFiles/ProjectFileGenerator.cs
@@ -238,6 +238,7 @@ namespace UnrealBuildTool
 		/// </summary>
 		protected string? GameProjectName = null;
 
+#pragma warning disable IDE0044 // Make field readonly - these private static fields are set by XML configuration.
 		/// <summary>
 		/// Whether we should include configurations for "Test" and "Shipping" in generated projects. Pass "-NoShippingConfigs" to disable this.
 		/// </summary>
@@ -365,6 +366,7 @@ namespace UnrealBuildTool
 		/// </summary>
 		[XmlConfigFile(Name = "Configurations")]
 		string[]? ConfigurationNames = null;
+#pragma warning restore IDE0044
 
 		/// <summary>
 		/// If set, strip out other found Target.cs files
@@ -1948,7 +1950,7 @@ namespace UnrealBuildTool
 		/// <summary>
 		/// List of non-redistributable folders
 		/// </summary>
-		private static string[] NoRedistFolders = new string[]
+		private static readonly string[] NoRedistFolders = new string[]
 		{
 			Path.DirectorySeparatorChar + "NoRedist" + Path.DirectorySeparatorChar,
 			Path.DirectorySeparatorChar + "NotForLicensees" + Path.DirectorySeparatorChar,
diff --git a/Engine/Source/Programs/UnrealBuildTool/ProjectFiles/Rider/RiderProjectFileGenerator.cs b/Engine/Source/Programs/UnrealBuildTool/ProjectFiles/Rider/RiderProjectFileGenerator.cs
--- a/Engine/Source/Programs/UnrealBuildTool/ProjectFiles/Rider/RiderProjectFileGenerator.cs
+++ b/Engine/Source/Programs/UnrealBuildTool/ProjectFiles/Rider/RiderProjectFileGenerator.cs
@@ -25,6 +25,7 @@ namespace UnrealBuildTool
 		/// <returns></returns>
 		private readonly HashSet<UnrealTargetPlatform> DeprecatedPlatforms = new HashSet<UnrealTargetPlatform>();
 
+#pragma warning disable IDE0044 // Make field readonly - these private static fields are set by command-line parsing.
 		/// <summary>
 		/// When specified only primary project file (root.json) will be generated.
 		/// Normally project files for specific configurations are generated together with primary project file.
@@ -61,6 +62,7 @@ namespace UnrealBuildTool
 		/// </summary>
 		[CommandLine("-Minimize", Value = "Compact")]
 		private JsonWriterStyle Minimize = JsonWriterStyle.Readable;
+#pragma warning restore IDE0044
 
 		public RiderProjectFileGenerator(FileReference? InOnlyGameProject,
 			CommandLineArguments InArguments)
diff --git a/Engine/Source/Programs/UnrealBuildTool/ProjectFiles/VisualStudio/VCProject.cs b/Engine/Source/Programs/UnrealBuildTool/ProjectFiles/VisualStudio/VCProject.cs
--- a/Engine/Source/Programs/UnrealBuildTool/ProjectFiles/VisualStudio/VCProject.cs
+++ b/Engine/Source/Programs/UnrealBuildTool/ProjectFiles/VisualStudio/VCProject.cs
@@ -285,18 +285,18 @@ namespace UnrealBuildTool
 		/// </summary>
 		public Guid ProjectGUID => _projectGuid.Value;
 
-		Lazy<Guid> _projectGuid;
+		readonly Lazy<Guid> _projectGuid;
 	}
 
 	class VCProjectFile : MSBuildProjectFile
 	{
-		VCProjectFileFormat ProjectFileFormat;
-		bool bUsePrecompiled;
-		bool bMakeProjectPerTarget;
-		string? BuildToolOverride;
-		Dictionary<DirectoryReference, string> ModuleDirToForceIncludePaths = new Dictionary<DirectoryReference, string>();
-		Dictionary<DirectoryReference, string> ModuleDirToPchHeaderFile = new Dictionary<DirectoryReference, string>();
-		VCProjectFileSettings Settings;
+		readonly VCProjectFileFormat ProjectFileFormat;
+		readonly bool bUsePrecompiled;
+		readonly bool bMakeProjectPerTarget;
+		readonly string? BuildToolOverride;
+		readonly Dictionary<DirectoryReference, string> ModuleDirToForceIncludePaths = new Dictionary<DirectoryReference, string>();
+		readonly Dictionary<DirectoryReference, string> ModuleDirToPchHeaderFile = new Dictionary<DirectoryReference, string>();
+		readonly VCProjectFileSettings Settings;
 
 		/// This is the platform name that Visual Studio is always guaranteed to support.  We'll use this as
 		/// a platform for any project configurations where our actual platform is not supported by the
@@ -2006,8 +2006,8 @@ namespace UnrealBuildTool
 			public override string ConfigurationName => Combination.ProjectConfigurationName;
 			public override string BuildCommand => $"{EscapePath(NormalizeProjectPath(CommandBuilder.BuildScript))} {CommandBuilder.GetBuildArguments()}";
 
-			private ProjectConfigAndTargetCombination Combination;
-			private BuildCommandBuilder CommandBuilder;
+			private readonly ProjectConfigAndTargetCombination Combination;
+			private readonly BuildCommandBuilder CommandBuilder;
 
 			public ProjectConfigurationForGenerator(ProjectConfigAndTargetCombination InCombination, BuildCommandBuilder InCommandBuilder)
 			{
diff --git a/Engine/Source/Programs/UnrealBuildTool/ProjectFiles/VisualStudio/VCProjectFileGenerator.cs b/Engine/Source/Programs/UnrealBuildTool/ProjectFiles/VisualStudio/VCProjectFileGenerator.cs
--- a/Engine/Source/Programs/UnrealBuildTool/ProjectFiles/VisualStudio/VCProjectFileGenerator.cs
+++ b/Engine/Source/Programs/UnrealBuildTool/ProjectFiles/VisualStudio/VCProjectFileGenerator.cs
@@ -125,7 +125,7 @@ namespace UnrealBuildTool
 		/// Override for the build tool to use in generated projects. If the compiler version is specified on the command line, we use the same argument on the 
 		/// command line for generated projects.
 		/// </summary>
-		string? BuildToolOverride;
+		readonly string? BuildToolOverride;
 
 		/// <summary>
 		/// Default constructor
diff --git a/Engine/Source/Programs/UnrealBuildTool/ProjectFiles/VisualStudio/VCSolutionOptions.cs b/Engine/Source/Programs/UnrealBuildTool/ProjectFiles/VisualStudio/VCSolutionOptions.cs
--- a/Engine/Source/Programs/UnrealBuildTool/ProjectFiles/VisualStudio/VCSolutionOptions.cs
+++ b/Engine/Source/Programs/UnrealBuildTool/ProjectFiles/VisualStudio/VCSolutionOptions.cs
@@ -538,7 +538,7 @@ namespace UnrealBuildTool
 	[SupportedOSPlatform("windows")]
 	class VCSolutionOptions : VCOleContainer
 	{
-		VCProjectFileFormat Format;
+		readonly VCProjectFileFormat Format;
 
 		public VCSolutionOptions(VCProjectFileFormat Format)
 		{
diff --git a/Engine/Source/Programs/UnrealBuildTool/ProjectFiles/VisualStudioCode/VSCodeProjectFileGenerator.cs b/Engine/Source/Programs/UnrealBuildTool/ProjectFiles/VisualStudioCode/VSCodeProjectFileGenerator.cs
--- a/Engine/Source/Programs/UnrealBuildTool/ProjectFiles/VisualStudioCode/VSCodeProjectFileGenerator.cs
+++ b/Engine/Source/Programs/UnrealBuildTool/ProjectFiles/VisualStudioCode/VSCodeProjectFileGenerator.cs
@@ -30,12 +30,13 @@ namespace UnrealBuildTool
 	{
 		private UnrealTargetPlatform HostPlatform = BuildHostPlatform.Current.Platform;
 		private bool bForeignProject;
-		private DirectoryReference ProjectRoot;
-		private string FrameworkExecutableExtension = OperatingSystem.IsWindows() ? ".exe" : "";
-		private string FrameworkLibraryExtension = ".dll";
+		private readonly DirectoryReference ProjectRoot;
+		private readonly string FrameworkExecutableExtension = OperatingSystem.IsWindows() ? ".exe" : "";
+		private readonly string FrameworkLibraryExtension = ".dll";
 
 		private readonly ConcurrentBag<BuildTarget> BuildTargets = new();
 
+#pragma warning disable IDE0044 // Make field readonly - these private static fields are set by XML configuration.
 		/// <summary>
 		/// Includes all files in the generated workspace.
 		/// </summary>
@@ -74,6 +75,7 @@ namespace UnrealBuildTool
 		/// </summary>
 		[XmlConfigFile(Name = "LinuxDebuggerType")]
 		private string LinuxDebuggerType = "cppdbg";
+#pragma warning restore IDE0044
 
 		private enum EPathType
 		{
@@ -168,6 +170,7 @@ namespace UnrealBuildTool
 			ProjectRoot = Unreal.RootDirectory;
 		}
 
+		// TODO: Do we really need this instead of just using System.Text.Json?
 		class JsonFile
 		{
 			public JsonFile()
@@ -242,7 +245,7 @@ namespace UnrealBuildTool
 				return "\"" + Value + "\"";
 			}
 
-			private List<string> Lines = new List<string>();
+			private readonly List<string> Lines = new List<string>();
 			private string TabString = "";
 		}
 
diff --git a/Engine/Source/Programs/UnrealBuildTool/ProjectFiles/VisualStudioWorkspace/VSWorkspaceProjectFileGenerator.cs b/Engine/Source/Programs/UnrealBuildTool/ProjectFiles/VisualStudioWorkspace/VSWorkspaceProjectFileGenerator.cs
--- a/Engine/Source/Programs/UnrealBuildTool/ProjectFiles/VisualStudioWorkspace/VSWorkspaceProjectFileGenerator.cs
+++ b/Engine/Source/Programs/UnrealBuildTool/ProjectFiles/VisualStudioWorkspace/VSWorkspaceProjectFileGenerator.cs
@@ -33,6 +33,7 @@ namespace UnrealBuildTool
 		/// <returns></returns>
 		private readonly HashSet<UnrealTargetPlatform> DeprecatedPlatforms = new();
 
+#pragma warning disable IDE0044 // Make field readonly - these private static fields are set by command-line parsing.
 		/// <summary>
 		/// Platforms to generate project files for
 		/// </summary>
@@ -74,6 +75,7 @@ namespace UnrealBuildTool
 
 		[CommandLine("-GenerateVsConfigOnly", Description = "Generate .vsconfig file with dependencies and exit")]
 		private bool GenerateVsConfigOnly = false;
+#pragma warning restore IDE0044
 
 		public VSWorkspaceProjectFileGenerator(FileReference? InOnlyGameProject,
 			CommandLineArguments InArguments)
diff --git a/Engine/Source/Programs/UnrealBuildTool/ProjectFiles/Xcode/XcodeProject.File.cs b/Engine/Source/Programs/UnrealBuildTool/ProjectFiles/Xcode/XcodeProject.File.cs
--- a/Engine/Source/Programs/UnrealBuildTool/ProjectFiles/Xcode/XcodeProject.File.cs
+++ b/Engine/Source/Programs/UnrealBuildTool/ProjectFiles/Xcode/XcodeProject.File.cs
@@ -183,7 +183,7 @@ namespace UnrealBuildTool.XcodeProjectXcconfig
 		}
 
 		// we need relative paths, so this is the directory that paths are relative to
-		private DirectoryReference ProjectDirectory;
+		private readonly DirectoryReference ProjectDirectory;
 
 		// the location of the uproject file, or the engine directory for non-uproject projects
 		public DirectoryReference RootDirectory;
@@ -196,8 +196,8 @@ namespace UnrealBuildTool.XcodeProjectXcconfig
 		internal List<Tuple<XcodeSourceFile, string>> BuildableResourceFiles = new();
 		internal List<XcodeSourceFile> AllFiles = new();
 
-		private List<ManualFileReference> ManualFiles = new();
-		private Dictionary<string, string> GuidsForGroups = new();
+		private readonly List<ManualFileReference> ManualFiles = new();
+		private readonly Dictionary<string, string> GuidsForGroups = new();
 
 		public XcodeFileCollection(XcodeProjectFile ProjectFile)
 		{
@@ -572,7 +572,7 @@ namespace UnrealBuildTool.XcodeProjectXcconfig
 		/// <summary>
 		/// Cache the folders excluded by Mac
 		/// </summary>
-		static string[] ExcludedFolders = PlatformExports.GetExcludedFolderNames(UnrealTargetPlatform.Mac);
+		static readonly string[] ExcludedFolders = PlatformExports.GetExcludedFolderNames(UnrealTargetPlatform.Mac);
 
 		/// <summary>
 		/// Returns a project navigator group to which the file should belong based on its path.
diff --git a/Engine/Source/Programs/UnrealBuildTool/ProjectFiles/Xcode/XcodeProject.Utils.cs b/Engine/Source/Programs/UnrealBuildTool/ProjectFiles/Xcode/XcodeProject.Utils.cs
--- a/Engine/Source/Programs/UnrealBuildTool/ProjectFiles/Xcode/XcodeProject.Utils.cs
+++ b/Engine/Source/Programs/UnrealBuildTool/ProjectFiles/Xcode/XcodeProject.Utils.cs
@@ -77,7 +77,7 @@ namespace UnrealBuildTool.XcodeProjectXcconfig
 		}
 
 		// cache for the below function
-		static Dictionary<string, UnrealArchitectures> CachedMacProjectArchitectures = new();
+		static readonly Dictionary<string, UnrealArchitectures> CachedMacProjectArchitectures = new();
 
 		/// <summary>
 		/// Returns the Mac architectures that should be configured for the provided target. If the target has a project we'll adhere
diff --git a/Engine/Source/Programs/UnrealBuildTool/ProjectFiles/Xcode/XcodeProject.cs b/Engine/Source/Programs/UnrealBuildTool/ProjectFiles/Xcode/XcodeProject.cs
--- a/Engine/Source/Programs/UnrealBuildTool/ProjectFiles/Xcode/XcodeProject.cs
+++ b/Engine/Source/Programs/UnrealBuildTool/ProjectFiles/Xcode/XcodeProject.cs
@@ -949,7 +949,7 @@ namespace UnrealBuildTool.XcodeProjectXcconfig
 
 	class XcodeResourcesBuildPhase : XcodeBuildPhase
 	{
-		private XcodeFileCollection FileCollection;
+		private readonly XcodeFileCollection FileCollection;
 
 		public XcodeResourcesBuildPhase(XcodeFileCollection FileCollection)
 			: base("Resources", "PBXResourcesBuildPhase")
@@ -977,7 +977,7 @@ namespace UnrealBuildTool.XcodeProjectXcconfig
 
 	class XcodeFrameworkBuildPhase : XcodeBuildPhase
 	{
-		private XcodeFileCollection FileCollection;
+		private readonly XcodeFileCollection FileCollection;
 
 		public XcodeFrameworkBuildPhase(XcodeFileCollection FileCollection)
 			: base("Frameworks", "PBXFrameworksBuildPhase")
@@ -996,7 +996,7 @@ namespace UnrealBuildTool.XcodeProjectXcconfig
 
 	class XcodeCopyFilesBuildPhase : XcodeBuildPhase
 	{
-		private XcodeFileCollection FileCollection;
+		private readonly XcodeFileCollection FileCollection;
 
 		public XcodeCopyFilesBuildPhase(XcodeFileCollection FileCollection)
 			: base("Embed Frameworks", "PBXCopyFilesBuildPhase")
@@ -1018,8 +1018,8 @@ namespace UnrealBuildTool.XcodeProjectXcconfig
 
 	class XcodeCopyAppExtFilesBuildPhase : XcodeBuildPhase
 	{
-		private XcodeFileCollection FileCollection;
-		private ExtensionConfig ExtConfig;
+		private readonly XcodeFileCollection FileCollection;
+		private readonly ExtensionConfig ExtConfig;
 
 		public XcodeCopyAppExtFilesBuildPhase(XcodeFileCollection FileCollection, ExtensionConfig ExtConfig)
 			: base("Embed Foundation Extensions", "PBXCopyFilesBuildPhase")
@@ -1088,7 +1088,7 @@ namespace UnrealBuildTool.XcodeProjectXcconfig
 		public UnrealBuildConfig Info;
 		// Because we don't make project-wide .xcconfig files, we need to specify at the project level that all of the platforms are supported,
 		// so that the _Build targets will have any platform as a supported platform, otherwise it can only compile for Mac (default Xcode platform)
-		private bool bIncludeAllPlatforms;
+		private readonly bool bIncludeAllPlatforms;
 
 		public XcodeBuildConfig(UnrealBuildConfig Info, bool bIncludeAllPlatforms)
 		{
@@ -1128,7 +1128,7 @@ namespace UnrealBuildTool.XcodeProjectXcconfig
 	{
 		public string Guid = XcodeProjectFileGenerator.MakeXcodeGuid();
 		public string TargetName;
-		private UnrealData UnrealData;
+		private readonly UnrealData UnrealData;
 
 		public List<XcodeBuildConfig> BuildConfigs = new();
 
@@ -1201,11 +1201,11 @@ namespace UnrealBuildTool.XcodeProjectXcconfig
 		//		string TargetAppGuid = XcodeProjectFileGenerator.MakeXcodeGuid();
 
 		// com.apple.product-type.application, etc
-		string ProductType;
+		readonly string ProductType;
 
 		// xcode target type name
-		string TargetTypeName;
-		Type TargetType;
+		readonly string TargetTypeName;
+		readonly Type TargetType;
 
 		// UnrealEngine_Build, etc
 		public string Name;
@@ -1222,7 +1222,7 @@ namespace UnrealBuildTool.XcodeProjectXcconfig
 		// build phases for this target (source, resource copying, etc)
 		public List<XcodeBuildPhase> BuildPhases = new List<XcodeBuildPhase>();
 
-		private FileReference? GameProject;
+		private readonly FileReference? GameProject;
 
 		public XcodeTarget(Type Type, UnrealData UnrealData, string? OverrideName = null)
 		{
@@ -1324,7 +1324,7 @@ namespace UnrealBuildTool.XcodeProjectXcconfig
 			// nothing by default
 		}
 		
-		private static Dictionary<string, string> PlistFileMap = new(StringComparer.OrdinalIgnoreCase);
+		private static readonly Dictionary<string, string> PlistFileMap = new(StringComparer.OrdinalIgnoreCase);
 		protected string GetPlistSigningName(string ProvisionSetting, UnrealData UnrealData, ILogger Logger)
 		{
 			string? SigningName;
@@ -1384,10 +1384,10 @@ namespace UnrealBuildTool.XcodeProjectXcconfig
 	
 	class XcodeRunTarget : XcodeTarget
 	{
-		private string ProductGuid = XcodeProjectFileGenerator.MakeXcodeGuid();
-		private UnrealData UnrealData;
-		public UnrealTargetPlatform Platform;
-		public TargetType TargetType;
+		private readonly string ProductGuid = XcodeProjectFileGenerator.MakeXcodeGuid();
+		private readonly UnrealData UnrealData;
+		public readonly UnrealTargetPlatform Platform;
+		public readonly TargetType TargetType;
 
 		public XcodeRunTarget(XcodeProject Project, string TargetName, TargetType TargetType, UnrealTargetPlatform Platform, XcodeBuildTarget? BuildTarget, XcodeProjectFile ProjectFile, ILogger Logger)
 			: base(Project.UnrealData.IsAppBundle(Platform) ? XcodeTarget.Type.Run_App : XcodeTarget.Type.Run_Tool, Project.UnrealData,
@@ -2370,13 +2370,13 @@ namespace UnrealBuildTool.XcodeProjectXcconfig
 			NotificationService
 		}
 
-		private UnrealData UnrealData;
-		private DirectoryInfo ExtensionPath;
-		private UnrealTargetPlatform Platform;
-		private TargetType TargetType;
-		private string TargetName;
-		private string AppGuid;
-		private ExtensionConfig ExtConfig;
+		private readonly UnrealData UnrealData;
+		private readonly DirectoryInfo ExtensionPath;
+		private readonly UnrealTargetPlatform Platform;
+		private readonly TargetType TargetType;
+		private readonly string TargetName;
+		private readonly string AppGuid;
+		private readonly ExtensionConfig ExtConfig;
 		
 		public XcodeAppExtRunTarget(XcodeProject Project, ExtensionConfig? ExtConfig, string TargetName, string AppGuid, TargetType TargetType, UnrealTargetPlatform Platform, XcodeBuildTarget? BuildTarget, XcodeProjectFile ProjectFile, DirectoryInfo ExtensionPath, ILogger Logger)
 			: base(XcodeTarget.Type.App_Ext, Project.UnrealData,
@@ -2874,7 +2874,7 @@ namespace UnrealBuildTool.XcodeProjectXcconfig
 
 	class XcodeIndexTarget : XcodeTarget
 	{
-		private UnrealData UnrealData;
+		private readonly UnrealData UnrealData;
 
 		// just take the Project since it has everything we need, and is needed when adding target dependencies
 		public XcodeIndexTarget(XcodeProject Project, UnrealTargetPlatform Platform)
@@ -2932,7 +2932,7 @@ namespace UnrealBuildTool.XcodeProjectXcconfig
 
 		// Guid for the project node
 		public string Guid = XcodeProjectFileGenerator.MakeXcodeGuid();
-		private string ProvisioningStyle;
+		private readonly string ProvisioningStyle;
 
 		public List<XcodeRunTarget> RunTargets = new();
 		public List<XcodeAppExtRunTarget> AppExtRunTargets = new();
@@ -3083,11 +3083,11 @@ namespace UnrealBuildTool.XcodeProjectXcconfig
 		/// <summary>
 		/// Gathers the files and generates project sections
 		/// </summary>
-		private XcodeFileCollection SharedFileCollection;
+		private readonly XcodeFileCollection SharedFileCollection;
 		public XcodeFileCollection FileCollection;
 
 		// if set, only this will be written
-		private string? SingleTargetName;
+		private readonly string? SingleTargetName;
 
 		/// <summary>
 		/// Allocates a generator-specific source file object
diff --git a/Engine/Source/Programs/UnrealBuildTool/ProjectFiles/Xcode/XcodeProjectFileGenerator.cs b/Engine/Source/Programs/UnrealBuildTool/ProjectFiles/Xcode/XcodeProjectFileGenerator.cs
--- a/Engine/Source/Programs/UnrealBuildTool/ProjectFiles/Xcode/XcodeProjectFileGenerator.cs
+++ b/Engine/Source/Programs/UnrealBuildTool/ProjectFiles/Xcode/XcodeProjectFileGenerator.cs
@@ -33,7 +33,7 @@ namespace UnrealBuildTool
 		public DirectoryReference? XCWorkspace;
 
 		// always seed the random number the same, so multiple runs of the generator will generate the same project
-		static Random Rand = new Random(0);
+		static readonly Random Rand = new Random(0);
 
 		// how to handle multiple platforms
 		public static XcodePerPlatformMode PerPlatformMode = XcodePerPlatformMode.OneWorkspacePerPlatform;
@@ -41,22 +41,22 @@ namespace UnrealBuildTool
 		/// <summary>
 		/// Mark for distribution builds
 		/// </summary>
-		bool bForDistribution = false;
+		readonly bool bForDistribution = false;
 
 		/// <summary>
 		/// Don't use entitlements
 		/// </summary>
-		bool bNoEntitlements = false;
+		readonly bool bNoEntitlements = false;
 
 		/// <summary>
 		/// Override BundleID
 		/// </summary>
-		string BundleIdentifier = "";
+		readonly string BundleIdentifier = "";
 
 		/// <summary>
 		/// Override AppName
 		/// </summary>
-		string AppName = "";
+		readonly string AppName = "";
 
 		/// <summary>
 		/// Store the single game project (when using -game -project=...) to a place that XcodeProjectLegacy can easily retrieve it
diff --git a/Engine/Source/Programs/UnrealBuildTool/ProjectFiles/Xcode/XcodeProjectLegacy.cs b/Engine/Source/Programs/UnrealBuildTool/ProjectFiles/Xcode/XcodeProjectLegacy.cs
--- a/Engine/Source/Programs/UnrealBuildTool/ProjectFiles/Xcode/XcodeProjectLegacy.cs
+++ b/Engine/Source/Programs/UnrealBuildTool/ProjectFiles/Xcode/XcodeProjectLegacy.cs
@@ -112,27 +112,27 @@ namespace UnrealBuildTool.XcodeProjectLegacy
 		/// <summary>
 		/// Temporary for developing this feature
 		/// </summary>
-		bool bMakeProjectPerTarget;
+		readonly bool bMakeProjectPerTarget;
 
 		/// <summary>
 		///  Used to mark the project for distribution (some platforms require this)
 		/// </summary>
-		bool bForDistribution = false;
+		readonly bool bForDistribution = false;
 
 		/// <summary>
 		/// Override for bundle identifier
 		/// </summary>
-		string BundleIdentifier = "";
+		readonly string BundleIdentifier = "";
 
 		/// <summary>
 		/// Override AppName
 		/// </summary>
-		string AppName = "";
+		readonly string AppName = "";
 
 		/// <summary>
 		/// Architectures supported for iOS
 		/// </summary>
-		UnrealArch[] SupportedIOSArchitectures = { UnrealArch.Arm64 };
+		readonly UnrealArch[] SupportedIOSArchitectures = { UnrealArch.Arm64 };
 
 		/// <summary>
 		/// Gets Xcode file category based on its extension
@@ -926,7 +926,7 @@ namespace UnrealBuildTool.XcodeProjectLegacy
 		}
 
 		// cache for the below function
-		Dictionary<string, UnrealArchitectures> CachedMacProjectArcitectures = new();
+		readonly Dictionary<string, UnrealArchitectures> CachedMacProjectArcitectures = new();
 
 		/// <summary>
 		/// Returns the Mac architectures that should be configured for the provided target. If the target has a project we'll adhere
diff --git a/Engine/Source/Programs/UnrealBuildTool/Storage/Impl/FileSystemStorageProvider.cs b/Engine/Source/Programs/UnrealBuildTool/Storage/Impl/FileSystemStorageProvider.cs
--- a/Engine/Source/Programs/UnrealBuildTool/Storage/Impl/FileSystemStorageProvider.cs
+++ b/Engine/Source/Programs/UnrealBuildTool/Storage/Impl/FileSystemStorageProvider.cs
@@ -15,7 +15,7 @@ namespace UnrealBuildTool.Storage.Impl
 		/// <summary>
 		/// Cached copy of the current process id
 		/// </summary>
-		static int ProcessId = Environment.ProcessId;
+		static readonly int ProcessId = Environment.ProcessId;
 
 		/// <summary>
 		/// Implements <see cref="IStorageReader"/> for writes to the backing storage
@@ -54,7 +54,7 @@ namespace UnrealBuildTool.Storage.Impl
 		class StorageWriter : IStorageWriter
 		{
 			FileReference? TempLocation;
-			FileReference FinalLocation;
+			readonly FileReference FinalLocation;
 
 			public Stream? Stream
 			{
diff --git a/Engine/Source/Programs/UnrealBuildTool/System/ActionGraphBuilder.cs b/Engine/Source/Programs/UnrealBuildTool/System/ActionGraphBuilder.cs
--- a/Engine/Source/Programs/UnrealBuildTool/System/ActionGraphBuilder.cs
+++ b/Engine/Source/Programs/UnrealBuildTool/System/ActionGraphBuilder.cs
@@ -175,7 +175,7 @@ namespace UnrealBuildTool
 		/// <summary>
 		/// The inner graph builder
 		/// </summary>
-		IActionGraphBuilder Inner;
+		readonly IActionGraphBuilder Inner;
 
 		/// <summary>
 		/// Constructor
diff --git a/Engine/Source/Programs/UnrealBuildTool/System/ActionHistory.cs b/Engine/Source/Programs/UnrealBuildTool/System/ActionHistory.cs
--- a/Engine/Source/Programs/UnrealBuildTool/System/ActionHistory.cs
+++ b/Engine/Source/Programs/UnrealBuildTool/System/ActionHistory.cs
@@ -360,12 +360,12 @@ namespace UnrealBuildTool
 		/// <summary>
 		/// The lock object for this history
 		/// </summary>
-		object LockObject = new object();
+		readonly object LockObject = new object();
 
 		/// <summary>
 		/// List of partitions
 		/// </summary>
-		List<ActionHistoryPartition> Partitions = new List<ActionHistoryPartition>();
+		readonly List<ActionHistoryPartition> Partitions = new List<ActionHistoryPartition>();
 
 		/// <summary>
 		/// Constructor
diff --git a/Engine/Source/Programs/UnrealBuildTool/System/ActionSerializer.cs b/Engine/Source/Programs/UnrealBuildTool/System/ActionSerializer.cs
--- a/Engine/Source/Programs/UnrealBuildTool/System/ActionSerializer.cs
+++ b/Engine/Source/Programs/UnrealBuildTool/System/ActionSerializer.cs
@@ -70,12 +70,12 @@ namespace UnrealBuildTool
 		/// <summary>
 		/// Map from type name to deserializing constructor
 		/// </summary>
-		static IReadOnlyDictionary<Type, IActionSerializer> TypeToSerializer;
+		static readonly IReadOnlyDictionary<Type, IActionSerializer> TypeToSerializer;
 
 		/// <summary>
 		/// Map from serializer name to instance
 		/// </summary>
-		static IReadOnlyDictionary<string, IActionSerializer> NameToSerializer;
+		static readonly IReadOnlyDictionary<string, IActionSerializer> NameToSerializer;
 
 #pragma warning disable CA1810 // Initialize reference type static fields inline
 		/// <summary>
diff --git a/Engine/Source/Programs/UnrealBuildTool/System/BuildVersion.cs b/Engine/Source/Programs/UnrealBuildTool/System/BuildVersion.cs
--- a/Engine/Source/Programs/UnrealBuildTool/System/BuildVersion.cs
+++ b/Engine/Source/Programs/UnrealBuildTool/System/BuildVersion.cs
@@ -242,7 +242,7 @@ namespace UnrealBuildTool
 		/// <summary>
 		/// The inner build version
 		/// </summary>
-		private BuildVersion Inner;
+		private readonly BuildVersion Inner;
 
 		/// <summary>
 		/// Cached copy of the current build version
diff --git a/Engine/Source/Programs/UnrealBuildTool/System/ConfigCache.cs b/Engine/Source/Programs/UnrealBuildTool/System/ConfigCache.cs
--- a/Engine/Source/Programs/UnrealBuildTool/System/ConfigCache.cs
+++ b/Engine/Source/Programs/UnrealBuildTool/System/ConfigCache.cs
@@ -83,7 +83,7 @@ namespace UnrealBuildTool
 			/// <summary>
 			/// Reflection description of the field with the config attribute.
 			/// </summary>
-			private FieldInfo FieldInfo;
+			private readonly FieldInfo FieldInfo;
 
 			/// <summary>
 			/// Constructor
@@ -110,7 +110,7 @@ namespace UnrealBuildTool
 			/// <summary>
 			/// Reflection description of the property with the config attribute.
 			/// </summary>
-			private PropertyInfo PropertyInfo;
+			private readonly PropertyInfo PropertyInfo;
 
 			/// <summary>
 			/// Constructor
@@ -248,17 +248,17 @@ namespace UnrealBuildTool
 		/// <summary>
 		/// Cache of individual config files
 		/// </summary>
-		static ConcurrentDictionary<FileReference, Lazy<ConfigFile?>> LocationToConfigFile = [];
+		static readonly ConcurrentDictionary<FileReference, Lazy<ConfigFile?>> LocationToConfigFile = [];
 
 		/// <summary>
 		/// Cache of config hierarchies by project
 		/// </summary>
-		static ConcurrentDictionary<ConfigHierarchyKey, Lazy<ConfigHierarchy>> HierarchyKeyToHierarchy = [];
+		static readonly ConcurrentDictionary<ConfigHierarchyKey, Lazy<ConfigHierarchy>> HierarchyKeyToHierarchy = [];
 
 		/// <summary>
 		/// Cache of config fields by type
 		/// </summary>
-		static ConcurrentDictionary<Type, Lazy<List<ConfigMember>>> TypeToConfigMembers = [];
+		static readonly ConcurrentDictionary<Type, Lazy<List<ConfigMember>>> TypeToConfigMembers = [];
 
 		/// <summary>
 		/// Attempts to read a config file (or retrieve it from the cache)
diff --git a/Engine/Source/Programs/UnrealBuildTool/System/ConfigFile.cs b/Engine/Source/Programs/UnrealBuildTool/System/ConfigFile.cs
--- a/Engine/Source/Programs/UnrealBuildTool/System/ConfigFile.cs
+++ b/Engine/Source/Programs/UnrealBuildTool/System/ConfigFile.cs
@@ -150,12 +150,12 @@ namespace UnrealBuildTool
 		/// <summary>
 		/// Maps names to config sections
 		/// </summary>
-		Dictionary<string, ConfigFileSection> Sections = new Dictionary<string, ConfigFileSection>(StringComparer.InvariantCultureIgnoreCase);
+		readonly Dictionary<string, ConfigFileSection> Sections = new Dictionary<string, ConfigFileSection>(StringComparer.InvariantCultureIgnoreCase);
 
 		// Remap of config names/sections
-		private static Dictionary<string, string> SectionNameRemap = new();
-		private static Dictionary<string, Dictionary<string, string>> SectionKeyRemap = new();
-		private static HashSet<string> WarnedKeys = new(StringComparer.InvariantCultureIgnoreCase);
+		private static readonly Dictionary<string, string> SectionNameRemap = new();
+		private static readonly Dictionary<string, Dictionary<string, string>> SectionKeyRemap = new();
+		private static readonly HashSet<string> WarnedKeys = new(StringComparer.InvariantCultureIgnoreCase);
 
 		private static string RemapSectionOrKey(Dictionary<string, string>? Remap, string Key, string Context)
 		{
diff --git a/Engine/Source/Programs/UnrealBuildTool/System/ConfigHierarchy.cs b/Engine/Source/Programs/UnrealBuildTool/System/ConfigHierarchy.cs
--- a/Engine/Source/Programs/UnrealBuildTool/System/ConfigHierarchy.cs
+++ b/Engine/Source/Programs/UnrealBuildTool/System/ConfigHierarchy.cs
@@ -71,7 +71,7 @@ namespace UnrealBuildTool
 		/// <summary>
 		/// Map of key names to their values
 		/// </summary>
-		Dictionary<string, List<string>> KeyToValue = new Dictionary<string, List<string>>(StringComparer.InvariantCultureIgnoreCase);
+		readonly Dictionary<string, List<string>> KeyToValue = new Dictionary<string, List<string>>(StringComparer.InvariantCultureIgnoreCase);
 
 		/// <summary>
 		/// Construct a merged config section from the given per-file config sections
@@ -952,7 +952,7 @@ namespace UnrealBuildTool
 			public string? After2 = null;
 		};
 
-		static string[] ConfigLayers =
+		static readonly string[] ConfigLayers =
 		{
 			// Engine/Base.ini
 			"{ENGINE}/Config/Base.ini",
@@ -987,7 +987,7 @@ namespace UnrealBuildTool
 			"{PROJECT}/Config/User{TYPE}.ini",
 		};
 
-		static ConfigLayerExpansion[] ConfigLayerExpansions =
+		static readonly ConfigLayerExpansion[] ConfigLayerExpansions =
 		{
 			// The base expansion (ie, no expansion)
 			new ConfigLayerExpansion { }, 
diff --git a/Engine/Source/Programs/UnrealBuildTool/System/CppDependencyCache.cs b/Engine/Source/Programs/UnrealBuildTool/System/CppDependencyCache.cs
--- a/Engine/Source/Programs/UnrealBuildTool/System/CppDependencyCache.cs
+++ b/Engine/Source/Programs/UnrealBuildTool/System/CppDependencyCache.cs
@@ -156,9 +156,9 @@ namespace UnrealBuildTool
 		/// <summary>
 		/// List of partitions
 		/// </summary>
-		List<(DirectoryReference Root, List<(string Filter, CachePartition Partition)>)> GlobalPartitions = [];
+		readonly List<(DirectoryReference Root, List<(string Filter, CachePartition Partition)>)> GlobalPartitions = [];
 
-		List<Task> GlobalTasks = [];
+		readonly List<Task> GlobalTasks = [];
 
 		/// <summary>
 		/// Minimum version of a MSVC source dependency json
diff --git a/Engine/Source/Programs/UnrealBuildTool/System/CppIncludeLookup.cs b/Engine/Source/Programs/UnrealBuildTool/System/CppIncludeLookup.cs
--- a/Engine/Source/Programs/UnrealBuildTool/System/CppIncludeLookup.cs
+++ b/Engine/Source/Programs/UnrealBuildTool/System/CppIncludeLookup.cs
@@ -40,10 +40,10 @@ namespace UnrealBuildTool
 
 		const int CurrentVersion = 1;
 
-		static Regex SourceFileRegex = new Regex(@"\.(?:c|cpp|cxx)$");
-		static Regex SourceOrHeaderFileRegex = new Regex(@"\.(?:c|h|cpp|hpp|cxx|hxx)$");
+		static readonly Regex SourceFileRegex = new Regex(@"\.(?:c|cpp|cxx)$");
+		static readonly Regex SourceOrHeaderFileRegex = new Regex(@"\.(?:c|h|cpp|hpp|cxx|hxx)$");
 
-		FileReference Location;
+		readonly FileReference Location;
 		Dictionary<string, CppIncludeNameInfo> NameToInfo = new Dictionary<string, CppIncludeNameInfo>(StringComparer.OrdinalIgnoreCase);
 
 		public CppIncludeLookup(FileReference Location)
@@ -298,8 +298,8 @@ namespace UnrealBuildTool
 		}
 
 		const string IncludeText = "include";
-		static ReadOnlyMemory<byte> IncludeBytes = Encoding.UTF8.GetBytes(IncludeText);
-		static sbyte[] SkipTable = CreateSkipTable();
+		static readonly ReadOnlyMemory<byte> IncludeBytes = Encoding.UTF8.GetBytes(IncludeText);
+		static readonly sbyte[] SkipTable = CreateSkipTable();
 
 		static sbyte[] CreateSkipTable()
 		{
diff --git a/Engine/Source/Programs/UnrealBuildTool/System/CustomBuildSteps.cs b/Engine/Source/Programs/UnrealBuildTool/System/CustomBuildSteps.cs
--- a/Engine/Source/Programs/UnrealBuildTool/System/CustomBuildSteps.cs
+++ b/Engine/Source/Programs/UnrealBuildTool/System/CustomBuildSteps.cs
@@ -12,7 +12,7 @@ namespace UnrealBuildTool
 	/// </summary>
 	public class CustomBuildSteps
 	{
-		Dictionary<UnrealTargetPlatform, string[]> HostPlatformToCommands = new Dictionary<UnrealTargetPlatform, string[]>();
+		readonly Dictionary<UnrealTargetPlatform, string[]> HostPlatformToCommands = new Dictionary<UnrealTargetPlatform, string[]>();
 
 		/// <summary>
 		/// Construct a custom build steps object from a Json object.
diff --git a/Engine/Source/Programs/UnrealBuildTool/System/DataDrivenPlatformInfo.cs b/Engine/Source/Programs/UnrealBuildTool/System/DataDrivenPlatformInfo.cs
--- a/Engine/Source/Programs/UnrealBuildTool/System/DataDrivenPlatformInfo.cs
+++ b/Engine/Source/Programs/UnrealBuildTool/System/DataDrivenPlatformInfo.cs
@@ -61,7 +61,7 @@ namespace UnrealBuildTool
 		};
 
 		static IReadOnlyDictionary<string, ConfigDataDrivenPlatformInfo>? PlatformInfos = null;
-		static object LockObject = new object();
+		static readonly object LockObject = new object();
 
 		/// <summary>
 		/// Return all data driven infos found
diff --git a/Engine/Source/Programs/UnrealBuildTool/System/EncryptionAndSigning.cs b/Engine/Source/Programs/UnrealBuildTool/System/EncryptionAndSigning.cs
--- a/Engine/Source/Programs/UnrealBuildTool/System/EncryptionAndSigning.cs
+++ b/Engine/Source/Programs/UnrealBuildTool/System/EncryptionAndSigning.cs
@@ -245,7 +245,7 @@ namespace UnrealBuildTool
 			return true;
 		}
 
-		static ConcurrentDictionary<FileReference, Lazy<List<EncryptionKey>?>> s_encryptionKeyCache = [];
+		static readonly ConcurrentDictionary<FileReference, Lazy<List<EncryptionKey>?>> s_encryptionKeyCache = [];
 
 		/// <summary>
 		/// Parse crypto settings from INI file
diff --git a/Engine/Source/Programs/UnrealBuildTool/System/FileHashCache.cs b/Engine/Source/Programs/UnrealBuildTool/System/FileHashCache.cs
--- a/Engine/Source/Programs/UnrealBuildTool/System/FileHashCache.cs
+++ b/Engine/Source/Programs/UnrealBuildTool/System/FileHashCache.cs
@@ -645,10 +645,10 @@ namespace UnrealBuildTool
 			internal ConcurrentDictionary<FileItem, FileEntry> Files = [];
 		}
 
-		private List<(DirectoryReference root, List<(string filter, Partition partition)>)> _roots = [];
-		private bool _warnOnNoPartition;
-		private bool _backupModifiedFiles;
-		private bool _errorOnRewrite;
+		private readonly List<(DirectoryReference root, List<(string filter, Partition partition)>)> _roots = [];
+		private readonly bool _warnOnNoPartition;
+		private readonly bool _backupModifiedFiles;
+		private readonly bool _errorOnRewrite;
 		private bool _hasErrors;
 	}
 }
\ No newline at end of file
diff --git a/Engine/Source/Programs/UnrealBuildTool/System/FileMetadataPrefetch.cs b/Engine/Source/Programs/UnrealBuildTool/System/FileMetadataPrefetch.cs
--- a/Engine/Source/Programs/UnrealBuildTool/System/FileMetadataPrefetch.cs
+++ b/Engine/Source/Programs/UnrealBuildTool/System/FileMetadataPrefetch.cs
@@ -18,20 +18,21 @@ namespace UnrealBuildTool
 	/// </summary>
 	static class FileMetadataPrefetch
 	{
+		// TODO: This can never be cancelled, should it be wired up or removed?
 		/// <summary>
 		/// Used to cancel any queued tasks
 		/// </summary>
-		static CancellationTokenSource CancelSource = new CancellationTokenSource();
+		static readonly CancellationTokenSource CancelSource = new CancellationTokenSource();
 
 		/// <summary>
 		/// The cancellation token
 		/// </summary>
-		static CancellationToken CancelToken = CancelSource.Token;
+		static readonly CancellationToken CancelToken = CancelSource.Token;
 
 		/// <summary>
 		/// Set of all the directory trees that have been queued up, to save adding any more than once.
 		/// </summary>
-		static ConcurrentDictionary<DirectoryReference, Lazy<Task>> QueuedDirectories = new();
+		static readonly ConcurrentDictionary<DirectoryReference, Lazy<Task>> QueuedDirectories = new();
 
 		/// <summary>
 		/// Enqueue the engine directory for prefetching
diff --git a/Engine/Source/Programs/UnrealBuildTool/System/InstalledPlatformInfo.cs b/Engine/Source/Programs/UnrealBuildTool/System/InstalledPlatformInfo.cs
--- a/Engine/Source/Programs/UnrealBuildTool/System/InstalledPlatformInfo.cs
+++ b/Engine/Source/Programs/UnrealBuildTool/System/InstalledPlatformInfo.cs
@@ -117,7 +117,7 @@ namespace UnrealBuildTool
 			}
 		}
 
-		private static List<InstalledPlatformConfiguration>? InstalledPlatformConfigurations;
+		private static readonly List<InstalledPlatformConfiguration>? InstalledPlatformConfigurations;
 
 #pragma warning disable CA1810 // Initialize reference type static fields inline
 		static InstalledPlatformInfo()
diff --git a/Engine/Source/Programs/UnrealBuildTool/System/ProgressWriter.cs b/Engine/Source/Programs/UnrealBuildTool/System/ProgressWriter.cs
--- a/Engine/Source/Programs/UnrealBuildTool/System/ProgressWriter.cs
+++ b/Engine/Source/Programs/UnrealBuildTool/System/ProgressWriter.cs
@@ -19,12 +19,12 @@ namespace UnrealBuildTool
 		/// <summary>
 		/// Logger for output
 		/// </summary>
-		ILogger Logger;
+		readonly ILogger Logger;
 
 		/// <summary>
 		/// The name to include with the status message
 		/// </summary>
-		string Message;
+		readonly string Message;
 
 		/// <summary>
 		/// The inner scope object
diff --git a/Engine/Source/Programs/UnrealBuildTool/System/RestrictedFolders.cs b/Engine/Source/Programs/UnrealBuildTool/System/RestrictedFolders.cs
--- a/Engine/Source/Programs/UnrealBuildTool/System/RestrictedFolders.cs
+++ b/Engine/Source/Programs/UnrealBuildTool/System/RestrictedFolders.cs
@@ -14,12 +14,12 @@ namespace UnrealBuildTool
 		/// <summary>
 		/// Unique Id for this folder
 		/// </summary>
-		private int Id;
+		private readonly int Id;
 
 		/// <summary>
 		/// Mapping for unique ids
 		/// </summary>
-		private static UniqueStringRegistry StringRegistry = new UniqueStringRegistry();
+		private static readonly UniqueStringRegistry StringRegistry = new UniqueStringRegistry();
 
 		/// <summary>
 		/// Array of all restricted folder names
diff --git a/Engine/Source/Programs/UnrealBuildTool/System/ScopedTimer.cs b/Engine/Source/Programs/UnrealBuildTool/System/ScopedTimer.cs
--- a/Engine/Source/Programs/UnrealBuildTool/System/ScopedTimer.cs
+++ b/Engine/Source/Programs/UnrealBuildTool/System/ScopedTimer.cs
@@ -11,13 +11,15 @@ namespace UnrealBuildTool
 	/// </summary>
 	public sealed class ScopedTimer : IDisposable
 	{
-		DateTime StartTime;
-		string Name;
-		ILogger Logger;
-		LogLevel Level;
-		bool bIncreaseIndent;
+		readonly DateTime StartTime;
+		readonly string Name;
+		readonly ILogger Logger;
+		readonly LogLevel Level;
+		readonly bool bIncreaseIndent;
+		// TODO: This isn't true nesting support, as it assumes that all new ScopedTimers are within the previous scope.
+		// This doesn't account for totally unrelated timers, so this should probably use some context-local storage to deal with this.
 		static int Indent = 0;
-		static object IndentLock = new object();
+		static readonly object IndentLock = new object();
 
 		/// <summary>
 		/// Constructor
diff --git a/Engine/Source/Programs/UnrealBuildTool/System/SourceFileMetadataCache.cs b/Engine/Source/Programs/UnrealBuildTool/System/SourceFileMetadataCache.cs
--- a/Engine/Source/Programs/UnrealBuildTool/System/SourceFileMetadataCache.cs
+++ b/Engine/Source/Programs/UnrealBuildTool/System/SourceFileMetadataCache.cs
@@ -94,32 +94,32 @@ namespace UnrealBuildTool
 		/// <summary>
 		/// Location of this dependency cache
 		/// </summary>
-		FileReference Location;
+		readonly FileReference Location;
 
 		/// <summary>
 		/// Directory for files to cache dependencies for.
 		/// </summary>
-		DirectoryReference BaseDirectory;
+		readonly DirectoryReference BaseDirectory;
 
 		/// <summary>
 		/// The parent cache.
 		/// </summary>
-		SourceFileMetadataCache? Parent;
+		readonly SourceFileMetadataCache? Parent;
 
 		/// <summary>
 		/// Map from file item to source file info
 		/// </summary>
-		ConcurrentDictionary<FileItem, SourceFileInfo> FileToSourceFileInfo = new ConcurrentDictionary<FileItem, SourceFileInfo>();
+		readonly ConcurrentDictionary<FileItem, SourceFileInfo> FileToSourceFileInfo = new ConcurrentDictionary<FileItem, SourceFileInfo>();
 
 		/// <summary>
 		/// Map from file item to header file info
 		/// </summary>
-		ConcurrentDictionary<FileItem, HeaderFileInfo> FileToHeaderFileInfo = new ConcurrentDictionary<FileItem, HeaderFileInfo>();
+		readonly ConcurrentDictionary<FileItem, HeaderFileInfo> FileToHeaderFileInfo = new ConcurrentDictionary<FileItem, HeaderFileInfo>();
 
 		/// <summary>
 		/// Map from file item to source file
 		/// </summary>
-		ConcurrentDictionary<FileItem, SourceFile> FileToSourceFile = new ConcurrentDictionary<FileItem, SourceFile>();
+		readonly ConcurrentDictionary<FileItem, SourceFile> FileToSourceFile = new ConcurrentDictionary<FileItem, SourceFile>();
 
 		/// <summary>
 		/// Whether the cache has been modified and needs to be saved
@@ -129,7 +129,7 @@ namespace UnrealBuildTool
 		/// <summary>
 		/// Logger for output
 		/// </summary>
-		ILogger Logger;
+		readonly ILogger Logger;
 
 		/// <summary>
 		/// Regex that matches C++ code with UObject declarations which we will need to generated code for.
@@ -154,7 +154,7 @@ namespace UnrealBuildTool
 		/// <summary>
 		/// Static cache of all constructed dependency caches
 		/// </summary>
-		static ConcurrentDictionary<FileReference, Lazy<SourceFileMetadataCache>> Caches = new();
+		static readonly ConcurrentDictionary<FileReference, Lazy<SourceFileMetadataCache>> Caches = new();
 
 		/// <summary>
 		/// Constructs a dependency cache. This method is private; call CppDependencyCache.Create() to create a cache hierarchy for a given project.
diff --git a/Engine/Source/Programs/UnrealBuildTool/System/SourceFileWorkingSet.cs b/Engine/Source/Programs/UnrealBuildTool/System/SourceFileWorkingSet.cs
--- a/Engine/Source/Programs/UnrealBuildTool/System/SourceFileWorkingSet.cs
+++ b/Engine/Source/Programs/UnrealBuildTool/System/SourceFileWorkingSet.cs
@@ -300,11 +300,13 @@ namespace UnrealBuildTool
 			Git
 		};
 
+#pragma warning disable IDE0044 // Make field readonly - these private static fields are set by XML configuration.
 		/// <summary>
 		/// Sets the provider to use for determining the working set.
 		/// </summary>
 		[XmlConfigFile]
 		static ProviderType Provider = ProviderType.Default;
+#pragma warning restore IDE0044
 
 		/// <summary>
 		/// Sets the path to use for the repository. Interpreted relative to the Unreal Engine root directory (the folder above the Engine folder) -- if relative.
diff --git a/Engine/Source/Programs/UnrealBuildTool/System/Subnet.cs b/Engine/Source/Programs/UnrealBuildTool/System/Subnet.cs
--- a/Engine/Source/Programs/UnrealBuildTool/System/Subnet.cs
+++ b/Engine/Source/Programs/UnrealBuildTool/System/Subnet.cs
@@ -31,7 +31,7 @@ namespace UnrealBuildTool
 		/// <summary>
 		/// Bytes corresponding to the prefix address
 		/// </summary>
-		byte[] PrefixBytes;
+		readonly byte[] PrefixBytes;
 
 		/// <summary>
 		/// Constructor
diff --git a/Engine/Source/Programs/UnrealBuildTool/System/SupportedTargetTypesAttribute.cs b/Engine/Source/Programs/UnrealBuildTool/System/SupportedTargetTypesAttribute.cs
--- a/Engine/Source/Programs/UnrealBuildTool/System/SupportedTargetTypesAttribute.cs
+++ b/Engine/Source/Programs/UnrealBuildTool/System/SupportedTargetTypesAttribute.cs
@@ -19,7 +19,7 @@ namespace UnrealBuildTool
 		/// <summary>
 		/// HashSet of supported target types
 		/// </summary>
-		private HashSet<TargetType> _targetTypes;
+		private readonly HashSet<TargetType> _targetTypes;
 
 		/// <summary>
 		/// Initialize the attribute with a list of target types
diff --git a/Engine/Source/Programs/UnrealBuildTool/System/TargetMakefile.cs b/Engine/Source/Programs/UnrealBuildTool/System/TargetMakefile.cs
--- a/Engine/Source/Programs/UnrealBuildTool/System/TargetMakefile.cs
+++ b/Engine/Source/Programs/UnrealBuildTool/System/TargetMakefile.cs
@@ -71,7 +71,7 @@ namespace UnrealBuildTool
 
 		public FileHashCache FileHashCache { get; }
 
-		private ConcurrentBag<Task> WriteTasks = new();
+		private readonly ConcurrentBag<Task> WriteTasks = new();
 
 		public TargetMakefileBuilder(TargetMakefile InMakefile, FileHashCache fileHashCache, ILogger InLogger)
 		{
diff --git a/Engine/Source/Programs/UnrealBuildTool/System/TopologicalSorter.cs b/Engine/Source/Programs/UnrealBuildTool/System/TopologicalSorter.cs
--- a/Engine/Source/Programs/UnrealBuildTool/System/TopologicalSorter.cs
+++ b/Engine/Source/Programs/UnrealBuildTool/System/TopologicalSorter.cs
@@ -68,12 +68,12 @@ namespace UnrealBuildTool
 		/// It's only needed for diagnostic purposes i.e. when we find a cycle we use Callstack
 		/// to output all the cycle's nodes.
 		/// </summary>
-		private Stack<GraphNode> Callstack = new Stack<GraphNode>();
+		private readonly Stack<GraphNode> Callstack = new Stack<GraphNode>();
 
 		/// <summary>
 		/// The resulting sequence of nodes (T, provided by the user) in the topological order.
 		/// </summary>
-		private List<T> Result = new List<T>();
+		private readonly List<T> Result = new List<T>();
 
 		/// <summary>
 		/// Initialize the sorter object and create the internal graph representation.
diff --git a/Engine/Source/Programs/UnrealBuildTool/System/UHTExecution.cs b/Engine/Source/Programs/UnrealBuildTool/System/UHTExecution.cs
--- a/Engine/Source/Programs/UnrealBuildTool/System/UHTExecution.cs
+++ b/Engine/Source/Programs/UnrealBuildTool/System/UHTExecution.cs
@@ -1247,8 +1247,8 @@ namespace UnrealBuildTool
 				}
 			}
 
-			FileReference _location;
-			SortedDictionary<string, (UInt128 hash, long lastUpdateUtc)> _modules = [];
+			readonly FileReference _location;
+			readonly SortedDictionary<string, (UInt128 hash, long lastUpdateUtc)> _modules = [];
 		}
 
 		/// <summary>
diff --git a/Engine/Source/Programs/UnrealBuildTool/System/UniqueStringRegistry.cs b/Engine/Source/Programs/UnrealBuildTool/System/UniqueStringRegistry.cs
--- a/Engine/Source/Programs/UnrealBuildTool/System/UniqueStringRegistry.cs
+++ b/Engine/Source/Programs/UnrealBuildTool/System/UniqueStringRegistry.cs
@@ -12,7 +12,7 @@ namespace UnrealBuildTool
 	public class UniqueStringRegistry
 	{
 		// protect the InstanceMap
-		private object LockObject = new object();
+		private readonly object LockObject = new object();
 
 		// holds a mapping of string name to single instance
 		private Dictionary<string, int> StringToInstanceMap = new Dictionary<string, int>(StringComparer.OrdinalIgnoreCase);
diff --git a/Engine/Source/Programs/UnrealBuildTool/System/Unity.cs b/Engine/Source/Programs/UnrealBuildTool/System/Unity.cs
--- a/Engine/Source/Programs/UnrealBuildTool/System/Unity.cs
+++ b/Engine/Source/Programs/UnrealBuildTool/System/Unity.cs
@@ -81,7 +81,7 @@ namespace UnrealBuildTool
 		{
 			private List<FileCollection> UnityFiles;
 			private FileCollection CurrentCollection;
-			private int SplitLength;
+			private readonly int SplitLength;
 
 			/// <summary>
 			/// Constructs a new UnityFileBuilder.
diff --git a/Engine/Source/Programs/UnrealBuildTool/System/UnrealPluginLanguage.cs b/Engine/Source/Programs/UnrealBuildTool/System/UnrealPluginLanguage.cs
--- a/Engine/Source/Programs/UnrealBuildTool/System/UnrealPluginLanguage.cs
+++ b/Engine/Source/Programs/UnrealBuildTool/System/UnrealPluginLanguage.cs
@@ -509,12 +509,12 @@ namespace UnrealBuildTool
 	class UnrealPluginLanguage
 	{
 		/** The merged XML program to run */
-		private XDocument XDoc;
+		private readonly XDocument XDoc;
 
 		/** XML namespace */
-		private XNamespace XMLDefaultNameSpace;
-		private Dictionary<string, XNamespace> XMLNameSpaceAliases = new Dictionary<string, XNamespace>();
-		private string XMLRootDefinition;
+		private readonly XNamespace XMLDefaultNameSpace;
+		private readonly Dictionary<string, XNamespace> XMLNameSpaceAliases = new Dictionary<string, XNamespace>();
+		private readonly string XMLRootDefinition;
 
 		private UnrealTargetPlatform TargetPlatform;
 
@@ -523,7 +523,7 @@ namespace UnrealBuildTool
 		private static bool bGlobalTraceFilename = false;
 
 		/** Project file reference */
-		private FileReference? ProjectFile;
+		private readonly FileReference? ProjectFile;
 
 		private static XDocument XMLDummy = XDocument.Parse("<manifest></manifest>");
 
@@ -565,11 +565,13 @@ namespace UnrealBuildTool
 			}
 		}
 
-		private UPLContext GlobalContext;
-		private Dictionary<string, UPLContext> Contexts;
+		private readonly UPLContext GlobalContext;
+		private readonly Dictionary<string, UPLContext> Contexts;
 		private int ContextIndex;
+#pragma warning disable IDE0044 // Make field readonly - this field appears intended to be mutable.
 		private string? LastError;
-		private ILogger Logger;
+#pragma warning restore IDE0044
+		private readonly ILogger Logger;
 
 		// First entry in InXMLNameSpaceAliases will be used as default namespace
 		public UnrealPluginLanguage(FileReference? InProjectFile, List<string> InXMLFiles, List<(string UEArch, string NativeArch)> InArchs, List<AliasedXMLNamespace>? InXMLNameSpaceAliases, UnrealTargetPlatform InTargetPlatform, ILogger InLogger)
@@ -2774,9 +2776,9 @@ namespace UnrealBuildTool
 		}
 
 		// cached ini files
-		static Dictionary<string, List<Command>> FileCache = new Dictionary<string, List<Command>>();
-		static Dictionary<string, ConfigCacheIni_UPL> IniCache = new Dictionary<string, ConfigCacheIni_UPL>();
-		static Dictionary<string, ConfigCacheIni_UPL> BaseIniCache = new Dictionary<string, ConfigCacheIni_UPL>();
+		static readonly Dictionary<string, List<Command>> FileCache = new Dictionary<string, List<Command>>();
+		static readonly Dictionary<string, ConfigCacheIni_UPL> IniCache = new Dictionary<string, ConfigCacheIni_UPL>();
+		static readonly Dictionary<string, ConfigCacheIni_UPL> BaseIniCache = new Dictionary<string, ConfigCacheIni_UPL>();
 
 		// static creation functions for ini files
 		public static ConfigCacheIni_UPL CreateConfigCacheIni_UPL(UnrealTargetPlatform Platform, string BaseIniName, DirectoryReference? ProjectDirectory, ILogger Logger, DirectoryReference? EngineDirectory = null)
@@ -2845,12 +2847,12 @@ namespace UnrealBuildTool
 		/// <summary>
 		/// True if we are loading a hierarchy of config files that should be merged together
 		/// </summary>
-		bool bIsMergingConfigs;
+		readonly bool bIsMergingConfigs = false;
 
 		/// <summary>
 		/// All sections parsed from ini file
 		/// </summary>
-		Dictionary<string, IniSection> Sections;
+		readonly Dictionary<string, IniSection> Sections;
 
 		private ConfigCacheIni_UPL()
 		{
diff --git a/Engine/Source/Programs/UnrealBuildTool/System/Utils.cs b/Engine/Source/Programs/UnrealBuildTool/System/Utils.cs
--- a/Engine/Source/Programs/UnrealBuildTool/System/Utils.cs
+++ b/Engine/Source/Programs/UnrealBuildTool/System/Utils.cs
@@ -45,7 +45,7 @@ namespace UnrealBuildTool
 		/// <summary>
 		/// Regular expression to match $(ENV) and/ or %ENV% environment variables.
 		/// </summary>
-		static Regex EnvironmentVariableRegex = new Regex(@"\$\((.*?)\)|\%(.*?)\%", RegexOptions.None);
+		static readonly Regex EnvironmentVariableRegex = new Regex(@"\$\((.*?)\)|\%(.*?)\%", RegexOptions.None);
 
 		/// <summary>
 		/// Resolves $(ENV) and/ or %ENV% to the value of the environment variable in the passed in string.
@@ -513,7 +513,7 @@ namespace UnrealBuildTool
 				: throw new ArgumentException($"'{(int)Class}' is not a valid value for {typeof(UnrealPlatformClass).Name}");
 
 		}
-		static Dictionary<UnrealPlatformClass, Lazy<UnrealTargetPlatform[]>> s_PlatformsInClass = new()
+		static readonly Dictionary<UnrealPlatformClass, Lazy<UnrealTargetPlatform[]>> s_PlatformsInClass = new()
 		{
 			{ UnrealPlatformClass.All, new(UnrealTargetPlatform.GetValidPlatforms) },
 			{ UnrealPlatformClass.Desktop, new(() => [.. UEBuildPlatform.GetPlatformsInGroup(UnrealPlatformGroup.Desktop)]) },
diff --git a/Engine/Source/Programs/UnrealBuildTool/ToolChain/IWYUToolChain.cs b/Engine/Source/Programs/UnrealBuildTool/ToolChain/IWYUToolChain.cs
--- a/Engine/Source/Programs/UnrealBuildTool/ToolChain/IWYUToolChain.cs
+++ b/Engine/Source/Programs/UnrealBuildTool/ToolChain/IWYUToolChain.cs
@@ -11,11 +11,11 @@ namespace UnrealBuildTool
 {
 	class IWYUToolChain : ClangToolChain
 	{
-		private List<string> CrossCompilingArguments = new();
-		private string IWYUMappingFile;
-		private string RelativePathToIWYUDirectory = @"Binaries\ThirdParty\IWYU";
-		private UnrealTargetPlatform TargetPlatform;
-		private ISPCToolChain InnerToolChain;
+		private readonly List<string> CrossCompilingArguments = new();
+		private readonly string IWYUMappingFile;
+		private readonly string RelativePathToIWYUDirectory = @"Binaries\ThirdParty\IWYU";
+		private readonly UnrealTargetPlatform TargetPlatform;
+		private readonly ISPCToolChain InnerToolChain;
 
 		public static void ValidateTarget(TargetRules Target)
 		{
