Make AutoRTFM instrumentation default disabled for modules with `[SupportedTargetTypes(TargetType.Editor)]` (there are none yet, but there will be soon!).

Idea gratuitously stolen (with permission!) from [REDACTED_HANDLE]

diff --git a/Engine/Source/Programs/UnrealBuildTool/Configuration/Rules/ModuleRules.cs b/Engine/Source/Programs/UnrealBuildTool/Configuration/Rules/ModuleRules.cs
--- a/Engine/Source/Programs/UnrealBuildTool/Configuration/Rules/ModuleRules.cs
+++ b/Engine/Source/Programs/UnrealBuildTool/Configuration/Rules/ModuleRules.cs
@@ -1154,10 +1154,31 @@ namespace UnrealBuildTool
 		/// </summary>
 		public bool bAllowConfidentialPlatformDefines { get; set; }
 
-		/// <summary>
-		/// Disables AutoRTFM instrumentation to this module only when AutoRTFMCompiler is enabled
-		/// </summary>
-		public bool bDisableAutoRTFMInstrumentation { get; set; }
+        /// <summary>
+        /// Disables AutoRTFM instrumentation in this module. Only matters when the AutoRTFM compiler is being used.
+		/// Modules with the `SupportedTargetTypes` attribute where the `TargetType`'s are only `Editor` or `Program`
+		/// (or both) will be by default disabled for AutoRTFM instrumentation. Otherwise AutoRTFM will by default be
+		/// enabled.
+        /// </summary>
+        public bool bDisableAutoRTFMInstrumentation
+		{
+			get
+			{
+				if (bDisableAutoRTFMInstrumentationPrivate != null)
+				{
+					return (bool)bDisableAutoRTFMInstrumentationPrivate;
+				}
+
+				// By default disable AutoRTFM for any Modules that do not support Game, Client, or Server
+				IEnumerable<TargetType> supportedTargetTypes = GetType().GetCustomAttributes<SupportedTargetTypesAttribute>().SelectMany(x => x.TargetTypes).Distinct();
+				return supportedTargetTypes.Any() && !supportedTargetTypes.Any(x => x == TargetType.Game || x == TargetType.Client || x == TargetType.Server);
+			}
+			set
+			{
+				bDisableAutoRTFMInstrumentationPrivate = value;
+			}
+		}
+		private bool? bDisableAutoRTFMInstrumentationPrivate;
 
 		/// <summary>
 		/// If `bDisableAutoRTFMInstrumentation` is true, will make the modules `*_API` define contain `AUTORTFM_DISABLE` to let us enable errors for code
