UnrealBuildTool: Begin cleanup of application entrypoint

#rnx

diff --git a/Engine/Source/Programs/Shared/EpicGames.Core/CommandLineAttribute.cs b/Engine/Source/Programs/Shared/EpicGames.Core/CommandLineAttribute.cs
--- a/Engine/Source/Programs/Shared/EpicGames.Core/CommandLineAttribute.cs
+++ b/Engine/Source/Programs/Shared/EpicGames.Core/CommandLineAttribute.cs
@@ -45,7 +45,7 @@ namespace EpicGames.Core
 		/// <summary>
 		/// Description of the operation.
 		/// </summary>
-		public string? Description { get; set; }
+		public string? Description { get; set;}
 
 		/// <summary>
 		/// Constructor
@@ -55,12 +55,9 @@ namespace EpicGames.Core
 		{
 			Prefix = prefix;
 
-			if(prefix != null)
+			if (prefix?.StartsWith('-') == false)
 			{
-				if(!prefix.StartsWith('-'))
-				{
-					throw new Exception("Command-line arguments must begin with a '-' character");
-				}
+				throw new Exception($"Command-line argument '{Prefix}' must begin with a '-' character");
 			}
 		}
 	}
diff --git a/Engine/Source/Programs/UnrealBuildTool/GlobalOptions.cs b/Engine/Source/Programs/UnrealBuildTool/GlobalOptions.cs
new file mode 100644
--- /dev/null
+++ b/Engine/Source/Programs/UnrealBuildTool/GlobalOptions.cs
@@ -0,0 +1,180 @@
+// Copyright Epic Games, Inc. All Rights Reserved.
+
+using System;
+using System.Collections.Generic;
+using EpicGames.Core;
+using UnrealBuildBase;
+
+namespace UnrealBuildTool
+{
+	/// <summary>
+	/// Global options for UBT (any modes)
+	/// </summary>
+	internal record GlobalOptions
+	{
+		/// <summary>
+		/// User asked for help
+		/// </summary>
+		[CommandLine(Prefix = "-Help", Description = "Display this help.")]
+		[CommandLine(Prefix = "-h")]
+		[CommandLine(Prefix = "--help")]
+		public bool GetHelp = false;
+
+		/// <summary>
+		/// The amount of detail to write to the log
+		/// </summary>
+		[CommandLine(Prefix = "-Verbose", Value = "Verbose", Description = "Increase output verbosity")]
+		[CommandLine(Prefix = "-VeryVerbose", Value = "VeryVerbose", Description = "Increase output verbosity more")]
+		public LogEventType LogOutputLevel = LogEventType.Log;
+
+		/// <summary>
+		/// Specifies the path to a log file to write. Note that the default mode (eg. building, generating project files) will create a log file by default if this not specified.
+		/// </summary>
+		[CommandLine(Prefix = "-Log", Description = "Specify a log file location instead of the default Engine/Programs/UnrealBuildTool/Log.txt")]
+		public FileReference? LogFileName = null;
+
+		/// <summary>
+		/// Log all attempts to write to the specified file
+		/// </summary>
+		[CommandLine(Prefix = "-TraceWrites", Description = "Trace writes requested to the specified file")]
+		public FileReference? TraceWrites = null;
+
+		/// <summary>
+		/// Whether to include timestamps in the log
+		/// </summary>
+		[CommandLine(Prefix = "-Timestamps", Description = "Include timestamps in the log")]
+		public bool bLogTimestamps = false;
+
+		/// <summary>
+		/// Whether to format messages in MsBuild format
+		/// </summary>
+		[CommandLine(Prefix = "-FromMsBuild", Description = "Format messages for msbuild")]
+		public bool LogFromMsBuild { get; set; } = false;
+
+		/// <summary>
+		/// Whether or not to suppress warnings of missing SDKs from warnings to LogEventType.Log in UEBuildPlatformSDK.cs 
+		/// </summary>
+		[CommandLine(Prefix = "-SuppressSDKWarnings", Description = "Missing SDKs error verbosity level will be reduced from warning to log")]
+		public bool ShouldSuppressSDKWarnings { get; set; } = false;
+
+		/// <summary>
+		/// Whether to write progress markup in a format that can be parsed by other programs
+		/// </summary>
+		[CommandLine(Prefix = "-Progress", Description = "Write progress messages in a format that can be parsed by other programs")]
+		public bool WriteProgressMarkup { get; set; } = false;
+
+		/// <summary>
+		/// Whether to ignore the mutex
+		/// </summary>
+		[CommandLine(Prefix = "-NoMutex", Description = "Allow more than one instance of the program to run at once")]
+		public bool NoMutex { get; set; } = false;
+
+		/// <summary>
+		/// Whether to wait for the mutex rather than aborting immediately
+		/// </summary>
+		[CommandLine(Prefix = "-WaitMutex", Description = "Wait for another instance to finish and then start, rather than aborting immediately")]
+		public bool WaitMutex { get; set; } = false;
+
+		/// <summary>
+		/// </summary>
+		[CommandLine(Prefix = "-RemoteIni", Description = "Remote tool ini directory")]
+		public string RemoteIni { get; set; } = String.Empty;
+
+		/// <summary>
+		/// The mode to execute
+		/// </summary>
+		[CommandLine("-Mode=", Description = "Tool mode to select. Default tool mode is 'Build'")]
+
+		[CommandLine("-Clean", Value = "Clean", Description = "Clean build products. Equivalent to -Mode=Clean")]
+
+		[CommandLine("-ProjectFiles", Value = "GenerateProjectFiles", Description = "Generate project files based on IDE preference. Equivalent to -Mode=GenerateProjectFiles")]
+		[CommandLine("-ProjectFileFormat=", Value = "GenerateProjectFiles", Description = "Generate project files in specified format. May be used multiple times.")]
+		[CommandLine("-Makefile", Value = "GenerateProjectFiles", Description = "Generate Makefile")]
+		[CommandLine("-CMakefile", Value = "GenerateProjectFiles", Description = "Generate project files for CMake")]
+		[CommandLine("-QMakefile", Value = "GenerateProjectFiles", Description = "Generate project files for QMake")]
+		[CommandLine("-KDevelopfile", Value = "GenerateProjectFiles", Description = "Generate project files for KDevelop")]
+		[CommandLine("-CodeliteFiles", Value = "GenerateProjectFiles", Description = "Generate project files for Codelite")]
+		[CommandLine("-XCodeProjectFiles", Value = "GenerateProjectFiles", Description = "Generate project files for XCode")]
+		[CommandLine("-EddieProjectFiles", Value = "GenerateProjectFiles", Description = "Generate project files for Eddie")]
+		[CommandLine("-VSCode", Value = "GenerateProjectFiles", Description = "Generate project files for Visual Studio Code")]
+		[CommandLine("-VSMac", Value = "GenerateProjectFiles", Description = "Generate project files for Visual Studio Mac")]
+		[CommandLine("-CLion", Value = "GenerateProjectFiles", Description = "Generate project files for CLion")]
+		[CommandLine("-Rider", Value = "GenerateProjectFiles", Description = "Generate project files for Rider")]
+		[CommandLine("-AndroidStudio", Value = "GenerateProjectFiles", Description = "Generate project files for Android Studio")]
+#if __VPROJECT_AVAILABLE__
+		[CommandLine("-VProject", Value = "GenerateProjectFiles")]
+#endif
+		public string? Mode { get; set; } = null;
+
+		// The following Log settings exists in this location because, at the time of writing, EpicGames.Core does
+		// not have access to XmlConfigFileAttribute.
+
+		/// <summary>
+		/// Whether to backup an existing log file, rather than overwriting it.
+		/// </summary>
+		[XmlConfigFile(Category = "Log", Name = "bBackupLogFiles")]
+		public bool BackupLogFiles { get; set; } = Log.BackupLogFiles;
+
+		/// <summary>
+		/// The number of log file backups to preserve. Older backups will be deleted.
+		/// </summary>
+		[XmlConfigFile(Category = "Log")]
+		public int LogFileBackupCount { get; set; } = Log.LogFileBackupCount;
+
+		/// <summary>
+		/// If set and tool execution was successful, then display an unreal build tool script execution timeline summary.
+		/// If unset or the tool execution failed, print the same information silently to the log.
+		/// </summary>
+		[XmlConfigFile(Category = "BuildConfiguration", Name = "bShowTimeline")]
+		public bool ShowTimeline { get; set; } = Unreal.IsBuildMachine();
+
+		/// <summary>
+		/// If set TMP\TEMP will be overidden to this directory, each process will create a unique subdirectory in this folder.
+		/// </summary>
+		[XmlConfigFile(Category = "BuildConfiguration")]
+		public string? TempDirectory { get; set; } = null;
+
+		/// <summary>
+		/// If set the application temp directory will be deleted on exit, only when running with a single instance mutex.
+		/// </summary>
+		[XmlConfigFile(Category = "BuildConfiguration", Name = "bDeleteTempDirectory")]
+		public bool DeleteTempDirectory { get; set; } = false;
+
+		/// <summary>
+		/// Providers to load opt-in telemetry connection information from ini. If unset, or the provider categories do not contain connection info, no telemetry will be sent.
+		/// </summary>
+		[XmlConfigFile(Category = "Telemetry", Name = "Providers")]
+		public string[] TelemetryProviders { get; set; } = [];
+
+		/// <summary>
+		/// Additional command line providers to load opt-in telemetry connection information from ini.
+		/// </summary>
+		[CommandLine(Prefix = "-TelemetryProvider", Description = "List of ini providers for telemetry", ListSeparator = '+')]
+		public List<string> CmdTelemetryProviders { get; set; } = [];
+
+		/// <summary>
+		/// Session identifier for this run of UBT, if unset defaults to a random Guid
+		/// </summary>
+		[CommandLine(Prefix = "-Session", Description = "Session identifier for this run of UBT, if unset defaults to a random Guid")]
+		public string? TelemetrySession { get; set; } = null;
+
+		/// <summary>
+		/// Whether telemetry/horde should use ipv6 or not. If network does not support ipv6 then using ipv6 can cause multi-second stalls
+		/// </summary>
+		[XmlConfigFile(Category = "BuildConfiguration", Name = "bUseIpv6")]
+		public bool UseIpv6 { get; set; } = true;
+
+		/// <summary>
+		/// Initialize the options with the given command line arguments
+		/// </summary>
+		/// <param name="Arguments"></param>
+		public GlobalOptions(CommandLineArguments Arguments)
+		{
+			Arguments.ApplyTo(this);
+			if (!String.IsNullOrEmpty(RemoteIni))
+			{
+				UnrealBuildTool.SetRemoteIniPath(RemoteIni);
+			}
+		}
+	}
+}
diff --git a/Engine/Source/Programs/UnrealBuildTool/UnrealBuildTool.cs b/Engine/Source/Programs/UnrealBuildTool/UnrealBuildTool.cs
--- a/Engine/Source/Programs/UnrealBuildTool/UnrealBuildTool.cs
+++ b/Engine/Source/Programs/UnrealBuildTool/UnrealBuildTool.cs
@@ -7,7 +7,9 @@ using System.IO;
 using System.Linq;
 using System.Reflection;
 using System.Text;
+using System.Text.RegularExpressions;
 using System.Threading.Tasks;
+using System.Xml.Linq;
 using EpicGames.Core;
 using Microsoft.Extensions.Logging;
 using OpenTracing.Util;
@@ -48,60 +50,6 @@ namespace UnrealBuildTool
 		/// </summary>
 		public static System.Collections.IDictionary? InitialEnvironment;
 
-		/// <summary>
-		/// The full name of the Engine/Source directory
-		/// </summary>
-		[Obsolete("Replace with Unreal.EngineSourceDirectory")]
-		public static DirectoryReference EngineSourceDirectory => Unreal.EngineSourceDirectory;
-
-		/// <summary>
-		/// Writable engine directory. Uses the user's settings folder for installed builds.
-		/// </summary>
-		[Obsolete("Replace with Unreal.WritableEngineDirectory")]
-		public static DirectoryReference WritableEngineDirectory => Unreal.WritableEngineDirectory;
-
-		/// <summary>
-		/// The engine programs directory
-		/// </summary>
-		[Obsolete("Replace with Unreal.EngineProgramSavedDirectory")]
-		public static DirectoryReference EngineProgramSavedDirectory => Unreal.EngineProgramSavedDirectory;
-
-		/// <summary>
-		/// The original root directory that was used to compile the installed engine
-		/// Used to remap source code paths when debugging.
-		/// </summary>
-		[Obsolete("Deprecated in UE5.6 - use Unreal.OriginalCompilationRootDirectory")]
-		public static DirectoryReference OriginalCompilationRootDirectory => Unreal.OriginalCompilationRootDirectory;
-
-		/// <summary>
-		/// Returns true if UnrealBuildTool is running using an installed project (ie. a mod kit)
-		/// </summary>
-		/// <returns>True if running using an installed project</returns>
-		[Obsolete("Deprecated in UE5.5 - use Unreal.IsProjectInstalled")]
-		public static bool IsProjectInstalled() => Unreal.IsProjectInstalled();
-
-		/// <summary>
-		/// Gets the installed project file
-		/// </summary>
-		/// <returns>Location of the installed project file</returns>
-		[Obsolete("Deprecated in UE5.5 - use Unreal.GetInstalledProjectFile")]
-		public static FileReference? GetInstalledProjectFile() => Unreal.GetInstalledProjectFile();
-
-		/// <summary>
-		/// Checks whether the given file is under an installed directory, and should not be overridden
-		/// </summary>
-		/// <param name="File">File to test</param>
-		/// <returns>True if the file is part of the installed distribution, false otherwise</returns>
-		[Obsolete("Deprecated in UE5.5 - use Unreal.IsFileInstalled")]
-		public static bool IsFileInstalled(FileReference File) => Unreal.IsFileInstalled(File);
-
-		/// <summary>
-		/// Gets the absolute path to the UBT assembly.
-		/// </summary>
-		/// <returns>A string containing the path to the UBT assembly.</returns>
-		[Obsolete("Deprecated in UE5.1 - use Unreal.DotnetPath and Unreal.UnrealBuildToolDllPath")]
-		public static FileReference GetUBTPath() => Unreal.UnrealBuildToolPath;
-
 		/// <summary>
 		/// The Unreal remote tool ini directory.  This should be valid if compiling using a remote server
 		/// </summary>
@@ -119,248 +67,76 @@ namespace UnrealBuildTool
 		/// </summary>
 		private static string? _RemoteIniPath = null;
 
-		/// <summary>
-		/// Global options for UBT (any modes)
-		/// </summary>
-		class GlobalOptions
-		{
-			/// <summary>
-			/// User asked for help
-			/// </summary>
-			[CommandLine(Prefix = "-Help", Description = "Display this help.")]
-			[CommandLine(Prefix = "-h")]
-			[CommandLine(Prefix = "--help")]
-			public bool bGetHelp = false;
-
-			/// <summary>
-			/// The amount of detail to write to the log
-			/// </summary>
-			[CommandLine(Prefix = "-Verbose", Value = "Verbose", Description = "Increase output verbosity")]
-			[CommandLine(Prefix = "-VeryVerbose", Value = "VeryVerbose", Description = "Increase output verbosity more")]
-			public LogEventType LogOutputLevel = LogEventType.Log;
-
-			/// <summary>
-			/// Specifies the path to a log file to write. Note that the default mode (eg. building, generating project files) will create a log file by default if this not specified.
-			/// </summary>
-			[CommandLine(Prefix = "-Log", Description = "Specify a log file location instead of the default Engine/Programs/UnrealBuildTool/Log.txt")]
-			public FileReference? LogFileName = null;
-
-			/// <summary>
-			/// Log all attempts to write to the specified file
-			/// </summary>
-			[CommandLine(Prefix = "-TraceWrites", Description = "Trace writes requested to the specified file")]
-			public FileReference? TraceWrites = null;
-
-			/// <summary>
-			/// Whether to include timestamps in the log
-			/// </summary>
-			[CommandLine(Prefix = "-Timestamps", Description = "Include timestamps in the log")]
-			public bool bLogTimestamps = false;
-
-			/// <summary>
-			/// Whether to format messages in MsBuild format
-			/// </summary>
-			[CommandLine(Prefix = "-FromMsBuild", Description = "Format messages for msbuild")]
-			public bool bLogFromMsBuild = false;
-
-			/// <summary>
-			/// Whether or not to suppress warnings of missing SDKs from warnings to LogEventType.Log in UEBuildPlatformSDK.cs 
-			/// </summary>
-			[CommandLine(Prefix = "-SuppressSDKWarnings", Description = "Missing SDKs error verbosity level will be reduced from warning to log")]
-			public bool bShouldSuppressSDKWarnings = false;
-
-			/// <summary>
-			/// Whether to write progress markup in a format that can be parsed by other programs
-			/// </summary>
-			[CommandLine(Prefix = "-Progress", Description = "Write progress messages in a format that can be parsed by other programs")]
-			public bool bWriteProgressMarkup = false;
-
-			/// <summary>
-			/// Whether to ignore the mutex
-			/// </summary>
-			[CommandLine(Prefix = "-NoMutex", Description = "Allow more than one instance of the program to run at once")]
-			public bool bNoMutex = false;
-
-			/// <summary>
-			/// Whether to wait for the mutex rather than aborting immediately
-			/// </summary>
-			[CommandLine(Prefix = "-WaitMutex", Description = "Wait for another instance to finish and then start, rather than aborting immediately")]
-			public bool bWaitMutex = false;
-
-			/// <summary>
-			/// </summary>
-			[CommandLine(Prefix = "-RemoteIni", Description = "Remote tool ini directory")]
-			public string RemoteIni = "";
-
-			/// <summary>
-			/// The mode to execute
-			/// </summary>
-			[CommandLine("-Mode=")] // description handling is special-cased in PrintUsage()
-
-			[CommandLine("-Clean", Value = "Clean", Description = "Clean build products. Equivalent to -Mode=Clean")]
-
-			[CommandLine("-ProjectFiles", Value = "GenerateProjectFiles", Description = "Generate project files based on IDE preference. Equivalent to -Mode=GenerateProjectFiles")]
-			[CommandLine("-ProjectFileFormat=", Value = "GenerateProjectFiles", Description = "Generate project files in specified format. May be used multiple times.")]
-			[CommandLine("-Makefile", Value = "GenerateProjectFiles", Description = "Generate Makefile")]
-			[CommandLine("-CMakefile", Value = "GenerateProjectFiles", Description = "Generate project files for CMake")]
-			[CommandLine("-QMakefile", Value = "GenerateProjectFiles", Description = "Generate project files for QMake")]
-			[CommandLine("-KDevelopfile", Value = "GenerateProjectFiles", Description = "Generate project files for KDevelop")]
-			[CommandLine("-CodeliteFiles", Value = "GenerateProjectFiles", Description = "Generate project files for Codelite")]
-			[CommandLine("-XCodeProjectFiles", Value = "GenerateProjectFiles", Description = "Generate project files for XCode")]
-			[CommandLine("-EddieProjectFiles", Value = "GenerateProjectFiles", Description = "Generate project files for Eddie")]
-			[CommandLine("-VSCode", Value = "GenerateProjectFiles", Description = "Generate project files for Visual Studio Code")]
-			[CommandLine("-VSMac", Value = "GenerateProjectFiles", Description = "Generate project files for Visual Studio Mac")]
-			[CommandLine("-CLion", Value = "GenerateProjectFiles", Description = "Generate project files for CLion")]
-			[CommandLine("-Rider", Value = "GenerateProjectFiles", Description = "Generate project files for Rider")]
-			[CommandLine("-AndroidStudio", Value = "GenerateProjectFiles", Description = "Generate project files for Android Studio")]
-#if __VPROJECT_AVAILABLE__
-			[CommandLine("-VProject", Value = "GenerateProjectFiles")]
-#endif
-			public string? Mode = null;
-
-			// The following Log settings exists in this location because, at the time of writing, EpicGames.Core does
-			// not have access to XmlConfigFileAttribute.
-
-			/// <summary>
-			/// Whether to backup an existing log file, rather than overwriting it.
-			/// </summary>
-			[XmlConfigFile(Category = "Log")]
-			public bool bBackupLogFiles = Log.BackupLogFiles;
-
-			/// <summary>
-			/// The number of log file backups to preserve. Older backups will be deleted.
-			/// </summary>
-			[XmlConfigFile(Category = "Log")]
-			public int LogFileBackupCount = Log.LogFileBackupCount;
-
-			/// <summary>
-			/// If set and tool execution was successful, then display an unreal build tool script execution timeline summary.
-			/// If unset or the tool execution failed, print the same information silently to the log.
-			/// </summary>
-			[XmlConfigFile(Category = "BuildConfiguration")]
-			public bool bShowTimeline = Unreal.IsBuildMachine();
-
-			/// <summary>
-			/// If set TMP\TEMP will be overidden to this directory, each process will create a unique subdirectory in this folder.
-			/// </summary>
-			[XmlConfigFile(Category = "BuildConfiguration")]
-			public string? TempDirectory = null;
-
-			/// <summary>
-			/// If set the application temp directory will be deleted on exit, only when running with a single instance mutex.
-			/// </summary>
-			[XmlConfigFile(Category = "BuildConfiguration")]
-			public bool bDeleteTempDirectory = false;
-
-			/// <summary>
-			/// Providers to load opt-in telemetry connection information from ini. If unset, or the provider categories do not contain connection info, no telemetry will be sent.
-			/// </summary>
-			[XmlConfigFile(Category = "Telemetry", Name = "Providers")]
-			public string[] TelemetryProviders = Array.Empty<string>();
-
-			/// <summary>
-			/// Additional command line providers to load opt-in telemetry connection information from ini.
-			/// </summary>
-			[CommandLine(Prefix = "-TelemetryProvider", Description = "List of ini providers for telemetry", ListSeparator = '+')]
-			public List<string> CmdTelemetryProviders = new();
-
-			/// <summary>
-			/// Session identifier for this run of UBT, if unset defaults to a random Guid
-			/// </summary>
-			[CommandLine(Prefix = "-Session", Description = "Session identifier for this run of UBT, if unset defaults to a random Guid")]
-			public string? TelemetrySession = null;
-
-			/// <summary>
-			/// Whether telemetry/horde should use ipv6 or not. If network does not support ipv6 then using ipv6 can cause multi-second stalls
-			/// </summary>
-			[XmlConfigFile(Category = "BuildConfiguration")]
-			public bool bUseIpv6 { get; set; } = true;
-
-			/// <summary>
-			/// Initialize the options with the given command line arguments
-			/// </summary>
-			/// <param name="Arguments"></param>
-			public GlobalOptions(CommandLineArguments Arguments)
-			{
-				Arguments.ApplyTo(this);
-				if (!String.IsNullOrEmpty(RemoteIni))
-				{
-					UnrealBuildTool.SetRemoteIniPath(RemoteIni);
-				}
-			}
-		}
-
 		/// <summary>
 		/// Get all the valid Modes
 		/// </summary>
 		/// <returns></returns>
 		private static Dictionary<string, Type> GetModes()
 		{
-			Dictionary<string, Type> ModeNameToType = new Dictionary<string, Type>(StringComparer.OrdinalIgnoreCase);
-			foreach (Type Type in Assembly.GetExecutingAssembly().GetTypes())
+			Dictionary<string, Type> ModeNameToType = new(StringComparer.OrdinalIgnoreCase);
+			foreach (Type Type in Assembly.GetExecutingAssembly().GetTypes().Where(x => x.IsClass && !x.IsAbstract && x.IsSubclassOf(typeof(ToolMode))))
 			{
-				if (Type.IsClass && !Type.IsAbstract && Type.IsSubclassOf(typeof(ToolMode)))
-				{
-					ToolModeAttribute? Attribute = Type.GetCustomAttribute<ToolModeAttribute>();
-					if (Attribute == null)
-					{
-						throw new BuildException("Class '{0}' should have a ToolModeAttribute", Type.Name);
-					}
-					ModeNameToType.Add(Attribute.Name, Type);
-				}
+				ToolModeAttribute? Attribute = Type.GetCustomAttribute<ToolModeAttribute>()
+					?? throw new BuildLogEventException("Class '{Name}' should have a ToolModeAttribute", Type.Name);
+				ModeNameToType.Add(Attribute.Name, Type);
 			}
 			return ModeNameToType;
 		}
-		public static readonly Dictionary<string, Type> ModeNameToType = GetModes();
 
 		/// <summary>
 		/// Print (incomplete) usage information
 		/// </summary>
 		private static void PrintUsage()
 		{
-			Console.WriteLine("Global options:");
-			int LongestPrefix = 0;
-			foreach (FieldInfo Info in typeof(GlobalOptions).GetFields())
+			IEnumerable<CommandLineAttribute> attributes = [
+				.. typeof(GlobalOptions).GetProperties().SelectMany(x => x.GetCustomAttributes<CommandLineAttribute>()),
+				.. typeof(GlobalOptions).GetFields().SelectMany(x => x.GetCustomAttributes<CommandLineAttribute>())
+			];
+			if (!attributes.Any())
 			{
-				foreach (CommandLineAttribute Att in Info.GetCustomAttributes<CommandLineAttribute>())
+				throw new BuildLogEventException("Unable to load properties and fields from {Type}", typeof(GlobalOptions));
+			}
+			int longestPrefix = attributes.Select(x => x.Prefix?.Length ?? 0).Max();
+
+			Console.WriteLine("Global Options:");
+			foreach (CommandLineAttribute Att in attributes)
+			{
+				if (Att.Prefix != null && Att.Description != null)
 				{
-					if (Att.Prefix != null && Att.Description != null)
-					{
-						LongestPrefix = Att.Prefix.Length > LongestPrefix ? Att.Prefix.Length : LongestPrefix;
-					}
+					Console.WriteLine($" {Att.Prefix.PadRight(longestPrefix)} :  {Att.Description}");
 				}
 			}
 
-			foreach (FieldInfo Info in typeof(GlobalOptions).GetFields())
+			Dictionary<string, Type> modes = GetModes();
+			if (modes.Count == 0)
 			{
-				foreach (CommandLineAttribute Att in Info.GetCustomAttributes<CommandLineAttribute>())
-				{
-					if (Att.Prefix != null && Att.Description != null)
-					{
-						Console.WriteLine($"  {Att.Prefix.PadRight(LongestPrefix)} :  {Att.Description}");
-					}
+				throw new BuildLogEventException("No valid subclasses of {Type} found", typeof(ToolMode));
+			}
 
-					// special case for Mode
-					if (String.Equals(Att.Prefix, "-Mode=", StringComparison.Ordinal))
-					{
-						Console.WriteLine($"  {Att.Prefix!.PadRight(LongestPrefix)} :  Select tool mode. One of the following (default tool mode is \"Build\"):");
-						string Indent = "".PadRight(LongestPrefix + 8);
-						string Line = Indent;
-						IOrderedEnumerable<string> SortedModeNames = ModeNameToType.Keys.ToList().OrderBy(Name => Name);
-						foreach (string ModeName in SortedModeNames.SkipLast(1))
-						{
-							Line += $"{ModeName}, ";
-							if (Line.Length > 110)
-							{
-								Console.WriteLine(Line);
-								Line = Indent;
-							}
-						}
-						Line += SortedModeNames.Last();
-						Console.WriteLine(Line);
-					}
-				}
+			// Load tool mode summary from xml documentation if available
+			XDocument? doc = null;
+			try
+			{
+				FileReference xmlDocPath = new FileReference(Assembly.GetExecutingAssembly().Location).ChangeExtension(".xml");
+				doc = FileReference.Exists(xmlDocPath) ? XDocument.Load(xmlDocPath.FullName) : null;
+			}
+			catch (Exception ex)
+			{
+				throw new BuildLogEventException(ex, "Documentation malformed for {Assembly}", Assembly.GetExecutingAssembly().Location);
 			}
+
+			longestPrefix = modes.Select(x => x.Key.Length).Max();
+			Console.WriteLine();
+			Console.WriteLine("Available Tool Modes:");
+			foreach ((string name, Type type) in modes)
+			{
+				string summary = doc?.Descendants("member")
+					.FirstOrDefault(x => x.Attribute("name")?.Value == $"T:{type.FullName}")?
+					.Element("summary")?.Value.Trim()
+					?? "Summary unavailable";
+				Console.WriteLine($" {name.PadRight(longestPrefix)} : {Regex.Replace(summary, @"\s+", " ")}");
+			}
+			Console.WriteLine();
 		}
 
 		/// <summary>
@@ -371,35 +147,35 @@ namespace UnrealBuildTool
 		private static string[] GetExtraArgsFromEnvVar()
 		{
 			string? extraArgs = Environment.GetEnvironmentVariable("UBT_EXTRA_ARGS");
-			return String.IsNullOrEmpty(extraArgs) ? Array.Empty<string>() : CommandLineArguments.Split(extraArgs);
+			return String.IsNullOrEmpty(extraArgs) ? [] : CommandLineArguments.Split(extraArgs);
 		}
 
-        /// <summary>
-        /// Event handler for the Console.CancelKeyPress event
-        /// </summary>
-        /// <param name="sender"></param>
-        /// <param name="e"></param>
-        private static async void CancelKeyPressAsync(object? sender, ConsoleCancelEventArgs e)
-        {
+		/// <summary>
+		/// Event handler for the Console.CancelKeyPress event
+		/// </summary>
+		/// <param name="sender"></param>
+		/// <param name="e"></param>
+		private static async void CancelKeyPressAsync(object? sender, ConsoleCancelEventArgs e)
+		{
 			Console.CancelKeyPress -= CancelKeyPressAsync;
 			Console.WriteLine($"UnrealBuildTool: Ctrl-{(e.SpecialKey == ConsoleSpecialKey.ControlC ? "C" : "Break")} pressed. Exiting...");
 
-            // Delay a few seconds to allow for the process to exit normally
-            await Task.Delay(2000);
+			// Delay a few seconds to allow for the process to exit normally
+			await Task.Delay(2000);
 
 			// While the Ctrl-C handler fixes most instances of a zombie process, we still need to 
 			// force an exit from the process to handle _all_ cases.  Ctrl-C should not be a regular event! 
 			// Note: this could be a dotnet (6.0.302) on macOS issue.  Recheck with next release if this is still required.
 			Environment.Exit(-1);
-        }
-
-        /// <summary>
-        /// Main entry point. Parses any global options and initializes the logging system, then invokes the appropriate command.
-        /// NB: That the entry point is deliberately NOT async, since we have a single-instance mutex that cannot be disposed from a different thread.
-        /// </summary>
-        /// <param name="ArgumentsArray">Command line arguments</param>
-        /// <returns>Zero on success, non-zero on error</returns>
-        private static int Main(string[] ArgumentsArray)
+		}
+
+		/// <summary>
+		/// Main entry point. Parses any global options and initializes the logging system, then invokes the appropriate command.
+		/// NB: That the entry point is deliberately NOT async, since we have a single-instance mutex that cannot be disposed from a different thread.
+		/// </summary>
+		/// <param name="ArgumentsArray">Command line arguments</param>
+		/// <returns>Zero on success, non-zero on error</returns>
+		private static int Main(string[] ArgumentsArray)
 		{
 			ToolModeOptions ModeOptions = ToolModeOptions.None;
 			FileReference? RunFile = null;
@@ -414,12 +190,12 @@ namespace UnrealBuildTool
 			// By putting this in, the Ctrl-C may not be handled immediately, but it shouldn't leave a blocking zombie process
 			if (OperatingSystem.IsMacOS())
 			{
-                Console.CancelKeyPress += CancelKeyPressAsync;
-            }
+				Console.CancelKeyPress += CancelKeyPressAsync;
+			}
 
 			// If command line contains -Session it is a recursive action and we don't want to report this trace to the default trace channel
 			string? traceChannel = null;
-			string? traceName =null;
+			string? traceName = null;
 			string? traceFile = null;
 			if (ArgumentsArray.FirstOrDefault(s => s.Contains("-Session=")) == null)
 			{
@@ -470,20 +246,20 @@ namespace UnrealBuildTool
 					ArgumentsArray.Length == 0
 
 					// Print usage if the user asks for help
-					|| Options.bGetHelp
+					|| Options.GetHelp
 					)
 				{
 					PrintUsage();
-					return Options.bGetHelp ? 0 : 1;
+					return Options.GetHelp ? (int)CompilationResult.Succeeded : (int)CompilationResult.Canceled;
 				}
 
 				// Configure the log system
 				Log.OutputLevel = Options.LogOutputLevel;
 				Log.IncludeTimestamps = Options.bLogTimestamps;
-				Log.IncludeProgramNameWithSeverityPrefix = Options.bLogFromMsBuild;
+				Log.IncludeProgramNameWithSeverityPrefix = Options.LogFromMsBuild;
 
 				// Reducing SDK warning events in the log to LogEventType.Log
-				if (Options.bShouldSuppressSDKWarnings)
+				if (Options.ShouldSuppressSDKWarnings)
 				{
 					UEBuildPlatformSDK.bSuppressSDKWarnings = true;
 				}
@@ -502,7 +278,7 @@ namespace UnrealBuildTool
 				}
 
 				// Configure the progress writer
-				ProgressWriter.bWriteMarkup = Options.bWriteProgressMarkup;
+				ProgressWriter.bWriteMarkup = Options.WriteProgressMarkup;
 
 				// Ensure we can resolve any external assemblies that are not in the same folder as our assembly.
 				AssemblyUtils.InstallAssemblyResolver(Path.GetDirectoryName(Assembly.GetEntryAssembly()!.GetOriginalLocation())!);
@@ -523,11 +299,10 @@ namespace UnrealBuildTool
 				if (Options.Mode != null)
 				{
 					// Try to get the correct mode
+					IReadOnlyDictionary<string, Type> ModeNameToType = GetModes();
 					if (!ModeNameToType.TryGetValue(Options.Mode, out ModeType))
 					{
-						List<string> ModuleNameList = ModeNameToType.Keys.ToList();
-						ModuleNameList.Sort(StringComparer.OrdinalIgnoreCase);
-						Logger.LogError("No mode named '{Name}'. Available modes are:\n  {ModeList}", Options.Mode, String.Join("\n  ", ModuleNameList));
+						Logger.LogError("No mode named '{Name}'. Available modes are:\n  {ModeList}", Options.Mode, String.Join("\n  ", ModeNameToType.Keys.Order()));
 						return 1;
 					}
 				}
@@ -560,7 +335,7 @@ namespace UnrealBuildTool
 					XmlConfig.ApplyTo(Options);
 				}
 
-				if (!Options.bUseIpv6)
+				if (!Options.UseIpv6)
 				{
 					HttpMessageHandlers.UseIpv6DnsResolving = false;
 				}
@@ -574,9 +349,9 @@ namespace UnrealBuildTool
 					}
 				}
 
-				Log.BackupLogFiles = Options.bBackupLogFiles;
+				Log.BackupLogFiles = Options.BackupLogFiles;
 				Log.LogFileBackupCount = Options.LogFileBackupCount;
-				bShowTimeline = Options.bShowTimeline;
+				bShowTimeline = Options.ShowTimeline;
 
 				// Add the log writer if requested. When building a target, we'll create the writer for the default log file later.
 				if (Options.LogFileName != null)
@@ -647,7 +422,7 @@ namespace UnrealBuildTool
 						Environment.SetEnvironmentVariable("TEMP", OverrideTempDirectory.FullName);
 
 						// Deleting the directory is only safe in single instance mode, and only if requested
-						if ((ModeOptions & ToolModeOptions.SingleInstance) != 0 && !Options.bNoMutex && Options.bDeleteTempDirectory)
+						if ((ModeOptions & ToolModeOptions.SingleInstance) != 0 && !Options.NoMutex && Options.DeleteTempDirectory)
 						{
 							Logger.LogDebug("Temp directory '{Path}' will be deleted on exit", OverrideTempDirectory);
 							TempDirectory = OverrideTempDirectory;
@@ -659,12 +434,12 @@ namespace UnrealBuildTool
 				}
 
 				// Acquire a lock for this branch
-				if ((ModeOptions & ToolModeOptions.SingleInstance) != 0 && !Options.bNoMutex)
+				if ((ModeOptions & ToolModeOptions.SingleInstance) != 0 && !Options.NoMutex)
 				{
 					using (GlobalTracer.Instance.BuildSpan("SingleInstanceMutex.Acquire()").StartActive())
 					{
 						string MutexName = GlobalSingleInstanceMutex.GetUniqueMutexForPath("UnrealBuildTool_Mutex", FileReference.FromString(Assembly.GetExecutingAssembly().Location));
-						Mutex = new GlobalSingleInstanceMutex(MutexName, Options.bWaitMutex);
+						Mutex = new GlobalSingleInstanceMutex(MutexName, Options.WaitMutex);
 					}
 				}
 
@@ -748,10 +523,10 @@ namespace UnrealBuildTool
 				Timeline.Stop();
 				bool bShowExecutionTime = ModeOptions.HasFlag(ToolModeOptions.ShowExecutionTime);
 				LogLevel ExecutionLogLevel = bShowExecutionTime ? LogLevel.Information : LogLevel.Debug;
-				LogLevel TimelineLogLevel = bShowExecutionTime && bShowTimeline && ApplicationResult == CompilationResult.Succeeded ?  LogLevel.Information : LogLevel.Debug;
-				
+				LogLevel TimelineLogLevel = bShowExecutionTime && bShowTimeline && ApplicationResult == CompilationResult.Succeeded ? LogLevel.Information : LogLevel.Debug;
+
 				Timeline.Print(TimeSpan.FromMilliseconds(100.0), TimeSpan.FromMilliseconds(200.0), /*MaxDepth*/ 4, TimelineLogLevel, Logger);
-				
+
 				Logger.Log(ExecutionLogLevel, "");
 				if (ApplicationResult == CompilationResult.Succeeded)
 				{
@@ -801,4 +576,3 @@ namespace UnrealBuildTool
 		}
 	}
 }
-
