Don't pass /c when using clang-cl with bPreprocessOnly because it rejects it as an unused argument.

diff --git a/Engine/Source/Programs/UnrealBuildTool/Platform/Windows/VCToolChain.cs b/Engine/Source/Programs/UnrealBuildTool/Platform/Windows/VCToolChain.cs
--- a/Engine/Source/Programs/UnrealBuildTool/Platform/Windows/VCToolChain.cs
+++ b/Engine/Source/Programs/UnrealBuildTool/Platform/Windows/VCToolChain.cs
@@ -658,7 +658,10 @@ namespace UnrealBuildTool
 			}
 
 			// Compile into an .obj file, and skip linking.
-			Arguments.Add("/c");
+			if (!Target.WindowsPlatform.Compiler.IsClang() || !CompileEnvironment.bPreprocessOnly)
+			{
+				Arguments.Add("/c");
+			}
 
 			// Put symbols into different sections so the linker can remove them.
 			if (Target.WindowsPlatform.bOptimizeGlobalData)
