UnrealBuildTool: Support a release branch only BuildConfiguration

#rnx

diff --git a/Engine/Source/Programs/Shared/EpicGames.Build/Unreal.cs b/Engine/Source/Programs/Shared/EpicGames.Build/Unreal.cs
--- a/Engine/Source/Programs/Shared/EpicGames.Build/Unreal.cs
+++ b/Engine/Source/Programs/Shared/EpicGames.Build/Unreal.cs
@@ -241,6 +241,19 @@ namespace UnrealBuildBase
 		public static bool IsEngineInstalled() => _IsEngineInstalled.Value;
 		private static readonly Lazy<bool> _IsEngineInstalled = new(() => FileReference.Exists(FileReference.Combine(EngineDirectory, "Build", "InstalledBuild.txt")));
 
+		/// <summary>
+		/// Returns true of the application is located in a release branch from source control and not for an installed build
+		/// </summary>
+		/// <returns>True if running in a release branch</returns>
+		public static bool IsReleaseBranch() => __IsReleaseBranch.Value;
+		private static readonly Lazy<bool> __IsReleaseBranch = new(() =>
+		{
+			return !IsEngineInstalled()
+				&& JsonObject.TryRead(FileReference.Combine(EngineDirectory, "Build", "Build.version"), out JsonObject? jsonObject)
+				&& jsonObject.TryGetStringField("BranchName", out string? branchName)
+				&& branchName.Contains("Release", StringComparison.OrdinalIgnoreCase);
+		});
+
 		/// <summary>
 		/// If we are running with an installed project, specifies the path to it
 		/// </summary>
diff --git a/Engine/Source/Programs/UnrealBuildTool/System/XmlConfig.cs b/Engine/Source/Programs/UnrealBuildTool/System/XmlConfig.cs
--- a/Engine/Source/Programs/UnrealBuildTool/System/XmlConfig.cs
+++ b/Engine/Source/Programs/UnrealBuildTool/System/XmlConfig.cs
@@ -308,6 +308,11 @@ namespace UnrealBuildTool
 					// Check for the engine config file under /Engine/Programs/NotForLicensees/UnrealBuildTool
 					configs.Add(new(new InputFile(FileReference.Combine(Unreal.EngineDirectory, "Restricted", "NotForLicensees", "Programs", "UnrealBuildTool", "BuildConfiguration.xml"), "Engine (NotForLicensees)"), false));
 
+					if (Unreal.IsReleaseBranch())
+					{
+						configs.Add(new(new InputFile(FileReference.Combine(Unreal.EngineDirectory, "Restricted", "NotForLicensees", "Programs", "UnrealBuildTool", "ReleaseBuildConfiguration.xml"), "Engine (NotForLicensees Release)"), false));
+					}
+
 					// Check for the engine user config file under /Engine/Saved/UnrealBuildTool
 					configs.Add(new(new InputFile(FileReference.Combine(Unreal.EngineDirectory, "Saved", "UnrealBuildTool", "BuildConfiguration.xml"), "Engine (Saved)"), true));
 				}
diff --git a/Engine/Source/Programs/UnrealBuildTool/UnrealBuildTool.csproj b/Engine/Source/Programs/UnrealBuildTool/UnrealBuildTool.csproj
--- a/Engine/Source/Programs/UnrealBuildTool/UnrealBuildTool.csproj
+++ b/Engine/Source/Programs/UnrealBuildTool/UnrealBuildTool.csproj
@@ -123,6 +123,7 @@
   <!-- Various BuildConfiguration files -->
   <ItemGroup>
       <None Condition="Exists('..\..\..\Restricted\NotForLicensees\Programs\UnrealBuildTool\BuildConfiguration.xml')" Include="..\..\..\Restricted\NotForLicensees\Programs\UnrealBuildTool\BuildConfiguration.xml"><Link>BuildConfiguration\BuildConfiguration.xml (Engine NotForLicensees)</Link></None>
+      <None Condition="Exists('..\..\..\Restricted\NotForLicensees\Programs\UnrealBuildTool\BuildConfiguration.xml')" Include="..\..\..\Restricted\NotForLicensees\Programs\UnrealBuildTool\BuildConfiguration.xml"><Link>BuildConfiguration\ReleaseBuildConfiguration.xml (Engine NotForLicensees Release)</Link></None>
       <None Condition="Exists('..\..\..\Saved\UnrealBuildTool\BuildConfiguration.xml')" Include="..\..\..\Saved\UnrealBuildTool\BuildConfiguration.xml"><Link>BuildConfiguration\BuildConfiguration.xml (Engine Saved)</Link></None>
       <None Condition="Exists('$(ProgramData)\Unreal Engine\UnrealBuildTool\BuildConfiguration.xml')" Include="$(ProgramData)\Unreal Engine\UnrealBuildTool\BuildConfiguration.xml"><Link>BuildConfiguration\BuildConfiguration.xml (Global ProgramData)</Link></None>
       <None Condition="Exists('$(AppData)\Unreal Engine\UnrealBuildTool\BuildConfiguration.xml')" Include="$(AppData)\Unreal Engine\UnrealBuildTool\BuildConfiguration.xml"><Link>BuildConfiguration\BuildConfiguration.xml (Global AppData)</Link></None>
