Change the logic for internal includes to check for the same scope. The previous logic checked if the dependency module contained the scope of the module that depended on it, which is only ever the case if the modules share the same scope, given that in `UEBuildTarget.ValidateModule`, we check for nested scopes in the other direction. Checking for equality is closer to what actually happens in practice, and should save a few cycles due to being a faster check too.

diff --git a/Engine/Source/Programs/UnrealBuildTool/Configuration/UEBuildModule.cs b/Engine/Source/Programs/UnrealBuildTool/Configuration/UEBuildModule.cs
--- a/Engine/Source/Programs/UnrealBuildTool/Configuration/UEBuildModule.cs
+++ b/Engine/Source/Programs/UnrealBuildTool/Configuration/UEBuildModule.cs
@@ -729,8 +729,8 @@ namespace UnrealBuildTool
 				AddIncludePaths(IncludePaths, LegacyPublicIncludePaths);
 			}
 
-			// Add this module's internal include paths, only if the scope contains the same as the SourceModule's scope or if the source module is an engine module
-			if (SourceModule != null && (Rules.Context.Scope.Contains(SourceModule.Rules.Context.Scope) || SourceModule.Rules.bTreatAsEngineModule))
+			// Add this module's internal include paths, only if the scope is the same as the SourceModule's scope or if the source module is an engine module
+			if (SourceModule != null && (Rules.Context.Scope == SourceModule.Rules.Context.Scope || SourceModule.Rules.bTreatAsEngineModule))
 			{
 				AddIncludePaths(IncludePaths, InternalIncludePaths);
 			}
