UnrealBuildTool: fix csproj link

#rnx

diff --git a/Engine/Source/Programs/UnrealBuildTool/UnrealBuildTool.csproj b/Engine/Source/Programs/UnrealBuildTool/UnrealBuildTool.csproj
--- a/Engine/Source/Programs/UnrealBuildTool/UnrealBuildTool.csproj
+++ b/Engine/Source/Programs/UnrealBuildTool/UnrealBuildTool.csproj
@@ -124,7 +124,7 @@
   <!-- Various BuildConfiguration files -->
   <ItemGroup>
       <None Condition="Exists('..\..\..\Restricted\NotForLicensees\Programs\UnrealBuildTool\BuildConfiguration.xml')" Include="..\..\..\Restricted\NotForLicensees\Programs\UnrealBuildTool\BuildConfiguration.xml"><Link>BuildConfiguration\BuildConfiguration.xml (Engine NotForLicensees)</Link></None>
-      <None Condition="Exists('..\..\..\Restricted\NotForLicensees\Programs\UnrealBuildTool\BuildConfiguration.xml')" Include="..\..\..\Restricted\NotForLicensees\Programs\UnrealBuildTool\BuildConfiguration.xml"><Link>BuildConfiguration\ReleaseBuildConfiguration.xml (Engine NotForLicensees Release)</Link></None>
+      <None Condition="Exists('..\..\..\Restricted\NotForLicensees\Programs\UnrealBuildTool\ReleaseBuildConfiguration.xml')" Include="..\..\..\Restricted\NotForLicensees\Programs\UnrealBuildTool\ReleaseBuildConfiguration.xml"><Link>BuildConfiguration\ReleaseBuildConfiguration.xml (Engine NotForLicensees Release)</Link></None>
       <None Condition="Exists('..\..\..\Saved\UnrealBuildTool\BuildConfiguration.xml')" Include="..\..\..\Saved\UnrealBuildTool\BuildConfiguration.xml"><Link>BuildConfiguration\BuildConfiguration.xml (Engine Saved)</Link></None>
       <None Condition="Exists('$(ProgramData)\Unreal Engine\UnrealBuildTool\BuildConfiguration.xml')" Include="$(ProgramData)\Unreal Engine\UnrealBuildTool\BuildConfiguration.xml"><Link>BuildConfiguration\BuildConfiguration.xml (Global ProgramData)</Link></None>
       <None Condition="Exists('$(AppData)\Unreal Engine\UnrealBuildTool\BuildConfiguration.xml')" Include="$(AppData)\Unreal Engine\UnrealBuildTool\BuildConfiguration.xml"><Link>BuildConfiguration\BuildConfiguration.xml (Global AppData)</Link></None>
