[UBT]
* Fixed so retry logic on failed actions does not fall back to native run if vfs is enabled (since it will fail anyway)

diff --git a/Engine/Source/Programs/UnrealBuildTool/Executors/UnrealBuildAccelerator/UBAExecutor.cs b/Engine/Source/Programs/UnrealBuildTool/Executors/UnrealBuildAccelerator/UBAExecutor.cs
--- a/Engine/Source/Programs/UnrealBuildTool/Executors/UnrealBuildAccelerator/UBAExecutor.cs
+++ b/Engine/Source/Programs/UnrealBuildTool/Executors/UnrealBuildAccelerator/UBAExecutor.cs
@@ -829,7 +829,11 @@ namespace UnrealBuildTool
 						{
 							if (UBAConfig.bAllowRetry && (exitCode != 0 && UBAConfig.bForcedRetry || (exitCode >= 9000 && exitCode < 10000)))
 							{
-								if (!CanRetryNatively || UBAConfig.bStoreObjFilesCompressed)
+								// If files are compressed (pch, obj etc), then native will not work so must retry detoured
+								// If vfs paths are used, then native will not work
+								bool mustRunDetoured = !CanRetryNatively || UBAConfig.bStoreObjFilesCompressed || action.RootPaths.bUseVfs;
+
+								if (mustRunDetoured)
 								{
 									logger.LogInformation("{Description} {StatusDescription}: Exited with error code {ExitCode}. This action will retry", action.CommandDescription, action.StatusDescription, exitCode);
 									++forcedRetryActions;
