UnrealBuildTool: Add toolchain version to generated pch cpp file

#rnx

diff --git a/Engine/Source/Programs/UnrealBuildTool/Configuration/UEBuildModuleCPP.cs b/Engine/Source/Programs/UnrealBuildTool/Configuration/UEBuildModuleCPP.cs
--- a/Engine/Source/Programs/UnrealBuildTool/Configuration/UEBuildModuleCPP.cs
+++ b/Engine/Source/Programs/UnrealBuildTool/Configuration/UEBuildModuleCPP.cs
@@ -2273,7 +2273,8 @@ namespace UnrealBuildTool
 			// Build the contents of the wrapper file
 			StringBuilder WrapperContents = new StringBuilder();
 			{
-				WrapperContents.AppendLine("// PCH for {0}", IncludedFile.Name);
+				WrapperContents.AppendLine("// Generated by UnrealBuildTool (UEBuildModuleCPP.cs) : PCH for {0}", IncludedFile.Name);
+				WrapperContents.AppendLine("#pragma once");
 				WrapperContents.AppendLine("#include \"{0}\"", DefinitionsFile.Name);
 				WrapperContents.AppendLine("#include \"{0}\"", IncludedFile.Name);
 				WrapperContents.AppendLine("#ifdef __ISPC_ALIGN__");
diff --git a/Engine/Source/Programs/UnrealBuildTool/Platform/Windows/VCToolChain.cs b/Engine/Source/Programs/UnrealBuildTool/Platform/Windows/VCToolChain.cs
--- a/Engine/Source/Programs/UnrealBuildTool/Platform/Windows/VCToolChain.cs
+++ b/Engine/Source/Programs/UnrealBuildTool/Platform/Windows/VCToolChain.cs
@@ -2148,9 +2148,13 @@ namespace UnrealBuildTool
 				{
 					// Generate a CPP File that just includes the precompiled header.
 					string PrecompiledHeaderIncludeFilenameString = CompileEnvironment.PrecompiledHeaderIncludeFilename!.GetFileName();
-					string PchCppFile = $"// Compiler: {EnvVars.CompilerVersion}\n#include \"{PrecompiledHeaderIncludeFilenameString}\"\r\n";
+					List<string> PchCppFileContents = [
+						"// Generated by UnrealBuildTool (VCToolChain.cs)",
+						$"// Compiler: {(EnvVars.Compiler.IsMSVC() ? $"MSVC {EnvVars.CompilerVersion}" : $"{EnvVars.Compiler} {EnvVars.CompilerVersion}, MSVC {EnvVars.ToolChainVersion}")}",
+						$"#include \"{PrecompiledHeaderIncludeFilenameString}\""
+					];
 					CompileAction.SourceFile = FileItem.GetItemByFileReference(CompileEnvironment.PrecompiledHeaderIncludeFilename!.ChangeExtension(".cpp"));
-					Graph.CreateIntermediateTextFile(CompileAction.SourceFile, PchCppFile);
+					Graph.CreateIntermediateTextFile(CompileAction.SourceFile, PchCppFileContents);
 
 					// Add the precompiled header file to the produced items list.
 					string PchExtension = CompileEnvironment.bUseHeaderUnitsForPch ? ".ifc" : ".pch";
