Fix issue IDE0059 (Remove unnecessary value assignment) in UBT.

diff --git a/Engine/Source/Programs/UnrealBuildTool/Configuration/Descriptors/PluginReferenceDescriptor.cs b/Engine/Source/Programs/UnrealBuildTool/Configuration/Descriptors/PluginReferenceDescriptor.cs
--- a/Engine/Source/Programs/UnrealBuildTool/Configuration/Descriptors/PluginReferenceDescriptor.cs
+++ b/Engine/Source/Programs/UnrealBuildTool/Configuration/Descriptors/PluginReferenceDescriptor.cs
@@ -272,8 +272,7 @@ namespace UnrealBuildTool
 			{
 				foreach (string PlatformName in PlatformNames)
 				{
-					UnrealTargetPlatform Platform;
-					if (!UnrealTargetPlatform.TryParse(PlatformName, out Platform))
+					if (!UnrealTargetPlatform.TryParse(PlatformName, out _))
 					{
 						Log.TraceLogOnce("Ignoring unknown platform '{0}' (referenced via a project's plugin descriptor for '{1}')", PlatformName, Name);
 					}
diff --git a/Engine/Source/Programs/UnrealBuildTool/Configuration/Rules/CppCompileWarnings.cs b/Engine/Source/Programs/UnrealBuildTool/Configuration/Rules/CppCompileWarnings.cs
--- a/Engine/Source/Programs/UnrealBuildTool/Configuration/Rules/CppCompileWarnings.cs
+++ b/Engine/Source/Programs/UnrealBuildTool/Configuration/Rules/CppCompileWarnings.cs
@@ -1660,8 +1660,6 @@ namespace UnrealBuildTool
 		{
 			_logger?.LogTrace("Applying default C++ Compiler warnings for the default construction context.");
 
-			BuildSettingsVersion buildSettingsVersion = BuildContextProvider != null ? BuildContextProvider._buildContext.GetBuildSettings() : BuildSettingsVersion.V1;
-
 			foreach ((MemberInfo mi, List<WarningLevelDefaultAttribute> miAttr) in s_membersWithWarningLevelDefaultAttribute.Value)
 			{
 				WarningLevel defaultWarningLevel = WarningLevel.Default;
diff --git a/Engine/Source/Programs/UnrealBuildTool/Configuration/UEBuildModule.cs b/Engine/Source/Programs/UnrealBuildTool/Configuration/UEBuildModule.cs
--- a/Engine/Source/Programs/UnrealBuildTool/Configuration/UEBuildModule.cs
+++ b/Engine/Source/Programs/UnrealBuildTool/Configuration/UEBuildModule.cs
@@ -1663,7 +1663,6 @@ namespace UnrealBuildTool
 			}
 
 			int words = (nodeCount + 63) >> 6;
-			long bytesNeeded = (long)nodeCount * words * sizeof(ulong);
 			ulong[] depBits = new ulong[(long)nodeCount * words];
 
 			for (int i = 0; i < nodeCount; i++)
diff --git a/Engine/Source/Programs/UnrealBuildTool/Configuration/UEBuildPlatform.cs b/Engine/Source/Programs/UnrealBuildTool/Configuration/UEBuildPlatform.cs
--- a/Engine/Source/Programs/UnrealBuildTool/Configuration/UEBuildPlatform.cs
+++ b/Engine/Source/Programs/UnrealBuildTool/Configuration/UEBuildPlatform.cs
@@ -1155,9 +1155,8 @@ namespace UnrealBuildTool
 			{
 				foreach (string Key in BoolKeys)
 				{
-					bool Default = false, Project = false;
-					DefaultIni.GetBool(Section, Key, out Default);
-					ProjIni.GetBool(Section, Key, out Project);
+					DefaultIni.GetBool(Section, Key, out bool Default);
+					ProjIni.GetBool(Section, Key, out bool Project);
 					if (Default != Project)
 					{
 						Log.TraceInformationOnce("{0} is not set to default. (Base: {1} vs. {2}: {3})", Key, Default, Path.GetFileName(ProjectDirectoryName.FullName), Project);
@@ -1171,9 +1170,8 @@ namespace UnrealBuildTool
 			{
 				foreach (string Key in IntKeys)
 				{
-					int Default = 0, Project = 0;
-					DefaultIni.GetInt32(Section, Key, out Default);
-					ProjIni.GetInt32(Section, Key, out Project);
+					DefaultIni.GetInt32(Section, Key, out int Default);
+					ProjIni.GetInt32(Section, Key, out int Project);
 					if (Default != Project)
 					{
 						Log.TraceInformationOnce("{0} is not set to default. (Base: {1} vs. {2}: {3})", Key, Default, Path.GetFileName(ProjectDirectoryName.FullName), Project);
@@ -1187,9 +1185,8 @@ namespace UnrealBuildTool
 			{
 				foreach (string Key in StringKeys)
 				{
-					string? Default = "", Project = "";
-					DefaultIni.GetString(Section, Key, out Default);
-					ProjIni.GetString(Section, Key, out Project);
+					DefaultIni.GetString(Section, Key, out string? Default);
+					ProjIni.GetString(Section, Key, out string? Project);
 					if (Default != Project)
 					{
 						Log.TraceInformationOnce("{0} is not set to default. (Base: {1} vs. {2}: {3})", Key, Default, Path.GetFileName(ProjectDirectoryName.FullName), Project);
diff --git a/Engine/Source/Programs/UnrealBuildTool/Configuration/UEBuildTarget.cs b/Engine/Source/Programs/UnrealBuildTool/Configuration/UEBuildTarget.cs
--- a/Engine/Source/Programs/UnrealBuildTool/Configuration/UEBuildTarget.cs
+++ b/Engine/Source/Programs/UnrealBuildTool/Configuration/UEBuildTarget.cs
@@ -6795,7 +6795,6 @@ namespace UnrealBuildTool
 			// @todo projectfiles: Cross-platform modules can appear here during project generation, but they may have already
 			//   been filtered out by the project generator.  This causes the projects to not be added to directories properly.
 			ModuleRules RulesObject = CreateModuleRulesAndSetDefaults(ModuleName, ReferenceChain, Logger);
-			DirectoryReference ModuleDirectory = RulesObject.File.Directory;
 
 			// Clear the bUsePrecompiled flag if we're compiling a foreign plugin; since it's treated like an engine module, it will default to true in an installed build.
 			if (RulesObject.Plugin != null && RulesObject.Plugin.File == ForeignPlugin)
diff --git a/Engine/Source/Programs/UnrealBuildTool/Executors/Experimental/FASTBuild.cs b/Engine/Source/Programs/UnrealBuildTool/Executors/Experimental/FASTBuild.cs
--- a/Engine/Source/Programs/UnrealBuildTool/Executors/Experimental/FASTBuild.cs
+++ b/Engine/Source/Programs/UnrealBuildTool/Executors/Experimental/FASTBuild.cs
@@ -801,8 +801,7 @@ namespace UnrealBuildTool
 
 		private string GetOptionValue(Dictionary<string, string> OptionsDictionary, string Key, LinkedAction Action, ILogger Logger, bool ProblemIfNotFound = false)
 		{
-			string? Value = String.Empty;
-			if (OptionsDictionary.TryGetValue(Key, out Value))
+			if (OptionsDictionary.TryGetValue(Key, out string? Value))
 			{
 				return Value.Trim(new char[] { '\"' });
 			}
diff --git a/Engine/Source/Programs/UnrealBuildTool/Executors/UnrealBuildAccelerator/UBAAgentCoordinatorHorde.cs b/Engine/Source/Programs/UnrealBuildTool/Executors/UnrealBuildAccelerator/UBAAgentCoordinatorHorde.cs
--- a/Engine/Source/Programs/UnrealBuildTool/Executors/UnrealBuildAccelerator/UBAAgentCoordinatorHorde.cs
+++ b/Engine/Source/Programs/UnrealBuildTool/Executors/UnrealBuildAccelerator/UBAAgentCoordinatorHorde.cs
@@ -120,7 +120,6 @@ namespace UnrealBuildTool
 
 			if (_client != null)
 			{
-				IComputeClient client = _client;
 				_client = null;
 				// Set CPU resource need to zero (will also expire on the server if not updated)
 				await UpdateCpuCoreNeedAsync(0);
@@ -859,10 +858,9 @@ namespace UnrealBuildTool
 			}
 
 			UBAHordeSession? hordeSession = null;
-			string reason = "Unknown";
 			if (_hordeSessionTask != null)
 			{
-				(hordeSession, reason) = _hordeSessionTask.Result;
+				(hordeSession, string reason) = _hordeSessionTask.Result;
 				_updateStatus!(0, 6, reason, EpicGames.UBA.LogEntryType.Info, null);
 			}
 			if (hordeSession == null)
@@ -1018,7 +1016,7 @@ namespace UnrealBuildTool
 				return;
 			}
 
-			(UBAHordeSession? hordeSession, string? reason) = await _hordeSessionTask;
+			(UBAHordeSession? hordeSession, _) = await _hordeSessionTask;
 			_hordeSessionTask = null;
 			if (hordeSession != null)
 			{
diff --git a/Engine/Source/Programs/UnrealBuildTool/Executors/UnrealBuildAccelerator/UBAExecutor.cs b/Engine/Source/Programs/UnrealBuildTool/Executors/UnrealBuildAccelerator/UBAExecutor.cs
--- a/Engine/Source/Programs/UnrealBuildTool/Executors/UnrealBuildAccelerator/UBAExecutor.cs
+++ b/Engine/Source/Programs/UnrealBuildTool/Executors/UnrealBuildAccelerator/UBAExecutor.cs
@@ -87,7 +87,7 @@ namespace UnrealBuildTool
 		{
 			if (_cacheClients.Count > 0)
 			{
-				foreach ((UnrealBuildAcceleratorCacheConfig config, ICacheClient client) in _cacheClients)
+				foreach ((_, ICacheClient client) in _cacheClients)
 				{
 					client.Dispose();
 				}
diff --git a/Engine/Source/Programs/UnrealBuildTool/Modes/PipInstallMode.cs b/Engine/Source/Programs/UnrealBuildTool/Modes/PipInstallMode.cs
--- a/Engine/Source/Programs/UnrealBuildTool/Modes/PipInstallMode.cs
+++ b/Engine/Source/Programs/UnrealBuildTool/Modes/PipInstallMode.cs
@@ -323,7 +323,7 @@ namespace UnrealBuildTool.Modes
 					continue;
 				}
 
-				foreach (JsonObject PlatformReqs in PipEnv.CompatibleRequirements(PluginJson, TargetPlatform, TargetArch))
+				foreach (JsonObject _ in PipEnv.CompatibleRequirements(PluginJson, TargetPlatform, TargetArch))
 				{
 					PluginsList.Add(Plugin.File.ToString());
 					break;
diff --git a/Engine/Source/Programs/UnrealBuildTool/Platform/Android/AndroidToolChain.cs b/Engine/Source/Programs/UnrealBuildTool/Platform/Android/AndroidToolChain.cs
--- a/Engine/Source/Programs/UnrealBuildTool/Platform/Android/AndroidToolChain.cs
+++ b/Engine/Source/Programs/UnrealBuildTool/Platform/Android/AndroidToolChain.cs
@@ -292,29 +292,25 @@ namespace UnrealBuildTool
 		private bool DisableFunctionDataSections()
 		{
 			ConfigHierarchy Ini = ConfigCache.ReadHierarchy(ConfigHierarchyType.Engine, DirectoryReference.FromFile(ProjectFile), UnrealTargetPlatform.Android);
-			bool bDisableFunctionDataSections = false;
-			return Ini.GetBool("/Script/AndroidRuntimeSettings.AndroidRuntimeSettings", "bDisableFunctionDataSections", out bDisableFunctionDataSections) && bDisableFunctionDataSections;
+			return Ini.GetBool("/Script/AndroidRuntimeSettings.AndroidRuntimeSettings", "bDisableFunctionDataSections", out bool bDisableFunctionDataSections) && bDisableFunctionDataSections;
 		}
 
 		private bool EnableAdvancedBinaryCompression()
 		{
 			ConfigHierarchy Ini = ConfigCache.ReadHierarchy(ConfigHierarchyType.Engine, DirectoryReference.FromFile(ProjectFile), UnrealTargetPlatform.Android);
-			bool bEnableAdvancedBinaryCompression = false;
-			return Ini.GetBool("/Script/AndroidRuntimeSettings.AndroidRuntimeSettings", "bEnableAdvancedBinaryCompression", out bEnableAdvancedBinaryCompression) && bEnableAdvancedBinaryCompression;
+			return Ini.GetBool("/Script/AndroidRuntimeSettings.AndroidRuntimeSettings", "bEnableAdvancedBinaryCompression", out bool bEnableAdvancedBinaryCompression) && bEnableAdvancedBinaryCompression;
 		}
 
 		private bool DisableStackProtector()
 		{
 			ConfigHierarchy Ini = ConfigCache.ReadHierarchy(ConfigHierarchyType.Engine, DirectoryReference.FromFile(ProjectFile), UnrealTargetPlatform.Android);
-			bool bDisableStackProtector = false;
-			return Ini.GetBool("/Script/AndroidRuntimeSettings.AndroidRuntimeSettings", "bDisableStackProtector", out bDisableStackProtector) && bDisableStackProtector;
+			return Ini.GetBool("/Script/AndroidRuntimeSettings.AndroidRuntimeSettings", "bDisableStackProtector", out bool bDisableStackProtector) && bDisableStackProtector;
 		}
 
 		private bool DisableLibCppSharedDependencyValidation()
 		{
 			ConfigHierarchy Ini = ConfigCache.ReadHierarchy(ConfigHierarchyType.Engine, DirectoryReference.FromFile(ProjectFile), UnrealTargetPlatform.Android);
-			bool bDisableLibCppSharedDependencyValidation = false;
-			return Ini.GetBool("/Script/AndroidRuntimeSettings.AndroidRuntimeSettings", "bDisableLibCppSharedDependencyValidation", out bDisableLibCppSharedDependencyValidation) && bDisableLibCppSharedDependencyValidation;
+			return Ini.GetBool("/Script/AndroidRuntimeSettings.AndroidRuntimeSettings", "bDisableLibCppSharedDependencyValidation", out bool bDisableLibCppSharedDependencyValidation) && bDisableLibCppSharedDependencyValidation;
 		}
 
 		private FileReference GetVersionScriptFileName(LinkEnvironment LinkEnvironment)
@@ -379,11 +375,11 @@ namespace UnrealBuildTool
 		}
 
 		//This doesn't take into account SDK version overrides in packaging
-		public int GetMinSdkVersion(int MinSdk = MinimumNDKAPILevel)
+		public int GetMinSdkVersion()
 		{
-			int MinSDKVersion = MinSdk;
 			ConfigHierarchy Ini = ConfigCache.ReadHierarchy(ConfigHierarchyType.Engine, DirectoryReference.FromFile(ProjectFile), UnrealTargetPlatform.Android);
-			Ini.GetInt32("/Script/AndroidRuntimeSettings.AndroidRuntimeSettings", "MinSDKVersion", out MinSDKVersion);
+			// TODO: Should this check the return value and return `MinimumNDKAPILevel` if false?
+			Ini.GetInt32("/Script/AndroidRuntimeSettings.AndroidRuntimeSettings", "MinSDKVersion", out int MinSDKVersion);
 			return MinSDKVersion;
 		}
 
@@ -431,8 +427,7 @@ namespace UnrealBuildTool
 
 			if (NDKLevel == "latest")
 			{
-				int MinPlatform, MaxPlatform;
-				if (!ReadMinMaxPlatforms(PlatformsFilename, out MinPlatform, out MaxPlatform))
+				if (!ReadMinMaxPlatforms(PlatformsFilename, out _, out int MaxPlatform))
 				{
 					throw new BuildException("No NDK platforms found in {0}", PlatformsFilename);
 				}
@@ -457,8 +452,7 @@ namespace UnrealBuildTool
 				throw new BuildException("No NDK platforms found in {0}", PlatformsFilename);
 			}
 
-			int MinPlatform, MaxPlatform;
-			if (!ReadMinMaxPlatforms(PlatformsFilename, out MinPlatform, out MaxPlatform))
+			if (!ReadMinMaxPlatforms(PlatformsFilename, out _, out int MaxPlatform))
 			{
 				throw new BuildException("No NDK platforms found in {0}", PlatformsFilename);
 			}
@@ -666,8 +660,6 @@ namespace UnrealBuildTool
 
 			Log.TraceInformationOnce("Compiling Native 64-bit code with NDK API '{0}'", NDKApiLevel64Bit);
 
-			string NativeGluePath = Path.GetFullPath(GetNativeGluePath());
-
 			if (BuildWithHiddenSymbolVisibility())
 			{
 				Arguments.Add("-fvisibility=hidden");
diff --git a/Engine/Source/Programs/UnrealBuildTool/Platform/Android/UEDeployAndroid.cs b/Engine/Source/Programs/UnrealBuildTool/Platform/Android/UEDeployAndroid.cs
--- a/Engine/Source/Programs/UnrealBuildTool/Platform/Android/UEDeployAndroid.cs
+++ b/Engine/Source/Programs/UnrealBuildTool/Platform/Android/UEDeployAndroid.cs
@@ -324,11 +324,10 @@ namespace UnrealBuildTool
 				return CachedBuildToolsVersion;
 			}
 
-			string? BestVersionString = null;
 			uint BestVersion = 0;
 
 			ConfigHierarchy Ini = GetConfigCacheIni(ConfigHierarchyType.Engine);
-			Ini.GetString("/Script/AndroidRuntimeSettings.AndroidRuntimeSettings", "BuildToolsOverride", out BestVersionString);
+			Ini.GetString("/Script/AndroidRuntimeSettings.AndroidRuntimeSettings", "BuildToolsOverride", out string? BestVersionString);
 
 			if (BestVersionString == null || String.IsNullOrEmpty(BestVersionString) || BestVersionString == "latest")
 			{
@@ -925,11 +924,8 @@ namespace UnrealBuildTool
 			if (CachedStoreVersion < 1)
 			{
 				ConfigHierarchy Ini = GetConfigCacheIni(ConfigHierarchyType.Engine);
-				int StoreVersion = 1;
-				Ini.GetInt32("/Script/AndroidRuntimeSettings.AndroidRuntimeSettings", "StoreVersion", out StoreVersion);
-
-				bool bUseChangeListAsStoreVersion = false;
-				Ini.GetBool("/Script/AndroidRuntimeSettings.AndroidRuntimeSettings", "bUseChangeListAsStoreVersion", out bUseChangeListAsStoreVersion);
+				Ini.GetInt32("/Script/AndroidRuntimeSettings.AndroidRuntimeSettings", "StoreVersion", out int StoreVersion);
+				Ini.GetBool("/Script/AndroidRuntimeSettings.AndroidRuntimeSettings", "bUseChangeListAsStoreVersion", out bool bUseChangeListAsStoreVersion);
 
 				// override store version with changelist if enabled
 				if (bUseChangeListAsStoreVersion)
@@ -937,14 +933,13 @@ namespace UnrealBuildTool
 					// make sure changelist is cached (clear unused warning)
 					string EngineVersion = ReadEngineVersion();
 
-					int Changelist = 0;
-					if (EngineVersion == null || !(Int32.TryParse(EngineChangelist, out Changelist)) || Changelist == 0)
+					if (EngineVersion == null || !(Int32.TryParse(EngineChangelist, out int Changelist)) || Changelist == 0)
 					{
 						throw new BuildException("AndroidRuntimeSettings.bUseChangeListAsStoreVersion is enabled but no valid Engine version changelist was found!");
 					}
 
-							StoreVersion = Changelist;
-						}
+					StoreVersion = Changelist;
+				}
 
 				bool IsBuildMachine = Unreal.IsBuildMachine();
 				Logger.LogInformation("GotStoreVersion found v{StoreVersion}. (bUseChangeListAsStoreVersion={UseChangeListAsStoreVersion} IsBuildMachine={IsBuildMachine} EngineChangeList={EngineChangeList})", StoreVersion, bUseChangeListAsStoreVersion, IsBuildMachine, EngineChangelist);
@@ -974,20 +969,17 @@ namespace UnrealBuildTool
 			if (String.IsNullOrEmpty(CachedVersionDisplayName))
 			{
 				ConfigHierarchy Ini = GetConfigCacheIni(ConfigHierarchyType.Engine);
-				string VersionDisplayName = "";
-				Ini.GetString("/Script/AndroidRuntimeSettings.AndroidRuntimeSettings", "VersionDisplayName", out VersionDisplayName);
+				Ini.GetString("/Script/AndroidRuntimeSettings.AndroidRuntimeSettings", "VersionDisplayName", out string VersionDisplayName);
 
 				if (Unreal.IsBuildMachine())
 				{
-					bool bAppendChangeListToVersionDisplayName = false;
-					Ini.GetBool("/Script/AndroidRuntimeSettings.AndroidRuntimeSettings", "bAppendChangeListToVersionDisplayName", out bAppendChangeListToVersionDisplayName);
+					Ini.GetBool("/Script/AndroidRuntimeSettings.AndroidRuntimeSettings", "bAppendChangeListToVersionDisplayName", out bool bAppendChangeListToVersionDisplayName);
 					if (bAppendChangeListToVersionDisplayName)
 					{
 						VersionDisplayName = String.Format("{0}-{1}", VersionDisplayName, EngineChangelist);
 					}
 
-					bool bAppendPlatformToVersionDisplayName = false;
-					Ini.GetBool("/Script/AndroidRuntimeSettings.AndroidRuntimeSettings", "bAppendPlatformToVersionDisplayName", out bAppendPlatformToVersionDisplayName);
+					Ini.GetBool("/Script/AndroidRuntimeSettings.AndroidRuntimeSettings", "bAppendPlatformToVersionDisplayName", out bool bAppendPlatformToVersionDisplayName);
 					if (bAppendPlatformToVersionDisplayName)
 					{
 						VersionDisplayName = String.Format("{0}-Android", VersionDisplayName);
@@ -996,8 +988,7 @@ namespace UnrealBuildTool
 					// append optional text to version name if embedded build
 					if (bIsEmbedded)
 					{
-						string EmbeddedAppendDisplayName = "";
-						if (Ini.GetString("/Script/AndroidRuntimeSettings.AndroidRuntimeSettings", "EmbeddedAppendDisplayName", out EmbeddedAppendDisplayName))
+						if (Ini.GetString("/Script/AndroidRuntimeSettings.AndroidRuntimeSettings", "EmbeddedAppendDisplayName", out string EmbeddedAppendDisplayName))
 						{
 							VersionDisplayName += EmbeddedAppendDisplayName;
 						}
@@ -1572,10 +1563,8 @@ namespace UnrealBuildTool
 		void LogBuildSetup()
 		{
 			ConfigHierarchy Ini = GetConfigCacheIni(ConfigHierarchyType.Engine);
-			bool bBuildForES31 = false;
-			Ini.GetBool("/Script/AndroidRuntimeSettings.AndroidRuntimeSettings", "bBuildForES31", out bBuildForES31);
-			bool bSupportsVulkan = false;
-			Ini.GetBool("/Script/AndroidRuntimeSettings.AndroidRuntimeSettings", "bSupportsVulkan", out bSupportsVulkan);
+			Ini.GetBool("/Script/AndroidRuntimeSettings.AndroidRuntimeSettings", "bBuildForES31", out bool bBuildForES31);
+			Ini.GetBool("/Script/AndroidRuntimeSettings.AndroidRuntimeSettings", "bSupportsVulkan", out bool bSupportsVulkan);
 
 			Logger.LogInformation("bBuildForES31: {BuildForES31}", (bBuildForES31 ? "true" : "false"));
 			Logger.LogInformation("bSupportsVulkan: {SupportsVulkan}", (bSupportsVulkan ? "true" : "false"));
@@ -1583,11 +1572,9 @@ namespace UnrealBuildTool
 
 		void CopyVulkanValidationLayers(string UnrealBuildPath, UnrealArch UnrealArch, string NDKArch, string Configuration)
 		{
-			bool bSupportsVulkan = false;
-			bool bSupportsVulkanSM5 = false;
 			ConfigHierarchy Ini = GetConfigCacheIni(ConfigHierarchyType.Engine);
-			Ini.GetBool("/Script/AndroidRuntimeSettings.AndroidRuntimeSettings", "bSupportsVulkan", out bSupportsVulkan);
-			Ini.GetBool("/Script/AndroidRuntimeSettings.AndroidRuntimeSettings", "bSupportsVulkanSM5", out bSupportsVulkanSM5);
+			Ini.GetBool("/Script/AndroidRuntimeSettings.AndroidRuntimeSettings", "bSupportsVulkan", out bool bSupportsVulkan);
+			Ini.GetBool("/Script/AndroidRuntimeSettings.AndroidRuntimeSettings", "bSupportsVulkanSM5", out bool bSupportsVulkanSM5);
 
 			bool bCopyVulkanLayers = (bSupportsVulkan || bSupportsVulkanSM5) && (Configuration != "Shipping");
 			if (bCopyVulkanLayers)
@@ -1611,8 +1598,7 @@ namespace UnrealBuildTool
 					Logger.LogWarning("{ANDROID_VULKAN_VALIDATION_LAYER} vulkan layer not found at {VulkanLayersDir}, skipping", ANDROID_VULKAN_VALIDATION_LAYER, VulkanLayersDir);
 				}
 
-				string DebugVulkanLayerDirectory = "";
-				AndroidPlatformSDK.GetPath(Ini, "/Script/AndroidRuntimeSettings.AndroidRuntimeSettings", "DebugVulkanLayerDirectory", out DebugVulkanLayerDirectory);
+				AndroidPlatformSDK.GetPath(Ini, "/Script/AndroidRuntimeSettings.AndroidRuntimeSettings", "DebugVulkanLayerDirectory", out string DebugVulkanLayerDirectory);
 				Logger.LogInformation("DebugVulkanLayerDirectory {LayerDir}", DebugVulkanLayerDirectory);
 
 				if (!String.IsNullOrEmpty(DebugVulkanLayerDirectory))
@@ -2131,7 +2117,6 @@ namespace UnrealBuildTool
 		{
 			string StringsXMLPath = Path.Combine(UnrealBuildPath, "res/values/strings.xml");
 
-			ApplicationDisplayName = null;
 			ConfigHierarchy Ini = GetConfigCacheIni(ConfigHierarchyType.Engine);
 			Ini.GetString("/Script/AndroidRuntimeSettings.AndroidRuntimeSettings", "ApplicationDisplayName", out ApplicationDisplayName);
 
@@ -2562,8 +2547,7 @@ namespace UnrealBuildTool
 		private void PickSplashScreenOrientation(string CookFlavor, string UnrealBuildPath, bool bNeedPortrait, bool bNeedLandscape)
 		{
 			ConfigHierarchy Ini = GetConfigCacheIni(ConfigHierarchyType.Engine);
-			bool bShowLaunchImage = false;
-			Ini.GetBool("/Script/AndroidRuntimeSettings.AndroidRuntimeSettings", "bShowLaunchImage", out bShowLaunchImage);
+			Ini.GetBool("/Script/AndroidRuntimeSettings.AndroidRuntimeSettings", "bShowLaunchImage", out bool bShowLaunchImage);
 			bool bPackageForMetaQuest = IsPackagingForMetaQuest(CookFlavor, Ini);
 
 			//override the parameters if we are not showing a launch image or are packaging for Meta Quest
@@ -2758,8 +2742,7 @@ namespace UnrealBuildTool
 		private string GetPublicKey()
 		{
 			ConfigHierarchy Ini = GetConfigCacheIni(ConfigHierarchyType.Engine);
-			string PlayLicenseKey = "";
-			Ini.GetString("/Script/AndroidRuntimeSettings.AndroidRuntimeSettings", "GooglePlayLicenseKey", out PlayLicenseKey);
+			Ini.GetString("/Script/AndroidRuntimeSettings.AndroidRuntimeSettings", "GooglePlayLicenseKey", out string PlayLicenseKey);
 			return PlayLicenseKey;
 		}
 
@@ -2796,10 +2779,7 @@ namespace UnrealBuildTool
 			int StoreVersion = GetStoreVersion(UnrealArch);
 
 			string Arch = GetNDKArch(UnrealArch);
-			int NDKLevelInt = 0;
-			int MinSDKVersion = 0;
-			int TargetSDKVersion = 0;
-			GetMinTargetSDKVersions(ToolChain, UnrealArch, UPL!, Arch, bEnableBundle, out MinSDKVersion, out TargetSDKVersion, out NDKLevelInt);
+			GetMinTargetSDKVersions(ToolChain, UnrealArch, UPL!, Arch, bEnableBundle, out int MinSDKVersion, out int TargetSDKVersion, out _);
 
 			// get project version from ini
 			ConfigHierarchy GameIni = GetConfigCacheIni(ConfigHierarchyType.Game);
@@ -2808,7 +2788,8 @@ namespace UnrealBuildTool
 
 			// ini file to get settings from
 			ConfigHierarchy Ini = GetConfigCacheIni(ConfigHierarchyType.Engine);
-			string PackageName = GetPackageName(ProjectName);
+			// Discard result but call the method for the logging side-effects to preserve previous behaviour.
+			_ = GetPackageName(ProjectName);
 			string VersionDisplayName = GetVersionDisplayName(bIsEmbedded);
 			string DepthBufferPreference;
 			Ini.GetString("/Script/AndroidRuntimeSettings.AndroidRuntimeSettings", "DepthBufferPreference", out DepthBufferPreference);
@@ -2823,55 +2804,33 @@ namespace UnrealBuildTool
 				MaxAspectRatioValue = 2.4f;
 			}
 			string Orientation = ConvertOrientationIniValue(GetOrientation(Arch));
-			bool EnableFullScreen;
-			Ini.GetBool("/Script/AndroidRuntimeSettings.AndroidRuntimeSettings", "bFullScreen", out EnableFullScreen);
-			bool bUseDisplayCutout;
-			Ini.GetBool("/Script/AndroidRuntimeSettings.AndroidRuntimeSettings", "bUseDisplayCutout", out bUseDisplayCutout);
-			string RoundedEdgeSafeZone;
-			Ini.GetString("/Script/AndroidRuntimeSettings.AndroidRuntimeSettings", "RoundedEdgeSafeZone", out RoundedEdgeSafeZone);
-			bool bAllowResizing;
-			Ini.GetBool("/Script/AndroidRuntimeSettings.AndroidRuntimeSettings", "bAllowResizing", out bAllowResizing);
-			bool bSupportSizeChanges;
-			Ini.GetBool("/Script/AndroidRuntimeSettings.AndroidRuntimeSettings", "bSupportSizeChanges", out bSupportSizeChanges);
-			bool bRestoreNotificationsOnReboot = false;
-			Ini.GetBool("/Script/AndroidRuntimeSettings.AndroidRuntimeSettings", "bRestoreNotificationsOnReboot", out bRestoreNotificationsOnReboot);
-			List<string>? ExtraManifestNodeTags;
-			Ini.GetArray("/Script/AndroidRuntimeSettings.AndroidRuntimeSettings", "ExtraManifestNodeTags", out ExtraManifestNodeTags);
-			List<string>? ExtraApplicationNodeTags;
-			Ini.GetArray("/Script/AndroidRuntimeSettings.AndroidRuntimeSettings", "ExtraApplicationNodeTags", out ExtraApplicationNodeTags);
-			List<string>? ExtraActivityNodeTags;
-			Ini.GetArray("/Script/AndroidRuntimeSettings.AndroidRuntimeSettings", "ExtraActivityNodeTags", out ExtraActivityNodeTags);
-			string ExtraActivitySettings;
-			Ini.GetString("/Script/AndroidRuntimeSettings.AndroidRuntimeSettings", "ExtraActivitySettings", out ExtraActivitySettings);
-			string ExtraApplicationSettings;
-			Ini.GetString("/Script/AndroidRuntimeSettings.AndroidRuntimeSettings", "ExtraApplicationSettings", out ExtraApplicationSettings);
-			List<string>? ExtraPermissions;
-			Ini.GetArray("/Script/AndroidRuntimeSettings.AndroidRuntimeSettings", "ExtraPermissions", out ExtraPermissions);
+			Ini.GetBool("/Script/AndroidRuntimeSettings.AndroidRuntimeSettings", "bFullScreen", out bool EnableFullScreen);
+			Ini.GetBool("/Script/AndroidRuntimeSettings.AndroidRuntimeSettings", "bUseDisplayCutout", out bool bUseDisplayCutout);
+			Ini.GetString("/Script/AndroidRuntimeSettings.AndroidRuntimeSettings", "RoundedEdgeSafeZone", out string RoundedEdgeSafeZone);
+			Ini.GetBool("/Script/AndroidRuntimeSettings.AndroidRuntimeSettings", "bAllowResizing", out bool bAllowResizing);
+			Ini.GetBool("/Script/AndroidRuntimeSettings.AndroidRuntimeSettings", "bSupportSizeChanges", out bool bSupportSizeChanges);
+			Ini.GetBool("/Script/AndroidRuntimeSettings.AndroidRuntimeSettings", "bRestoreNotificationsOnReboot", out bool bRestoreNotificationsOnReboot);
+			Ini.GetArray("/Script/AndroidRuntimeSettings.AndroidRuntimeSettings", "ExtraManifestNodeTags", out List<string>? ExtraManifestNodeTags);
+			Ini.GetArray("/Script/AndroidRuntimeSettings.AndroidRuntimeSettings", "ExtraApplicationNodeTags", out List<string>? ExtraApplicationNodeTags);
+			Ini.GetArray("/Script/AndroidRuntimeSettings.AndroidRuntimeSettings", "ExtraActivityNodeTags", out List<string>? ExtraActivityNodeTags);
+			Ini.GetString("/Script/AndroidRuntimeSettings.AndroidRuntimeSettings", "ExtraActivitySettings", out string ExtraActivitySettings);
+			Ini.GetString("/Script/AndroidRuntimeSettings.AndroidRuntimeSettings", "ExtraApplicationSettings", out string ExtraApplicationSettings);
+			Ini.GetArray("/Script/AndroidRuntimeSettings.AndroidRuntimeSettings", "ExtraPermissions", out List<string>? ExtraPermissions);
 			bool bPackageForMetaQuest = IsPackagingForMetaQuest(CookFlavor, Ini);
-			bool bEnableIAP = false;
-			Ini.GetBool("OnlineSubsystemGooglePlay.Store", "bSupportsInAppPurchasing", out bEnableIAP);
-			bool bShowLaunchImage = false;
-			Ini.GetBool("/Script/AndroidRuntimeSettings.AndroidRuntimeSettings", "bShowLaunchImage", out bShowLaunchImage);
-			string AndroidGraphicsDebugger;
-			Ini.GetString("/Script/AndroidRuntimeSettings.AndroidRuntimeSettings", "AndroidGraphicsDebugger", out AndroidGraphicsDebugger);
-			bool bSupportAdMob = true;
-			Ini.GetBool("/Script/AndroidRuntimeSettings.AndroidRuntimeSettings", "bSupportAdMob", out bSupportAdMob);
-			bool bValidateTextureFormats;
-			Ini.GetBool("/Script/AndroidRuntimeSettings.AndroidRuntimeSettings", "bValidateTextureFormats", out bValidateTextureFormats);
+			Ini.GetBool("OnlineSubsystemGooglePlay.Store", "bSupportsInAppPurchasing", out bool bEnableIAP);
+			Ini.GetBool("/Script/AndroidRuntimeSettings.AndroidRuntimeSettings", "bShowLaunchImage", out bool bShowLaunchImage);
+			Ini.GetString("/Script/AndroidRuntimeSettings.AndroidRuntimeSettings", "AndroidGraphicsDebugger", out string AndroidGraphicsDebugger);
+			Ini.GetBool("/Script/AndroidRuntimeSettings.AndroidRuntimeSettings", "bSupportAdMob", out bool bSupportAdMob);
+			Ini.GetBool("/Script/AndroidRuntimeSettings.AndroidRuntimeSettings", "bValidateTextureFormats", out bool bValidateTextureFormats);
 
-			bool bBuildForES31 = false;
-			Ini.GetBool("/Script/AndroidRuntimeSettings.AndroidRuntimeSettings", "bBuildForES31", out bBuildForES31);
-			bool bSupportsVulkan = false;
-			Ini.GetBool("/Script/AndroidRuntimeSettings.AndroidRuntimeSettings", "bSupportsVulkan", out bSupportsVulkan);
+			Ini.GetBool("/Script/AndroidRuntimeSettings.AndroidRuntimeSettings", "bBuildForES31", out bool bBuildForES31);
+			Ini.GetBool("/Script/AndroidRuntimeSettings.AndroidRuntimeSettings", "bSupportsVulkan", out bool bSupportsVulkan);
 
-			int PropagateAlpha = 0;
-			Ini.GetInt32("/Script/Engine.RendererSettings", "r.Mobile.PropagateAlpha", out PropagateAlpha);
+			Ini.GetInt32("/Script/Engine.RendererSettings", "r.Mobile.PropagateAlpha", out int PropagateAlpha);
 
-			bool bAllowIMU = true;
-			Ini.GetBool("/Script/AndroidRuntimeSettings.AndroidRuntimeSettings", "bAllowIMU", out bAllowIMU);
+			Ini.GetBool("/Script/AndroidRuntimeSettings.AndroidRuntimeSettings", "bAllowIMU", out bool bAllowIMU);
 
-			bool bPublicLogFiles = true;
-			Ini.GetBool("/Script/AndroidRuntimeSettings.AndroidRuntimeSettings", "bPublicLogFiles", out bPublicLogFiles);
+			Ini.GetBool("/Script/AndroidRuntimeSettings.AndroidRuntimeSettings", "bPublicLogFiles", out bool bPublicLogFiles);
 			if (!bUseExternalFilesDir)
 			{
 				bPublicLogFiles = false;
@@ -2892,11 +2851,9 @@ namespace UnrealBuildTool
 					break;
 			}
 
-			bool bEnableMulticastSupport = false;
-			Ini.GetBool("/Script/AndroidRuntimeSettings.AndroidRuntimeSettings", "bEnableMulticastSupport", out bEnableMulticastSupport);
+			Ini.GetBool("/Script/AndroidRuntimeSettings.AndroidRuntimeSettings", "bEnableMulticastSupport", out bool bEnableMulticastSupport);
 
-			bool bASISEnabled = false;
-			Ini.GetBool("/Script/AndroidSingleInstanceServiceEditor.AndroidSingleInstanceServiceRuntimeSettings", "bEnableASISPlugin", out bASISEnabled);
+			Ini.GetBool("/Script/AndroidSingleInstanceServiceEditor.AndroidSingleInstanceServiceRuntimeSettings", "bEnableASISPlugin", out bool bASISEnabled);
 
 			// only apply density to configChanges if using android-24 or higher and minimum sdk is 17
 			bool bAddDensity = (SDKLevelInt >= 24) && (MinSDKVersion >= 17);
@@ -3433,8 +3390,7 @@ namespace UnrealBuildTool
 			Ini.GetString("/Script/AndroidRuntimeSettings.AndroidRuntimeSettings", "GamesAppID", out IniAppId);
 
 			//validate the value found in the AndroidRuntimeSettings
-			long Value;
-			if (IniAppId.Length == 0 || !Int64.TryParse(IniAppId, out Value))
+			if (IniAppId.Length == 0 || !Int64.TryParse(IniAppId, out _))
 			{
 				bInvalidIniAppId = true;
 			}
@@ -3474,7 +3430,7 @@ namespace UnrealBuildTool
 						ReplacementId = IniAppId;
 
 					}
-					else if (XmlAppId.Length == 0 || !Int64.TryParse(XmlAppId, out Value))
+					else if (XmlAppId.Length == 0 || !Int64.TryParse(XmlAppId, out _))
 					{
 						Logger.LogWarning("\nWARNING: GooglePlay Games App ID is invalid! Replacing it with \"1\"");
 
@@ -3695,8 +3651,7 @@ namespace UnrealBuildTool
 				return false;
 			}
 			ConfigHierarchy Ini = GetConfigCacheIni(ConfigHierarchyType.Engine);
-			bool bEnableBundle = false;
-			Ini.GetBool("/Script/AndroidRuntimeSettings.AndroidRuntimeSettings", "bEnableBundle", out bEnableBundle);
+			Ini.GetBool("/Script/AndroidRuntimeSettings.AndroidRuntimeSettings", "bEnableBundle", out bool bEnableBundle);
 			return bEnableBundle;
 		}
 
@@ -3727,15 +3682,13 @@ namespace UnrealBuildTool
 		{
 			ConfigHierarchy Ini = GetConfigCacheIni(ConfigHierarchyType.Engine);
 			Ini.GetInt32("/Script/AndroidRuntimeSettings.AndroidRuntimeSettings", "MinSDKVersion", out MinSDKVersion);
-			TargetSDKVersion = MinSDKVersion;
 			Ini.GetInt32("/Script/AndroidRuntimeSettings.AndroidRuntimeSettings", "TargetSDKVersion", out TargetSDKVersion);
 
 			// Check for targetSDKOverride from UPL
 			string TargetOverride = UPL.ProcessPluginNode(NDKArch, "targetSDKOverride", "");
 			if (!String.IsNullOrEmpty(TargetOverride))
 			{
-				int OverrideInt = 0;
-				if (Int32.TryParse(TargetOverride, out OverrideInt))
+				if (Int32.TryParse(TargetOverride, out int OverrideInt))
 				{
 					TargetSDKVersion = OverrideInt;
 				}
@@ -4085,8 +4038,7 @@ namespace UnrealBuildTool
 
 			ConfigHierarchy Ini = GetConfigCacheIni(ConfigHierarchyType.Engine);
 
-			bool bEnableUniversalAPK = false;
-			Ini.GetBool("/Script/AndroidRuntimeSettings.AndroidRuntimeSettings", "bEnableUniversalAPK", out bEnableUniversalAPK);
+			Ini.GetBool("/Script/AndroidRuntimeSettings.AndroidRuntimeSettings", "bEnableUniversalAPK", out bool bEnableUniversalAPK);
 
 			GradleProperties.AppendLine("org.gradle.daemon=false");
 			GradleProperties.AppendLine("org.gradle.jvmargs=-XX:MaxHeapSize=4096m -Xmx9216m");
@@ -4162,8 +4114,7 @@ namespace UnrealBuildTool
 
 			GradleBuildAdditionsContent.AppendLine("android {");
 
-			bool bExtractNativeLibs = true;
-			Ini.GetBool("/Script/AndroidRuntimeSettings.AndroidRuntimeSettings", "bExtractNativeLibs", out bExtractNativeLibs);
+			Ini.GetBool("/Script/AndroidRuntimeSettings.AndroidRuntimeSettings", "bExtractNativeLibs", out bool bExtractNativeLibs);
 			AndroidToolChain.ClangSanitizer Sanitizer = ToolChain.BuildWithSanitizer();
 			// hwasan on NDK r26b+ requires wrap.sh that needs to be unpacked
 			if ((Sanitizer != AndroidToolChain.ClangSanitizer.None && (Sanitizer != AndroidToolChain.ClangSanitizer.HwAddress || ToolChain.HasEmbeddedHWASanSupport())) || bEnableScudoMemoryTracing)
@@ -4336,8 +4287,7 @@ dependencies.implementation('androidx.tracing:tracing:1.2.0')
 
 			// Add lint if requested (note depreciation warnings can be suppressed with @SuppressWarnings("deprecation")
 			string GradleBaseBuildAdditionsContents = "";
-			bool bEnableLint = false;
-			Ini.GetBool("/Script/AndroidRuntimeSettings.AndroidRuntimeSettings", "bEnableLint", out bEnableLint);
+			Ini.GetBool("/Script/AndroidRuntimeSettings.AndroidRuntimeSettings", "bEnableLint", out bool bEnableLint);
 			if (bEnableLint)
 			{
 				GradleBaseBuildAdditionsContents =
@@ -4405,8 +4355,7 @@ dependencies.implementation('androidx.tracing:tracing:1.2.0')
 				return false;
 			}
 
-			bool bJavaPerfettoMarkers = false;
-			Ini.GetBool("/Script/AndroidRuntimeSettings.AndroidRuntimeSettings", "bJavaPerfettoMarkers", out bJavaPerfettoMarkers);
+			Ini.GetBool("/Script/AndroidRuntimeSettings.AndroidRuntimeSettings", "bJavaPerfettoMarkers", out bool bJavaPerfettoMarkers);
 			return bJavaPerfettoMarkers;
 		}
 
@@ -6242,15 +6191,6 @@ popd
 			ClangToolChainOptions Options = AndroidPlatform.CreateToolChainOptions(TargetRules);
 			AndroidToolChain ToolChain = new AndroidToolChain(ProjectFile, Options, Logger);
 
-			// hunt down the receipt that compiled the .so
-			//FileReference ReceiptFilename = TargetReceipt.GetDefaultPath(ProjectDirectory, ProjectName, SC.StageTargetPlatform.IniPlatformType, Configuration, Architecture);
-			//if (!FileReference.Exists(ReceiptFilename))
-			//{
-			//	ReceiptFilename = TargetReceipt.GetDefaultPath(Unreal.EngineDirectory, ReceiptName, SC.StageTargetPlatform.IniPlatformType, Configuration, Architecture);
-			//}
-			FileReference ReceiptFilename = new FileReference(ExecutablePath).ChangeExtension(".target");
-			TargetReceipt Receipt = TargetReceipt.Read(ReceiptFilename);
-
 			SavePackageInfo(ProjectName, ProjectDirectory.FullName, Type, bSkipGradleBuild);
 
 			MakeApk(ToolChain, ProjectName, Type, ProjectDirectory.FullName, ExecutablePath, EngineDirectory, bForDistribution: bForDistribution, CookFlavor: CookFlavor, Configuration: Configuration,
@@ -6329,7 +6269,7 @@ popd
 
 			string GameActivityImportAdditionsDefault = "";
 			string GameActivityClassAdditionsDefault = "";
-			string LoadUnrealDefault = "";
+			string LoadUnrealDefault;
 			string LoadLibraryDefaults = "";
 
 			string SuperClassDefault;
diff --git a/Engine/Source/Programs/UnrealBuildTool/Platform/Apple/AppleExports.cs b/Engine/Source/Programs/UnrealBuildTool/Platform/Apple/AppleExports.cs
--- a/Engine/Source/Programs/UnrealBuildTool/Platform/Apple/AppleExports.cs
+++ b/Engine/Source/Programs/UnrealBuildTool/Platform/Apple/AppleExports.cs
@@ -374,16 +374,14 @@ namespace UnrealBuildTool
 					$"-c 'plutil -extract CFBundleShortVersionString raw $(xcode-select -p)/../Info.plist'", out ExitCode);
 				if (ExitCode == 0 && XcodeVersion != null)
 				{
-					// parse it into Major/Minor version numbers
+					// Parse the major version, we don't need the minor version.
 					int Major = 15;
-					int Minor = 0;
 					try
 					{
 						string[] Tokens = XcodeVersion!.Split(".".ToCharArray());
 						if (Tokens.Length >= 2)
 						{
 							Major = Int32.Parse(Tokens[0]);
-							Minor = Int32.Parse(Tokens[1]);
 						}
 					}
 					catch (Exception)
diff --git a/Engine/Source/Programs/UnrealBuildTool/Platform/Apple/RemoteMac.cs b/Engine/Source/Programs/UnrealBuildTool/Platform/Apple/RemoteMac.cs
--- a/Engine/Source/Programs/UnrealBuildTool/Platform/Apple/RemoteMac.cs
+++ b/Engine/Source/Programs/UnrealBuildTool/Platform/Apple/RemoteMac.cs
@@ -766,7 +766,6 @@ namespace UnrealBuildTool
 			Logger.LogInformation("[Remote] Generating an Xcode Project file...");
 			string CommandLine = EscapeShellArgument(GetRemotePath(Unreal.EngineDirectory));
 			CommandLine += "/Build/BatchFiles/Mac/GenerateProjectFiles.sh";
-			StringBuilder BuildCommandLine = new StringBuilder(CommandLine);
 
 			Execute(GetRemotePath(Unreal.RootDirectory), CommandLine, Logger);
 
@@ -814,7 +813,6 @@ namespace UnrealBuildTool
 			Logger.LogInformation("[Remote] Generating an Xcode Project file...");
 			string CommandLine = EscapeShellArgument(GetRemotePath(Unreal.EngineDirectory));
 			CommandLine += "/Build/BatchFiles/Mac/GenerateProjectFiles.sh";
-			StringBuilder BuildCommandLine = new StringBuilder(CommandLine);
 			Execute(GetRemotePath(Unreal.RootDirectory), CommandLine, Logger);
 
 			UploadDirectory(new DirectoryReference("UE5.xcworkspace"), Logger);
diff --git a/Engine/Source/Programs/UnrealBuildTool/Platform/IOS/IOSCodeSigning.cs b/Engine/Source/Programs/UnrealBuildTool/Platform/IOS/IOSCodeSigning.cs
--- a/Engine/Source/Programs/UnrealBuildTool/Platform/IOS/IOSCodeSigning.cs
+++ b/Engine/Source/Programs/UnrealBuildTool/Platform/IOS/IOSCodeSigning.cs
@@ -1798,8 +1798,6 @@ namespace UnrealBuildTool
 		/// </summary>
 		public static bool Initialize(FileReference? ProjectFile, bool bIsTVOS)
 		{
-			bool bIsEpicInternal = File.Exists(@"..\..\EpicInternal.txt");
-
 			// if the path is a directory (relative to where the game was launched from), then get the absolute directory
 			string FullGamePath;
 			string OrigGamePath;
diff --git a/Engine/Source/Programs/UnrealBuildTool/Platform/IOS/IOSToolChain.cs b/Engine/Source/Programs/UnrealBuildTool/Platform/IOS/IOSToolChain.cs
--- a/Engine/Source/Programs/UnrealBuildTool/Platform/IOS/IOSToolChain.cs
+++ b/Engine/Source/Programs/UnrealBuildTool/Platform/IOS/IOSToolChain.cs
@@ -245,17 +245,15 @@ namespace UnrealBuildTool
 
 		static void CleanIntermediateDirectory(string Path)
 		{
-			string ResultsText;
-
 			// Delete the local dest directory if it exists
 			if (Directory.Exists(Path))
 			{
 				// this can deal with linked files
-				RunExecutableAndWait("rm", String.Format("-rf \"{0}\"", Path), out ResultsText);
+				RunExecutableAndWait("rm", String.Format("-rf \"{0}\"", Path), out _);
 			}
 
 			// Create the intermediate local directory
-			RunExecutableAndWait("mkdir", String.Format("-p \"{0}\"", Path), out ResultsText);
+			RunExecutableAndWait("mkdir", String.Format("-p \"{0}\"", Path), out _);
 		}
 
 		public string GetAdditionalLinkerFlags(CppConfiguration InConfiguration)
@@ -394,7 +392,6 @@ namespace UnrealBuildTool
 			if (LinkEnvironment.bIsBuildingDLL)
 			{
 				// @todo roll this put into GetLinkArguments_Global
-				LinkerPath = Info.Clang;
 				LinkArguments.Add("-dynamiclib -Xlinker -export_dynamic -Xlinker -no_deduplicate");
 
 				string InstallName = LinkEnvironment.InstallName ?? $"@executable_path/Frameworks/{LinkEnvironment.OutputFilePath.MakeRelativeTo(LinkEnvironment.OutputFilePath.Directory.ParentDirectory!)}";
diff --git a/Engine/Source/Programs/UnrealBuildTool/Platform/IOS/PListHelper.cs b/Engine/Source/Programs/UnrealBuildTool/Platform/IOS/PListHelper.cs
--- a/Engine/Source/Programs/UnrealBuildTool/Platform/IOS/PListHelper.cs
+++ b/Engine/Source/Programs/UnrealBuildTool/Platform/IOS/PListHelper.cs
@@ -208,7 +208,7 @@ namespace UnrealBuildTool
 
 			if (dictRoot != null)
 			{
-				XmlNode? childNode = Doc.CreateElement("null");
+				XmlNode? childNode;
 				if (Value != null)
 				{
 					childNode = ConvertValueToPListFormat(Value);
@@ -431,8 +431,6 @@ namespace UnrealBuildTool
 				throw new AccessViolationException("PList has been set to read only and may not be modified");
 			}
 
-			XmlNode? DictionaryNode = Doc.DocumentElement?.SelectSingleNode("/plist/dict");
-
 			string PathToKey = String.Format("/plist/dict/key[.='{0}']", KeyName);
 			XmlNode? KeyNode = Doc.DocumentElement?.SelectSingleNode(PathToKey);
 			if (KeyNode != null && KeyNode.ParentNode != null)
diff --git a/Engine/Source/Programs/UnrealBuildTool/Platform/IOS/UEBuildIOS.cs b/Engine/Source/Programs/UnrealBuildTool/Platform/IOS/UEBuildIOS.cs
--- a/Engine/Source/Programs/UnrealBuildTool/Platform/IOS/UEBuildIOS.cs
+++ b/Engine/Source/Programs/UnrealBuildTool/Platform/IOS/UEBuildIOS.cs
@@ -461,8 +461,9 @@ namespace UnrealBuildTool
 
 			if (!String.IsNullOrEmpty(SigningCertificate))
 			{
-				List<string> Certs = AppleCodeSign.FindCertificates();
-				List<FileReference> Provisions = AppleCodeSign.FindProvisions(ProjectSettings.BundleIdentifier, bForDistribution, out _);
+				// TODO: Verify if there are side-effects that require these methods to run.
+				_ = AppleCodeSign.FindCertificates();
+				_ = AppleCodeSign.FindProvisions(ProjectSettings.BundleIdentifier, bForDistribution, out _);
 			}
 			else
 			{
@@ -483,14 +484,10 @@ namespace UnrealBuildTool
 			{
 
 				SigningCertificate = "";
-				MobileProvision = "";
 				MobileProvisionFile = null;
 				Logger.LogInformation("Provision not specified or not found for {Project}, searching for compatible match...", ((ProjectFile != null) ? ProjectFile.GetFileNameWithoutAnyExtensions() : "UnrealGame"));
 
-				if (AppleCodeSign.FindCertAndProvision(ProjectSettings.BundleIdentifier, out MobileProvisionFile, out SigningCertificate))
-				{
-					MobileProvision = MobileProvisionFile!.FullName;
-				}
+				_ = AppleCodeSign.FindCertAndProvision(ProjectSettings.BundleIdentifier, out MobileProvisionFile, out SigningCertificate);
 
 				if (MobileProvisionFile != null)
 				{
@@ -589,7 +586,6 @@ namespace UnrealBuildTool
 
 				if (String.IsNullOrEmpty(MobileProvisionUUID) || String.IsNullOrEmpty(TeamUUID))
 				{
-					MobileProvision = null;
 					SigningCertificate = null;
 					Logger.LogInformation("Failed to parse the mobile provisioning profile.");
 				}
diff --git a/Engine/Source/Programs/UnrealBuildTool/Platform/IOS/UEDeployIOS.cs b/Engine/Source/Programs/UnrealBuildTool/Platform/IOS/UEDeployIOS.cs
--- a/Engine/Source/Programs/UnrealBuildTool/Platform/IOS/UEDeployIOS.cs
+++ b/Engine/Source/Programs/UnrealBuildTool/Platform/IOS/UEDeployIOS.cs
@@ -117,15 +117,14 @@ namespace UnrealBuildTool
 				int RunningMajorVersion;
 				int RunningMinorVersion;
 
-				string DummyDate;
 				string RunningVersion = ReadRunningVersion();
-				PullApartVersion(RunningVersion, out RunningMajorVersion, out RunningMinorVersion, out DummyDate);
+				PullApartVersion(RunningVersion, out RunningMajorVersion, out RunningMinorVersion, out _);
 				RunningMinorVersion++;
 
 				// Read the passed in version and bump it
 				int MajorVersion;
 				int MinorVersion;
-				PullApartVersion(CFBundleVersion, out MajorVersion, out MinorVersion, out DummyDate);
+				PullApartVersion(CFBundleVersion, out MajorVersion, out MinorVersion, out _);
 				MinorVersion++;
 
 				// Combine them if the stub time is older
@@ -151,7 +150,7 @@ namespace UnrealBuildTool
 			/// </summary>
 			public static string UpdateBundleVersion(string OldPList, string EngineDirectory)
 			{
-				string CFBundleVersion = "-1";
+				string CFBundleVersion;
 				if (!Unreal.IsBuildMachine())
 				{
 					int Index = OldPList.IndexOf("CFBundleVersion");
@@ -196,7 +195,7 @@ namespace UnrealBuildTool
 
 		public static string GetMinimumOSVersion(string MinVersion, ILogger Logger)
 		{
-			string MinVersionToReturn = "";
+			string MinVersionToReturn;
 			switch (MinVersion)
 			{
 				case "":
@@ -226,24 +225,11 @@ namespace UnrealBuildTool
 
 			// orientations
 			string InterfaceOrientation = "";
-			string PreferredLandscapeOrientation = "";
-			Ini.GetString("/Script/IOSRuntimeSettings.IOSRuntimeSettings", "PreferredLandscapeOrientation", out PreferredLandscapeOrientation);
-
-			string SupportedOrientations = "";
-			bool bSupported = true;
-			Ini.GetBool("/Script/IOSRuntimeSettings.IOSRuntimeSettings", "bSupportsPortraitOrientation", out bSupported);
-			SupportedOrientations += bSupported ? "\t\t<string>UIInterfaceOrientationPortrait</string>\n" : "";
-			bool bSupportsPortrait = bSupported;
-
-			Ini.GetBool("/Script/IOSRuntimeSettings.IOSRuntimeSettings", "bSupportsUpsideDownOrientation", out bSupported);
-			SupportedOrientations += bSupported ? "\t\t<string>UIInterfaceOrientationPortraitUpsideDown</string>\n" : "";
-			bSupportsPortrait |= bSupported;
-
-			bool bSupportsLandscapeLeft = false;
-			Ini.GetBool("/Script/IOSRuntimeSettings.IOSRuntimeSettings", "bSupportsLandscapeLeftOrientation", out bSupportsLandscapeLeft);
-			bool bSupportsLandscapeRight = false;
-			Ini.GetBool("/Script/IOSRuntimeSettings.IOSRuntimeSettings", "bSupportsLandscapeRightOrientation", out bSupportsLandscapeRight);
-			bool bSupportsLandscape = bSupportsLandscapeLeft || bSupportsLandscapeRight;
+			Ini.GetString("/Script/IOSRuntimeSettings.IOSRuntimeSettings", "PreferredLandscapeOrientation", out string PreferredLandscapeOrientation);
+			Ini.GetBool("/Script/IOSRuntimeSettings.IOSRuntimeSettings", "bSupportsPortraitOrientation", out bool bSupportsPortrait);
+			Ini.GetBool("/Script/IOSRuntimeSettings.IOSRuntimeSettings", "bSupportsUpsideDownOrientation", out bool bSupportsUpsideDown);
+			Ini.GetBool("/Script/IOSRuntimeSettings.IOSRuntimeSettings", "bSupportsLandscapeLeftOrientation", out bool bSupportsLandscapeLeft);
+			Ini.GetBool("/Script/IOSRuntimeSettings.IOSRuntimeSettings", "bSupportsLandscapeRightOrientation", out bool bSupportsLandscapeRight);
 
 			if (bSupportsLandscapeLeft && bSupportsLandscapeRight)
 			{
@@ -252,31 +238,20 @@ namespace UnrealBuildTool
 				if (PreferredLandscapeOrientation == "LandscapeLeft")
 				{
 					InterfaceOrientation = "\t<key>UIInterfaceOrientation</key>\n\t<string>UIInterfaceOrientationLandscapeLeft</string>\n";
-					SupportedOrientations += "\t\t<string>UIInterfaceOrientationLandscapeLeft</string>\n\t\t<string>UIInterfaceOrientationLandscapeRight</string>\n";
 				}
 				else
 				{
 					// by default, landscape right is the preferred orientation - Apple's UI guidlines
 					InterfaceOrientation = "\t<key>UIInterfaceOrientation</key>\n\t<string>UIInterfaceOrientationLandscapeRight</string>\n";
-					SupportedOrientations += "\t\t<string>UIInterfaceOrientationLandscapeRight</string>\n\t\t<string>UIInterfaceOrientationLandscapeLeft</string>\n";
 				}
 			}
-			else
-			{
-				// max one landscape orientation is supported
-				SupportedOrientations += bSupportsLandscapeRight ? "\t\t<string>UIInterfaceOrientationLandscapeRight</string>\n" : "";
-				SupportedOrientations += bSupportsLandscapeLeft ? "\t\t<string>UIInterfaceOrientationLandscapeLeft</string>\n" : "";
-			}
 
 			// ITunes file sharing
-			bool bSupportsITunesFileSharing = false;
-			Ini.GetBool("/Script/IOSRuntimeSettings.IOSRuntimeSettings", "bSupportsITunesFileSharing", out bSupportsITunesFileSharing);
-			bool bSupportsFilesApp = false;
-			Ini.GetBool("/Script/IOSRuntimeSettings.IOSRuntimeSettings", "bSupportsFilesApp", out bSupportsFilesApp);
+			Ini.GetBool("/Script/IOSRuntimeSettings.IOSRuntimeSettings", "bSupportsITunesFileSharing", out bool bSupportsITunesFileSharing);
+			Ini.GetBool("/Script/IOSRuntimeSettings.IOSRuntimeSettings", "bSupportsFilesApp", out bool bSupportsFilesApp);
 
 			// disable https requirement
-			bool bDisableHTTPS;
-			Ini.GetBool("/Script/IOSRuntimeSettings.IOSRuntimeSettings", "bDisableHTTPS", out bDisableHTTPS);
+			Ini.GetBool("/Script/IOSRuntimeSettings.IOSRuntimeSettings", "bDisableHTTPS", out bool bDisableHTTPS);
 
 			bool bUseZenStore = false;
 			if (Config != UnrealTargetConfiguration.Shipping)
@@ -286,43 +261,30 @@ namespace UnrealBuildTool
 			}
 
 			// bundle display name
-			string BundleDisplayName;
-			Ini.GetString("/Script/IOSRuntimeSettings.IOSRuntimeSettings", "BundleDisplayName", out BundleDisplayName);
-
-			// short version string
-			string BundleShortVersion;
-			Ini.GetString("/Script/IOSRuntimeSettings.IOSRuntimeSettings", "VersionInfo", out BundleShortVersion);
+			Ini.GetString("/Script/IOSRuntimeSettings.IOSRuntimeSettings", "BundleDisplayName", out string BundleDisplayName);
 
 			// Get Google Support details
-			bool bEnableGoogleSupport = true;
-			Ini.GetBool("/Script/IOSRuntimeSettings.IOSRuntimeSettings", "bEnableGoogleSupport", out bEnableGoogleSupport);
+			Ini.GetBool("/Script/IOSRuntimeSettings.IOSRuntimeSettings", "bEnableGoogleSupport", out bool bEnableGoogleSupport);
 
 			// Write the Google iOS URL Scheme if we need it.
-			string GoogleReversedClientId = "";
-			Ini.GetString("/Script/IOSRuntimeSettings.IOSRuntimeSettings", "GoogleReversedClientId", out GoogleReversedClientId);
+			Ini.GetString("/Script/IOSRuntimeSettings.IOSRuntimeSettings", "GoogleReversedClientId", out string GoogleReversedClientId);
 			bEnableGoogleSupport = bEnableGoogleSupport && !String.IsNullOrWhiteSpace(GoogleReversedClientId);
 
 			// Add remote-notifications as background mode
-			bool bRemoteNotificationsSupported = false;
-			Ini.GetBool("/Script/IOSRuntimeSettings.IOSRuntimeSettings", "bEnableRemoteNotificationsSupport", out bRemoteNotificationsSupported);
+			Ini.GetBool("/Script/IOSRuntimeSettings.IOSRuntimeSettings", "bEnableRemoteNotificationsSupport", out bool bRemoteNotificationsSupported);
 
 			// Add audio as background mode
-			bool bBackgroundAudioSupported = false;
-			Ini.GetBool("/Script/IOSRuntimeSettings.IOSRuntimeSettings", "bSupportsBackgroundAudio", out bBackgroundAudioSupported);
+			Ini.GetBool("/Script/IOSRuntimeSettings.IOSRuntimeSettings", "bSupportsBackgroundAudio", out bool bBackgroundAudioSupported);
 
 			// Add background fetch as background mode
-			bool bBackgroundFetch = false;
-			Ini.GetBool("/Script/IOSRuntimeSettings.IOSRuntimeSettings", "bEnableBackgroundFetch", out bBackgroundFetch);
+			Ini.GetBool("/Script/IOSRuntimeSettings.IOSRuntimeSettings", "bEnableBackgroundFetch", out bool bBackgroundFetch);
 
 			// Get any Location Services permission descriptions added
-			string LocationAlwaysUsageDescription = "";
-			string LocationWhenInUseDescription = "";
-			Ini.GetString("/Script/LocationServicesIOSEditor.LocationServicesIOSSettings", "LocationAlwaysUsageDescription", out LocationAlwaysUsageDescription);
-			Ini.GetString("/Script/LocationServicesIOSEditor.LocationServicesIOSSettings", "LocationWhenInUseDescription", out LocationWhenInUseDescription);
+			Ini.GetString("/Script/LocationServicesIOSEditor.LocationServicesIOSSettings", "LocationAlwaysUsageDescription", out string LocationAlwaysUsageDescription);
+			Ini.GetString("/Script/LocationServicesIOSEditor.LocationServicesIOSSettings", "LocationWhenInUseDescription", out string LocationWhenInUseDescription);
 
 			// extra plist data
-			string ExtraData = "";
-			Ini.GetString("/Script/IOSRuntimeSettings.IOSRuntimeSettings", "AdditionalPlistData", out ExtraData);
+			Ini.GetString("/Script/IOSRuntimeSettings.IOSRuntimeSettings", "AdditionalPlistData", out string ExtraData);
 
 			Ini.GetBool("/Script/IOSRuntimeSettings.IOSRuntimeSettings", "bCustomLaunchscreenStoryboard", out VersionUtilities.bCustomLaunchscreenStoryboard);
 
@@ -384,18 +346,15 @@ namespace UnrealBuildTool
 
 			// Support high refresh rates (iPhone only)
 			// https://developer.apple.com/documentation/quartzcore/optimizing_promotion_refresh_rates_for_iphone_13_pro_and_ipad_pro
-			bool bSupportHighRefreshRates = false;
-			Ini.GetBool("/Script/IOSRuntimeSettings.IOSRuntimeSettings", "bSupportHighRefreshRates", out bSupportHighRefreshRates);
+			Ini.GetBool("/Script/IOSRuntimeSettings.IOSRuntimeSettings", "bSupportHighRefreshRates", out bool bSupportHighRefreshRates);
 			if (bSupportHighRefreshRates)
 			{
 				Text.AppendLine("\t<key>CADisableMinimumFrameDurationOnPhone</key><true/>");
 			}
 
 			// set exempt encryption
-			bool bUsesNonExemptEncryption = false;
-			string ITSEncryptionExportComplianceCode = "";
-			Ini.GetBool("/Script/IOSRuntimeSettings.IOSRuntimeSettings", "bUsesNonExemptEncryption", out bUsesNonExemptEncryption);
-			Ini.GetString("/Script/IOSRuntimeSettings.IOSRuntimeSettings", "ITSEncryptionExportComplianceCode", out ITSEncryptionExportComplianceCode);
+			Ini.GetBool("/Script/IOSRuntimeSettings.IOSRuntimeSettings", "bUsesNonExemptEncryption", out bool bUsesNonExemptEncryption);
+			Ini.GetString("/Script/IOSRuntimeSettings.IOSRuntimeSettings", "ITSEncryptionExportComplianceCode", out string ITSEncryptionExportComplianceCode);
 			Text.AppendLine("\t<key>ITSAppUsesNonExemptEncryption</key>");
 			Text.AppendLine(String.Format("\t<{0}/>", bUsesNonExemptEncryption ? "true" : "false"));
 			if (bUsesNonExemptEncryption && !String.IsNullOrWhiteSpace(ITSEncryptionExportComplianceCode))
diff --git a/Engine/Source/Programs/UnrealBuildTool/ProjectFiles/VisualStudio/VCProjectFileGenerator.cs b/Engine/Source/Programs/UnrealBuildTool/ProjectFiles/VisualStudio/VCProjectFileGenerator.cs
--- a/Engine/Source/Programs/UnrealBuildTool/ProjectFiles/VisualStudio/VCProjectFileGenerator.cs
+++ b/Engine/Source/Programs/UnrealBuildTool/ProjectFiles/VisualStudio/VCProjectFileGenerator.cs
@@ -256,7 +256,6 @@ namespace UnrealBuildTool
 
 		public static void AppendPlatformToolsetProperty(StringBuilder VCProjectFileContent, VCProjectFileFormat ProjectFileFormat)
 		{
-			string ToolVersionString = GetProjectFileToolVersionString(ProjectFileFormat);
 			string PlatformToolsetVersionString = GetProjectFilePlatformToolsetVersionString(ProjectFileFormat);
 			VCProjectFileContent.AppendLine("    <PlatformToolset>{0}</PlatformToolset>", PlatformToolsetVersionString);
 		}
diff --git a/Engine/Source/Programs/UnrealBuildTool/ProjectFiles/VisualStudio/VCSolutionOptions.cs b/Engine/Source/Programs/UnrealBuildTool/ProjectFiles/VisualStudio/VCSolutionOptions.cs
--- a/Engine/Source/Programs/UnrealBuildTool/ProjectFiles/VisualStudio/VCSolutionOptions.cs
+++ b/Engine/Source/Programs/UnrealBuildTool/ProjectFiles/VisualStudio/VCSolutionOptions.cs
@@ -224,8 +224,7 @@ namespace UnrealBuildTool
 
 		public VCOleContainer(string InputFileName)
 		{
-			IOleStorage? Storage = null;
-			int HResult = StgOpenStorage(InputFileName, null, STGM.Direct | STGM.Read | STGM.ShareExclusive, IntPtr.Zero, 0, out Storage);
+			int HResult = StgOpenStorage(InputFileName, null, STGM.Direct | STGM.Read | STGM.ShareExclusive, IntPtr.Zero, 0, out IOleStorage? Storage);
 			if (HResult != 0)
 			{
 				throw new COMException("Call to StgOpenStorage failed", HResult);
@@ -271,8 +270,7 @@ namespace UnrealBuildTool
 
 		public void Write(string OutputFileName)
 		{
-			IOleStorage? OleStorage = null;
-			int HResult = StgCreateDocfile(OutputFileName, STGM.Direct | STGM.Create | STGM.Write | STGM.ShareExclusive, 0, out OleStorage);
+			int HResult = StgCreateDocfile(OutputFileName, STGM.Direct | STGM.Create | STGM.Write | STGM.ShareExclusive, 0, out IOleStorage? OleStorage);
 			if (HResult != 0)
 			{
 				throw new COMException("Call to StgCreateDocfile failed", HResult);
diff --git a/Engine/Source/Programs/UnrealBuildTool/ProjectFiles/VisualStudioWorkspace/VSWorkspaceProjectFile.cs b/Engine/Source/Programs/UnrealBuildTool/ProjectFiles/VisualStudioWorkspace/VSWorkspaceProjectFile.cs
--- a/Engine/Source/Programs/UnrealBuildTool/ProjectFiles/VisualStudioWorkspace/VSWorkspaceProjectFile.cs
+++ b/Engine/Source/Programs/UnrealBuildTool/ProjectFiles/VisualStudioWorkspace/VSWorkspaceProjectFile.cs
@@ -190,7 +190,6 @@ namespace UnrealBuildTool
 			Func<UnrealTargetConfiguration, UnrealTargetPlatform, ProjectTarget, Task> Handler)
 		{
 			List<Task> Tasks = new List<Task>();
-			List<(UnrealTargetConfiguration, UnrealTargetPlatform, ProjectTarget)> ret = new List<(UnrealTargetConfiguration, UnrealTargetPlatform, ProjectTarget)>();
 			foreach (UnrealTargetConfiguration Configuration in InConfigurations)
 			{
 				foreach (UnrealTargetPlatform Platform in InPlatforms)
diff --git a/Engine/Source/Programs/UnrealBuildTool/ProjectFiles/Xcode/XcodeProject.Utils.cs b/Engine/Source/Programs/UnrealBuildTool/ProjectFiles/Xcode/XcodeProject.Utils.cs
--- a/Engine/Source/Programs/UnrealBuildTool/ProjectFiles/Xcode/XcodeProject.Utils.cs
+++ b/Engine/Source/Programs/UnrealBuildTool/ProjectFiles/Xcode/XcodeProject.Utils.cs
@@ -236,7 +236,7 @@ namespace UnrealBuildTool.XcodeProjectXcconfig
 		{
 			List<string> Orientations = new();
 
-			bool bSupported = true;
+			bool bSupported;
 			if (Ini.TryGetValue("/Script/IOSRuntimeSettings.IOSRuntimeSettings", "bSupportsPortraitOrientation", out bSupported) && bSupported)
 			{
 				Orientations.Add("UIInterfaceOrientationPortrait");
@@ -246,12 +246,9 @@ namespace UnrealBuildTool.XcodeProjectXcconfig
 				Orientations.Add("UIInterfaceOrientationPortraitUpsideDown");
 			}
 
-			string? PreferredLandscapeOrientation;
-			Ini.TryGetValue("/Script/IOSRuntimeSettings.IOSRuntimeSettings", "PreferredLandscapeOrientation", out PreferredLandscapeOrientation);
-			bool bSupportsLandscapeLeft = false;
-			Ini.TryGetValue("/Script/IOSRuntimeSettings.IOSRuntimeSettings", "bSupportsLandscapeLeftOrientation", out bSupportsLandscapeLeft);
-			bool bSupportsLandscapeRight = false;
-			Ini.TryGetValue("/Script/IOSRuntimeSettings.IOSRuntimeSettings", "bSupportsLandscapeRightOrientation", out bSupportsLandscapeRight);
+			Ini.TryGetValue("/Script/IOSRuntimeSettings.IOSRuntimeSettings", "PreferredLandscapeOrientation", out string? PreferredLandscapeOrientation);
+			Ini.TryGetValue("/Script/IOSRuntimeSettings.IOSRuntimeSettings", "bSupportsLandscapeLeftOrientation", out bool bSupportsLandscapeLeft);
+			Ini.TryGetValue("/Script/IOSRuntimeSettings.IOSRuntimeSettings", "bSupportsLandscapeRightOrientation", out bool bSupportsLandscapeRight);
 
 			if (bSupportsLandscapeLeft && PreferredLandscapeOrientation == "LandscapeLeft")
 			{
diff --git a/Engine/Source/Programs/UnrealBuildTool/ProjectFiles/Xcode/XcodeProject.cs b/Engine/Source/Programs/UnrealBuildTool/ProjectFiles/Xcode/XcodeProject.cs
--- a/Engine/Source/Programs/UnrealBuildTool/ProjectFiles/Xcode/XcodeProject.cs
+++ b/Engine/Source/Programs/UnrealBuildTool/ProjectFiles/Xcode/XcodeProject.cs
@@ -1455,11 +1455,6 @@ namespace UnrealBuildTool.XcodeProjectXcconfig
 							// assume each Extension in here will create a resulting Extension.appex
 							string Extension = DI.Name + ".appex";
 							string ExtensionGuid = XcodeProjectFileGenerator.MakeXcodeGuid();
-
-							// is it an extension we support? Grab the Plist file for it and check
-							// Also, only some AppExtensions need to embed
-							string AppExtPlist = DI.FullName + "/Info.plist";
-							PListHelper PlistDoc = PListHelper.CreateFromFile(AppExtPlist);
 							
 							XcodeCopyAppExtFilesBuildPhase EmbedAppExt = new XcodeCopyAppExtFilesBuildPhase(Project.FileCollection, ExtConfig);
 							
@@ -1945,9 +1940,8 @@ namespace UnrealBuildTool.XcodeProjectXcconfig
 
 		private string? CopyAndWriteEntitlements(ref MetadataItem? EntitlementsMetadata, string ConfigDisplayName, bool bUseEntitlementsForPrograms, DirectoryReference ProjectOrEngineDir)
 		{
-			bool bEnableContactsManager = false;
 			ConfigHierarchy PlatformGameConfig = ConfigCache.ReadHierarchy(ConfigHierarchyType.Game, UnrealData.ConfigDirectory, Platform);
-			PlatformGameConfig.GetBool("/Script/AppleTargetPlatform.AppleTargetSettings", "bEnableContactsManager", out bEnableContactsManager);
+			PlatformGameConfig.GetBool("/Script/AppleTargetPlatform.AppleTargetSettings", "bEnableContactsManager", out bool bEnableContactsManager);
 
 			// programs almost never want entitlements, so handle them specially
 			if ((UnrealData.bNoEntitlements) || (UnrealData.TargetRules.Type == TargetType.Program && !bUseEntitlementsForPrograms))
diff --git a/Engine/Source/Programs/UnrealBuildTool/ProjectFiles/Xcode/XcodeProjectLegacy.cs b/Engine/Source/Programs/UnrealBuildTool/ProjectFiles/Xcode/XcodeProjectLegacy.cs
--- a/Engine/Source/Programs/UnrealBuildTool/ProjectFiles/Xcode/XcodeProjectLegacy.cs
+++ b/Engine/Source/Programs/UnrealBuildTool/ProjectFiles/Xcode/XcodeProjectLegacy.cs
@@ -1544,7 +1544,6 @@ namespace UnrealBuildTool.XcodeProjectLegacy
 				Result.AppendFormat("-{0}-{1}", Platform.ToString(), Configuration.ToString());
 			}
 
-			UEBuildPlatform BuildPlatform = UEBuildPlatform.GetBuildPlatform(Platform);
 			if (UnrealArchitectureConfig.ForPlatform(Platform).RequiresArchitectureFilenames(Architectures))
 			{
 				Result.Append(Architectures.ToString());
diff --git a/Engine/Source/Programs/UnrealBuildTool/System/ConfigFile.cs b/Engine/Source/Programs/UnrealBuildTool/System/ConfigFile.cs
--- a/Engine/Source/Programs/UnrealBuildTool/System/ConfigFile.cs
+++ b/Engine/Source/Programs/UnrealBuildTool/System/ConfigFile.cs
@@ -348,8 +348,7 @@ namespace UnrealBuildTool
 
 				if (SectionName.Length > 0)
 				{
-					ConfigFileSection? CurrentSection = null;
-					if (!Sections.TryGetValue(SectionName, out CurrentSection))
+					if (!Sections.TryGetValue(SectionName, out ConfigFileSection? CurrentSection))
 					{
 						CurrentSection = new ConfigFileSection(SectionName);
 						Sections.Add(SectionName, CurrentSection);
diff --git a/Engine/Source/Programs/UnrealBuildTool/System/CppRootPaths.cs b/Engine/Source/Programs/UnrealBuildTool/System/CppRootPaths.cs
--- a/Engine/Source/Programs/UnrealBuildTool/System/CppRootPaths.cs
+++ b/Engine/Source/Programs/UnrealBuildTool/System/CppRootPaths.cs
@@ -312,7 +312,7 @@ namespace UnrealBuildTool
 				return reference;
 			}
 
-			foreach ((CppRootPathFolder root, (DirectoryReference vfs, DirectoryReference local)) in _lookup)
+			foreach ((_, (DirectoryReference vfs, DirectoryReference local)) in _lookup)
 			{
 				if (reference.IsUnderDirectory(vfs))
 				{
@@ -378,7 +378,7 @@ namespace UnrealBuildTool
 			if (_lookup.Count > 0)
 			{
 				writer.WriteArrayStart("Lookup");
-				foreach ((CppRootPathFolder root, (DirectoryReference vfs, DirectoryReference local)) in _lookup)
+				foreach ((CppRootPathFolder root, (_, DirectoryReference local)) in _lookup)
 				{
 					writer.WriteObjectStart();
 					writer.WriteEnumValue("Root", root);
diff --git a/Engine/Source/Programs/UnrealBuildTool/System/TargetReceipt.cs b/Engine/Source/Programs/UnrealBuildTool/System/TargetReceipt.cs
--- a/Engine/Source/Programs/UnrealBuildTool/System/TargetReceipt.cs
+++ b/Engine/Source/Programs/UnrealBuildTool/System/TargetReceipt.cs
@@ -649,9 +649,6 @@ namespace UnrealBuildTool
 					{
 						FileReference File = ExpandPathVariables(Path, EngineDir, ProjectDir);
 
-						string? Module;
-						BuildProductObject.TryGetStringField("Module", out Module);
-
 						Receipt.AddBuildProduct(File, Type);
 					}
 				}
diff --git a/Engine/Source/Programs/UnrealBuildTool/System/UniqueStringRegistry.cs b/Engine/Source/Programs/UnrealBuildTool/System/UniqueStringRegistry.cs
--- a/Engine/Source/Programs/UnrealBuildTool/System/UniqueStringRegistry.cs
+++ b/Engine/Source/Programs/UnrealBuildTool/System/UniqueStringRegistry.cs
@@ -43,8 +43,7 @@ namespace UnrealBuildTool
 		public int FindOrAddByName(string Name)
 		{
 			// look for existing one
-			int Instance = -1;
-			if (!StringToInstanceMap.TryGetValue(Name, out Instance))
+			if (!StringToInstanceMap.TryGetValue(Name, out int Instance))
 			{
 				lock (LockObject)
 				{
diff --git a/Engine/Source/Programs/UnrealBuildTool/System/UnrealPluginLanguage.cs b/Engine/Source/Programs/UnrealBuildTool/System/UnrealPluginLanguage.cs
--- a/Engine/Source/Programs/UnrealBuildTool/System/UnrealPluginLanguage.cs
+++ b/Engine/Source/Programs/UnrealBuildTool/System/UnrealPluginLanguage.cs
@@ -908,8 +908,7 @@ namespace UnrealBuildTool
 
 		private int StringToInt(UPLContext Context, XElement Node, string? Input)
 		{
-			int Result = 0;
-			if (!Int32.TryParse(Input, out Result))
+			if (!Int32.TryParse(Input, out int Result))
 			{
 				Logger.LogWarning("\nInvalid integer '{Input}' in '{Node}' (defaulting to 0)", Input, TraceNodeString(Context, Node));
 			}
@@ -953,8 +952,7 @@ namespace UnrealBuildTool
 		private ConfigCacheIni_UPL GetConfigCacheIni_UPL(string baseIniName)
 		{
 			ConfigCache ??= new Dictionary<string, ConfigCacheIni_UPL>();
-			ConfigCacheIni_UPL? config = null;
-			if (!ConfigCache.TryGetValue(baseIniName, out config))
+			if (!ConfigCache.TryGetValue(baseIniName, out ConfigCacheIni_UPL? config))
 			{
 				// note: use our own ConfigCacheIni since EngineConfiguration.cs only parses RequiredSections!
 				config = ConfigCacheIni_UPL.CreateConfigCacheIni_UPL(TargetPlatform, baseIniName, DirectoryReference.FromFile(ProjectFile), Logger);
@@ -1226,8 +1224,7 @@ namespace UnrealBuildTool
 
 					case "isDistribution":
 						{
-							bool Result = false;
-							if (GetCondition(CurrentContext, Node, "Distribution", out Result))
+							if (GetCondition(CurrentContext, Node, "Distribution", out bool Result))
 							{
 								if (Result)
 								{
diff --git a/Engine/Source/Programs/UnrealBuildTool/System/Utils.cs b/Engine/Source/Programs/UnrealBuildTool/System/Utils.cs
--- a/Engine/Source/Programs/UnrealBuildTool/System/Utils.cs
+++ b/Engine/Source/Programs/UnrealBuildTool/System/Utils.cs
@@ -413,8 +413,7 @@ namespace UnrealBuildTool
 		/// <param name="Logger">Logger for output</param>
 		public static string RunLocalProcessAndReturnStdOut(string Command, string Args, ILogger? Logger)
 		{
-			int ExitCode;
-			return RunLocalProcessAndReturnStdOut(Command, Args, Logger, out ExitCode);
+			return RunLocalProcessAndReturnStdOut(Command, Args, Logger, out _);
 		}
 
 		/// <summary>
@@ -1247,9 +1246,8 @@ namespace UnrealBuildTool
 					return FreeMemoryBytes;
 				}
 
-				int PageSize = 0;
 				ulong OutSize = 4;
-				if (0 != sysctlbyname("hw.pagesize", out PageSize, ref OutSize, IntPtr.Zero, 0))
+				if (0 != sysctlbyname("hw.pagesize", out int PageSize, ref OutSize, IntPtr.Zero, 0))
 				{
 					PageSize = 4096; // likely result
 				}
