Fix issue IDE0052 (Remove unread private member) in UBT.

diff --git a/Engine/Source/Programs/UnrealBuildTool/Modes/AggregateParsedTimingInfo.cs b/Engine/Source/Programs/UnrealBuildTool/Modes/AggregateParsedTimingInfo.cs
--- a/Engine/Source/Programs/UnrealBuildTool/Modes/AggregateParsedTimingInfo.cs
+++ b/Engine/Source/Programs/UnrealBuildTool/Modes/AggregateParsedTimingInfo.cs
@@ -181,8 +181,6 @@ namespace UnrealBuildTool
 			}
 		}
 
-		private int HeaderFilesFound = 0;
-		private readonly object CounterLock = new object();
 		private void FlattenIncludes(Dictionary<string, TimingData> Includes, IEnumerable<TimingData> IncludeData)
 		{
 			foreach (TimingData Include in IncludeData)
@@ -195,14 +193,6 @@ namespace UnrealBuildTool
 				}
 				else
 				{
-					if (Include.Name.EndsWith(".h"))
-					{
-						lock (CounterLock)
-						{
-							HeaderFilesFound += 1;
-						}
-					}
-
 					Includes.Add(Include.Name, new TimingData(Include.Name, TimingDataType.Include) { ExclusiveDuration = Include.ExclusiveDuration });
 				}
 
diff --git a/Engine/Source/Programs/UnrealBuildTool/Modes/ProfileUnitySizesMode.cs b/Engine/Source/Programs/UnrealBuildTool/Modes/ProfileUnitySizesMode.cs
--- a/Engine/Source/Programs/UnrealBuildTool/Modes/ProfileUnitySizesMode.cs
+++ b/Engine/Source/Programs/UnrealBuildTool/Modes/ProfileUnitySizesMode.cs
@@ -44,14 +44,13 @@ namespace UnrealBuildTool
 			static readonly Regex NumFilesRegex = new Regex(@"\[\d+/(\d+)\]");
 			static readonly Regex CPUTimingRegex = new Regex(@"CPU Time:\s(\d+\.*\d*)");
 			private readonly ILogger Inner;
-			private readonly ILogger OldLogger;
 
 			public TimingData TimingData = new();
 
 			public TimingLogger(ILogger inner)
 			{
 				Inner = inner;
-				OldLogger = EpicGames.Core.Log.EventParser.Logger;
+				// TODO: Assigning to a global static inside this ctor seems dangerous, is this needed?
 				EpicGames.Core.Log.EventParser.Logger = this;
 			}
 
diff --git a/Engine/Source/Programs/UnrealBuildTool/Platform/Android/AndroidToolChain.cs b/Engine/Source/Programs/UnrealBuildTool/Platform/Android/AndroidToolChain.cs
--- a/Engine/Source/Programs/UnrealBuildTool/Platform/Android/AndroidToolChain.cs
+++ b/Engine/Source/Programs/UnrealBuildTool/Platform/Android/AndroidToolChain.cs
@@ -1378,68 +1378,9 @@ namespace UnrealBuildTool
 
 			Logger.LogInformation("Link: {LinkActionCommandPathFullName} {LinkActionCommandArguments}", LinkAction.CommandPath.FullName, LinkAction.CommandArguments);
 
-			// Windows can run into an issue with too long of a commandline when clang tries to call ld to link.
-			// To work around this we call clang to just get the command it would execute and generate a
-			// second response file to directly call ld with the right arguments instead of calling through clang.
-			/* disable while tracking down some linker errors this introduces
-						if (OperatingSystem.IsWindows())
-						{
-							// capture the actual link command without running it
-							ProcessStartInfo StartInfo = new ProcessStartInfo();
-							StartInfo.WorkingDirectory = LinkEnvironment.IntermediateDirectory.FullName;
-							StartInfo.FileName = LinkAction.CommandPath;
-							StartInfo.Arguments = "-### " + LinkAction.CommandArguments;
-							StartInfo.UseShellExecute = false;
-							StartInfo.CreateNoWindow = true;
-							StartInfo.RedirectStandardError = true;
-
-							LinkerCommandline = "";
-
-							Process Proc = new Process();
-							Proc.StartInfo = StartInfo;
-							Proc.ErrorDataReceived += new DataReceivedEventHandler(OutputReceivedForLinker);
-							Proc.Start();
-							Proc.BeginErrorReadLine();
-							Proc.WaitForExit(5000);
-
-							LinkerCommandline = LinkerCommandline.Trim();
-
-							// the command should be in quotes; if not we'll just use clang to link as usual
-							int FirstQuoteIndex = LinkerCommandline.IndexOf('"');
-							if (FirstQuoteIndex >= 0)
-							{
-								int SecondQuoteIndex = LinkerCommandline.Substring(FirstQuoteIndex + 1).IndexOf('"');
-								if (SecondQuoteIndex >= 0)
-								{
-									LinkAction.CommandPath = LinkerCommandline.Substring(FirstQuoteIndex + 1, SecondQuoteIndex - FirstQuoteIndex);
-									LinkAction.CommandArguments = LinkerCommandline.Substring(FirstQuoteIndex + SecondQuoteIndex + 3);
-
-									// replace double backslashes
-									LinkAction.CommandPath = LinkAction.CommandPath.Replace("\\\\", "/");
-
-									// now create a response file for the full command using ld directly
-									FileReference FinalResponseFileName = FileReference.Combine(LinkEnvironment.IntermediateDirectory, OutputFile.Location.GetFileName() + ".responseFinal");
-									FileItem FinalResponseFileItem = Graph.CreateIntermediateTextFile(FinalResponseFileName, LinkAction.CommandArguments);
-									LinkAction.CommandArguments = string.Format("@\"{0}\"", FinalResponseFileName);
-									LinkAction.PrerequisiteItems.Add(FinalResponseFileItem);
-								}
-							}
-						}
-			*/
-
 			return OutputFile;
 		}
 
-		// captures stderr from clang
-		private static string LinkerCommandline = "";
-		public static void OutputReceivedForLinker(object Sender, DataReceivedEventArgs Line)
-		{
-			if ((Line != null) && (Line.Data != null) && (Line.Data.Contains("--sysroot")))
-			{
-				LinkerCommandline += Line.Data;
-			}
-		}
-
 		private void ExportObjectFilePaths(LinkEnvironment LinkEnvironment, string FileName)
 		{
 			// Write the list of object file directories
diff --git a/Engine/Source/Programs/UnrealBuildTool/Platform/Apple/AppleToolChain.cs b/Engine/Source/Programs/UnrealBuildTool/Platform/Apple/AppleToolChain.cs
--- a/Engine/Source/Programs/UnrealBuildTool/Platform/Apple/AppleToolChain.cs
+++ b/Engine/Source/Programs/UnrealBuildTool/Platform/Apple/AppleToolChain.cs
@@ -76,7 +76,6 @@ namespace UnrealBuildTool
 
 		// cache some ini settings
 		readonly bool bUseSwiftUIMain;
-		readonly bool bCreateSwiftBridgingHeader;
 
 		/// <summary>
 		/// Serializer options for serializing <see cref="IOSPostBuildSyncTarget"/> and <see cref="ApplePostBuildSyncTarget"/> instances.
@@ -94,7 +93,7 @@ namespace UnrealBuildTool
 			ProjectFile = Target?.ProjectFile;
 			ToolChainSettings = new Lazy<AppleToolChainSettings>(InCreateSettings);
 
-			AppleExports.GetSwiftIntegrationSettings(ProjectFile, Target == null ? TargetType.Game : Target.Type, Target == null ? UnrealTargetPlatform.Mac : Target.Platform, out bUseSwiftUIMain, out bCreateSwiftBridgingHeader);
+			AppleExports.GetSwiftIntegrationSettings(ProjectFile, Target == null ? TargetType.Game : Target.Type, Target == null ? UnrealTargetPlatform.Mac : Target.Platform, out bUseSwiftUIMain, out _);
 		}
 
 		/// <inheritdoc/>
diff --git a/Engine/Source/Programs/UnrealBuildTool/Platform/Apple/AppleToolChainSettings.cs b/Engine/Source/Programs/UnrealBuildTool/Platform/Apple/AppleToolChainSettings.cs
--- a/Engine/Source/Programs/UnrealBuildTool/Platform/Apple/AppleToolChainSettings.cs
+++ b/Engine/Source/Programs/UnrealBuildTool/Platform/Apple/AppleToolChainSettings.cs
@@ -32,11 +32,6 @@ namespace UnrealBuildTool
 		/// </summary>
 		private readonly string TargetOSName;
 
-		/// <summary>
-		/// The name apple uses for SDK directories
-		/// </summary>
-		private readonly string TargetSDKName;
-
 		/// <summary>
 		/// The version of the SDK being used to build with
 		/// </summary>
@@ -87,7 +82,6 @@ namespace UnrealBuildTool
 		protected AppleToolChainSettings(string OSPrefix, string? SimulatorOSPrefix, string TargetOSName, string TargetOSVersion, bool bVerbose, ILogger Logger)
 		{
 			this.TargetOSName = TargetOSName;
-			TargetSDKName = OSPrefix;
 			PlatformDirName = OSPrefix.ToLower();
 
 			// set up directories
diff --git a/Engine/Source/Programs/UnrealBuildTool/Platform/Windows/UEBuildWindows.cs b/Engine/Source/Programs/UnrealBuildTool/Platform/Windows/UEBuildWindows.cs
--- a/Engine/Source/Programs/UnrealBuildTool/Platform/Windows/UEBuildWindows.cs
+++ b/Engine/Source/Programs/UnrealBuildTool/Platform/Windows/UEBuildWindows.cs
@@ -1224,8 +1224,6 @@ namespace UnrealBuildTool
 
 	class WindowsPlatform : UEBuildPlatform
 	{
-		readonly MicrosoftPlatformSDK SDK;
-
 		/// <summary>
 		/// Constructor
 		/// </summary>
@@ -1247,7 +1245,6 @@ namespace UnrealBuildTool
 		public WindowsPlatform(UnrealTargetPlatform InPlatform, MicrosoftPlatformSDK InSDK, UnrealArchitectureConfig ArchConfig, ILogger InLogger)
 			: base(InPlatform, InSDK, ArchConfig, InLogger)
 		{
-			SDK = InSDK;
 		}
 		/// <summary>
 		/// Reset a target's settings to the default
diff --git a/Engine/Source/Programs/UnrealBuildTool/ProjectFiles/Xcode/XcodeProject.cs b/Engine/Source/Programs/UnrealBuildTool/ProjectFiles/Xcode/XcodeProject.cs
--- a/Engine/Source/Programs/UnrealBuildTool/ProjectFiles/Xcode/XcodeProject.cs
+++ b/Engine/Source/Programs/UnrealBuildTool/ProjectFiles/Xcode/XcodeProject.cs
@@ -2365,9 +2365,7 @@ namespace UnrealBuildTool.XcodeProjectXcconfig
 		}
 
 		private readonly UnrealData UnrealData;
-		private readonly DirectoryInfo ExtensionPath;
 		private readonly UnrealTargetPlatform Platform;
-		private readonly TargetType TargetType;
 		private readonly string TargetName;
 		private readonly string AppGuid;
 		private readonly ExtensionConfig ExtConfig;
@@ -2376,11 +2374,9 @@ namespace UnrealBuildTool.XcodeProjectXcconfig
 			: base(XcodeTarget.Type.App_Ext, Project.UnrealData,
 				TargetName + (XcodeProjectFileGenerator.PerPlatformMode == XcodePerPlatformMode.OneWorkspacePerPlatform ? "" : $"_{Platform}"))
 		{
-			this.TargetType = TargetType;
 			this.Platform = Platform;
 			this.AppGuid = AppGuid;
 			this.TargetName = TargetName;
-			this.ExtensionPath = ExtensionPath;
 			this.ExtConfig = ExtConfig ?? new ExtensionConfig();
 			UnrealData = Project.UnrealData;
 
