Fix or selectively disable issue IDE0051 (Remove unused private member) in UBT.

diff --git a/Engine/Source/Programs/UnrealBuildTool/Configuration/Rules/TestModuleRules.cs b/Engine/Source/Programs/UnrealBuildTool/Configuration/Rules/TestModuleRules.cs
--- a/Engine/Source/Programs/UnrealBuildTool/Configuration/Rules/TestModuleRules.cs
+++ b/Engine/Source/Programs/UnrealBuildTool/Configuration/Rules/TestModuleRules.cs
@@ -178,11 +178,13 @@ namespace UnrealBuildTool
 		{
 		}
 
+#pragma warning disable IDE0051 // Remove unused private member - used via reflection, see TestMode.cs
 		/// <summary>
 		/// Generates or updates metadata file for LowLevelTests.xml containing test flags: name, short name, target name, relative binaries path, supported platforms etc.
 		/// Called by RunUBT.bat -Mode=Test -GenerateMetadata
 		/// </summary>
 		private static void UpdateBuildGraphMetadata(Metadata TestMetadata, string ModuleDirectory, string ModuleName, ILogger Log)
+#pragma warning restore IDE0051
 		{
 			string BaseFolder = GetBaseFolder(ModuleDirectory);
 
diff --git a/Engine/Source/Programs/UnrealBuildTool/Configuration/UEBuildTarget.cs b/Engine/Source/Programs/UnrealBuildTool/Configuration/UEBuildTarget.cs
--- a/Engine/Source/Programs/UnrealBuildTool/Configuration/UEBuildTarget.cs
+++ b/Engine/Source/Programs/UnrealBuildTool/Configuration/UEBuildTarget.cs
@@ -410,11 +410,6 @@ namespace UnrealBuildTool
 			return StringRegistry;
 		}
 
-		private UnrealPlatformGroup(string Name)
-		{
-			Id = GetUniqueStringRegistry().FindOrAddByName(Name);
-		}
-
 		private UnrealPlatformGroup(int InId)
 		{
 			Id = InId;
@@ -4762,17 +4757,6 @@ namespace UnrealBuildTool
 			return ScriptFiles.ToArray();
 		}
 
-		private static FileReference AddModuleFilenameSuffix(string ModuleName, FileReference FilePath, string Suffix)
-		{
-			int MatchPos = FilePath.FullName.LastIndexOf(ModuleName, StringComparison.InvariantCultureIgnoreCase);
-			if (MatchPos < 0)
-			{
-				throw new BuildException("Failed to find module name \"{0}\" specified on the command line inside of the output filename \"{1}\" to add appendage.", ModuleName, FilePath);
-			}
-			string Appendage = "-" + Suffix;
-			return new FileReference(FilePath.FullName.Insert(MatchPos + ModuleName.Length, Appendage));
-		}
-
 		/// <summary>
 		/// Finds a list of module names which can be hot-reloaded
 		/// </summary>
diff --git a/Engine/Source/Programs/UnrealBuildTool/Executors/ParallelExecutor.cs b/Engine/Source/Programs/UnrealBuildTool/Executors/ParallelExecutor.cs
--- a/Engine/Source/Programs/UnrealBuildTool/Executors/ParallelExecutor.cs
+++ b/Engine/Source/Programs/UnrealBuildTool/Executors/ParallelExecutor.cs
@@ -28,9 +28,11 @@ namespace UnrealBuildTool
 		[Obsolete("ParallelExecutor.MaxProcessorCount is deprecated. Please update xml to use BuildConfiguration.MaxParallelActions")]
 #pragma warning disable 0169
 #pragma warning disable CA1823 // Avoid unused private fields - part of the public API
+#pragma warning disable IDE0051 // Remove unused private member
 		private static int MaxProcessorCount;
 #pragma warning restore 0169
 #pragma warning restore CA1823
+#pragma warning restore IDE0051
 
 		/// <summary>
 		/// Processor count multiplier for local execution. Can be below 1 to reserve CPU for other tasks.
diff --git a/Engine/Source/Programs/UnrealBuildTool/Executors/UnrealBuildAccelerator/UBAAgentCoordinatorHorde.cs b/Engine/Source/Programs/UnrealBuildTool/Executors/UnrealBuildAccelerator/UBAAgentCoordinatorHorde.cs
--- a/Engine/Source/Programs/UnrealBuildTool/Executors/UnrealBuildAccelerator/UBAAgentCoordinatorHorde.cs
+++ b/Engine/Source/Programs/UnrealBuildTool/Executors/UnrealBuildAccelerator/UBAAgentCoordinatorHorde.cs
@@ -17,8 +17,6 @@ using EpicGames.Horde.Server;
 using EpicGames.Horde.Storage;
 using EpicGames.Horde.Storage.Bundles;
 using EpicGames.Horde.Storage.Nodes;
-using EpicGames.OIDC;
-using Microsoft.Extensions.Configuration;
 using Microsoft.Extensions.DependencyInjection;
 using Microsoft.Extensions.Logging;
 using Microsoft.Extensions.Logging.Abstractions;
@@ -504,33 +502,6 @@ namespace UnrealBuildTool
 			}
 		}
 
-		static async Task<string> GetOidcBearerTokenAsync(DirectoryReference? projectDir, string oidcProvider, ILogger logger, CancellationToken cancellationToken = default)
-		{
-			logger.LogInformation("Performing OIDC token refresh...");
-
-			using ITokenStore tokenStore = TokenStoreFactory.CreateTokenStore();
-			IConfiguration providerConfiguration = ProviderConfigurationFactory.ReadConfiguration(Unreal.EngineDirectory.ToDirectoryInfo(), projectDir?.ToDirectoryInfo());
-			OidcTokenManager oidcTokenManager = OidcTokenManager.CreateTokenManager(providerConfiguration, tokenStore, new List<string>() { oidcProvider });
-
-			OidcTokenInfo result;
-			try
-			{
-				result = await oidcTokenManager.GetAccessToken(oidcProvider, cancellationToken);
-			}
-			catch (NotLoggedInException)
-			{
-				result = await oidcTokenManager.LoginAsync(oidcProvider, cancellationToken);
-			}
-
-			if (result.AccessToken == null)
-			{
-				throw new Exception($"Unable to get access token for {oidcProvider}");
-			}
-
-			logger.LogInformation("Received bearer token for {OidcProvider}", oidcProvider);
-			return result.AccessToken;
-		}
-
 		async Task RunWorkerAsync(Worker self, IComputeLease lease, BlobLocator tool, string executable, ILogger logger, UnrealBuildAcceleratorHordeConfig hordeConfig, string agentName, string desc, CancellationToken cancellationToken)
 		{
 			logger.LogDebug("Running worker task..");
diff --git a/Engine/Source/Programs/UnrealBuildTool/Platform/Android/UEDeployAndroid.cs b/Engine/Source/Programs/UnrealBuildTool/Platform/Android/UEDeployAndroid.cs
--- a/Engine/Source/Programs/UnrealBuildTool/Platform/Android/UEDeployAndroid.cs
+++ b/Engine/Source/Programs/UnrealBuildTool/Platform/Android/UEDeployAndroid.cs
@@ -1820,44 +1820,6 @@ namespace UnrealBuildTool
 			return Proc.ExitCode;
 		}
 
-		private static int RunCommandLineProgramAndReturnResult(string WorkingDirectory, string Command, string Params, ILogger Logger, string? OverrideDesc = null, bool bUseShellExecute = false)
-		{
-			// Process Arguments follow windows conventions in .NET Core
-			// Which means single quotes ' are not considered quotes.
-			// see [REDACTED_PROFILE_URL]/runtime/issues/29857
-			// also see UE-102580
-			// for rules see https://docs.microsoft.com/en-us/cpp/cpp/main-function-command-line-args
-			if (Params.Contains('\''))
-			{
-				Params = Params.Replace("\"", "\\\"");
-				Params = Params.Replace('\'', '\"');
-			}
-
-			if (OverrideDesc == null)
-			{
-				Logger.LogInformation("\nRunning: {Command} {Params}", Command, Params);
-			}
-			else if (!String.IsNullOrEmpty(OverrideDesc))
-			{
-				Logger.LogInformation("{Message}", OverrideDesc);
-				Logger.LogDebug("\nRunning: {Command} {Params}", Command, Params);
-			}
-
-			ProcessStartInfo StartInfo = new ProcessStartInfo();
-			StartInfo.WorkingDirectory = WorkingDirectory;
-			StartInfo.FileName = Command;
-			StartInfo.Arguments = Params;
-			StartInfo.UseShellExecute = bUseShellExecute;
-			StartInfo.WindowStyle = ProcessWindowStyle.Minimized;
-
-			Process Proc = new Process();
-			Proc.StartInfo = StartInfo;
-			Proc.Start();
-			Proc.WaitForExit();
-
-			return Proc.ExitCode;
-		}
-
 		private static void RunCommandLineProgramWithException(string WorkingDirectory, string Command, string Params, ILogger Logger, string? OverrideDesc = null, bool bUseShellExecute = false)
 		{
 			// Process Arguments follow windows conventions in .NET Core
diff --git a/Engine/Source/Programs/UnrealBuildTool/Platform/IOS/UEBuildIOS.cs b/Engine/Source/Programs/UnrealBuildTool/Platform/IOS/UEBuildIOS.cs
--- a/Engine/Source/Programs/UnrealBuildTool/Platform/IOS/UEBuildIOS.cs
+++ b/Engine/Source/Programs/UnrealBuildTool/Platform/IOS/UEBuildIOS.cs
@@ -2,7 +2,6 @@
 
 using System;
 using System.Collections.Generic;
-using System.Diagnostics;
 using System.IO;
 using System.Linq;
 using System.Text;
@@ -591,45 +590,6 @@ namespace UnrealBuildTool
 				}
 			}
 		}
-
-		void IPPDataReceivedHandler(object Sender, DataReceivedEventArgs Line, ILogger Logger)
-		{
-			if ((Line != null) && (Line.Data != null))
-			{
-				if (Line.Data.StartsWith("IPP WARNING:"))
-				{
-					// Don't output IPP warnings to the console as they may not be warnings relevant to the build and could cause build failures.
-					Logger.LogDebug("{LineData}", Line.Data);
-				}
-				else
-				{
-					Logger.LogInformation("{LineData}", Line.Data);
-				}
-
-				if (!String.IsNullOrEmpty(SigningCertificate))
-				{
-					if (Line.Data.Contains("CERTIFICATE-") && Line.Data.Contains(SigningCertificate))
-					{
-						bHaveCertificate = true;
-					}
-				}
-				else
-				{
-					int cindex = Line.Data.IndexOf("CERTIFICATE-");
-					int pindex = Line.Data.IndexOf("PROVISION-");
-					if (cindex > -1 && pindex > -1)
-					{
-						cindex += "CERTIFICATE-".Length;
-						SigningCertificate = Line.Data.Substring(cindex, pindex - cindex - 1);
-						pindex += "PROVISION-".Length;
-						if (pindex < Line.Data.Length)
-						{
-							MobileProvisionFile = new FileReference(Line.Data.Substring(pindex));
-						}
-					}
-				}
-			}
-		}
 	}
 
 	class IOSArchitectureConfig : AppleArchitectureConfig
diff --git a/Engine/Source/Programs/UnrealBuildTool/Platform/Mac/MacToolChain.cs b/Engine/Source/Programs/UnrealBuildTool/Platform/Mac/MacToolChain.cs
--- a/Engine/Source/Programs/UnrealBuildTool/Platform/Mac/MacToolChain.cs
+++ b/Engine/Source/Programs/UnrealBuildTool/Platform/Mac/MacToolChain.cs
@@ -276,19 +276,6 @@ namespace UnrealBuildTool
 			Writer.Write(PreLine + "\n");
 		}
 
-		private int LoadEngineCL()
-		{
-			BuildVersion? Version;
-			if (BuildVersion.TryRead(BuildVersion.GetDefaultFileName(), out Version))
-			{
-				return Version.Changelist;
-			}
-			else
-			{
-				return 0;
-			}
-		}
-
 		public static string LoadEngineDisplayVersion(bool bIgnorePatchVersion = false)
 		{
 			BuildVersion? Version;
@@ -302,11 +289,6 @@ namespace UnrealBuildTool
 			}
 		}
 
-		private int LoadBuiltFromChangelistValue()
-		{
-			return LoadEngineCL();
-		}
-
 		private string LoadEngineAPIVersion()
 		{
 			if (BuildVersion.TryRead(BuildVersion.GetDefaultFileName(), out BuildVersion? Version))
diff --git a/Engine/Source/Programs/UnrealBuildTool/Platform/Windows/UEBuildWindows.cs b/Engine/Source/Programs/UnrealBuildTool/Platform/Windows/UEBuildWindows.cs
--- a/Engine/Source/Programs/UnrealBuildTool/Platform/Windows/UEBuildWindows.cs
+++ b/Engine/Source/Programs/UnrealBuildTool/Platform/Windows/UEBuildWindows.cs
@@ -1731,16 +1731,6 @@ namespace UnrealBuildTool
 			return arch.ToString();
 		}
 
-		/// <summary>
-		/// Determines if a directory contains a valid DIA SDK
-		/// </summary>
-		/// <param name="DiaSdkDir">The directory to check</param>
-		/// <returns>True if it contains a valid DIA SDK</returns>
-		static bool IsValidDiaSdkDir(DirectoryReference DiaSdkDir)
-		{
-			return FileReference.Exists(FileReference.Combine(DiaSdkDir, "bin", "amd64", "msdia140.dll"));
-		}
-
 		[SupportedOSPlatform("windows")]
 		public static bool TryGetWindowsSdkDir(string? DesiredVersion, ILogger Logger, [NotNullWhen(true)] out VersionNumber? OutSdkVersion, [NotNullWhen(true)] out DirectoryReference? OutSdkDir)
 		{
diff --git a/Engine/Source/Programs/UnrealBuildTool/ProjectFiles/CodeLite/CodeLiteProject.cs b/Engine/Source/Programs/UnrealBuildTool/ProjectFiles/CodeLite/CodeLiteProject.cs
--- a/Engine/Source/Programs/UnrealBuildTool/ProjectFiles/CodeLite/CodeLiteProject.cs
+++ b/Engine/Source/Programs/UnrealBuildTool/ProjectFiles/CodeLite/CodeLiteProject.cs
@@ -29,14 +29,6 @@ namespace UnrealBuildTool
 			return true;
 		}
 
-		string GetPathRelativeTo(string BasePath, string ToPath)
-		{
-			Uri path1 = new Uri(BasePath);
-			Uri path2 = new Uri(ToPath);
-			Uri diff = path1.MakeRelativeUri(path2);
-			return diff.ToString();
-		}
-
 		public override bool WriteProjectFile(List<UnrealTargetPlatform> InPlatforms, List<UnrealTargetConfiguration> InConfigurations, PlatformProjectGeneratorCollection PlatformProjectGenerators, ILogger Logger)
 		{
 			bool bSuccess = false;
diff --git a/Engine/Source/Programs/UnrealBuildTool/ProjectFiles/Eddie/EddieProject.cs b/Engine/Source/Programs/UnrealBuildTool/ProjectFiles/Eddie/EddieProject.cs
--- a/Engine/Source/Programs/UnrealBuildTool/ProjectFiles/Eddie/EddieProject.cs
+++ b/Engine/Source/Programs/UnrealBuildTool/ProjectFiles/Eddie/EddieProject.cs
@@ -53,11 +53,6 @@ namespace UnrealBuildTool
 			return ProjectFilePath.GetFileNameWithoutExtension();
 		}
 
-		private bool IsSourceCode(string Extension)
-		{
-			return Extension == ".c" || Extension == ".cc" || Extension == ".cxx" || Extension == ".cpp" || Extension == ".m" || Extension == ".mm" || Extension == ".cs";
-		}
-
 		public override SourceFile? AllocSourceFile(FileReference InitFilePath, DirectoryReference? InitProjectSubFolder)
 		{
 			if (InitFilePath.GetFileName().StartsWith("."))
@@ -105,11 +100,6 @@ namespace UnrealBuildTool
 			return null;
 		}
 
-		private static string ConvertPath(string InPath)
-		{
-			return InPath.Replace("\\", "/");
-		}
-
 		private void ParseSourceFilesIntoGroups()
 		{
 			foreach (SourceFile CurSourceFile in SourceFiles)
diff --git a/Engine/Source/Programs/UnrealBuildTool/ProjectFiles/ProjectFileGenerator.cs b/Engine/Source/Programs/UnrealBuildTool/ProjectFiles/ProjectFileGenerator.cs
--- a/Engine/Source/Programs/UnrealBuildTool/ProjectFiles/ProjectFileGenerator.cs
+++ b/Engine/Source/Programs/UnrealBuildTool/ProjectFiles/ProjectFileGenerator.cs
@@ -1947,31 +1947,6 @@ namespace UnrealBuildTool
 			return Rules.FindAllRulesSourceFiles(Rules.RulesFileType.Module, GameFolders: AllGameProjects.Select(x => x.Directory).ToList(), ForeignPlugins: null, AdditionalSearchPaths: AdditionalSearchPaths);
 		}
 
-		/// <summary>
-		/// List of non-redistributable folders
-		/// </summary>
-		private static readonly string[] NoRedistFolders = new string[]
-		{
-			Path.DirectorySeparatorChar + "NoRedist" + Path.DirectorySeparatorChar,
-			Path.DirectorySeparatorChar + "NotForLicensees" + Path.DirectorySeparatorChar,
-			Path.DirectorySeparatorChar + "LimitedAccess" + Path.DirectorySeparatorChar
-		};
-
-		/// <summary>
-		/// Checks if a module is in a non-redistributable folder
-		/// </summary>
-		private static bool IsNoRedistModule(FileReference ModulePath)
-		{
-			foreach (string NoRedistFolderName in NoRedistFolders)
-			{
-				if (ModulePath.FullName.Contains(NoRedistFolderName, StringComparison.InvariantCultureIgnoreCase))
-				{
-					return true;
-				}
-			}
-			return false;
-		}
-
 		/// <summary>
 		/// Finds all target files (filtering by platform)
 		/// </summary>
diff --git a/Engine/Source/Programs/UnrealBuildTool/ProjectFiles/QMake/QMakefileGenerator.cs b/Engine/Source/Programs/UnrealBuildTool/ProjectFiles/QMake/QMakefileGenerator.cs
--- a/Engine/Source/Programs/UnrealBuildTool/ProjectFiles/QMake/QMakefileGenerator.cs
+++ b/Engine/Source/Programs/UnrealBuildTool/ProjectFiles/QMake/QMakefileGenerator.cs
@@ -81,11 +81,6 @@ namespace UnrealBuildTool
 			IncludeDirectories.Add(FullPath);
 		}
 
-		// For later when moving all this back to MakefileGenerator.cs
-		private void AddDefines(ref List<string> DefinesContent, string define, string value)
-		{
-		}
-
 		private bool WriteQMakePro(ILogger Logger)
 		{
 			// Some more stuff borrowed from Mac side of things.
diff --git a/Engine/Source/Programs/UnrealBuildTool/ProjectFiles/Rider/RiderProjectFile.cs b/Engine/Source/Programs/UnrealBuildTool/ProjectFiles/Rider/RiderProjectFile.cs
--- a/Engine/Source/Programs/UnrealBuildTool/ProjectFiles/Rider/RiderProjectFile.cs
+++ b/Engine/Source/Programs/UnrealBuildTool/ProjectFiles/Rider/RiderProjectFile.cs
@@ -146,21 +146,6 @@ namespace UnrealBuildTool
 			return true;
 		}
 
-		private bool IsPlatformInHostGroup(UnrealTargetPlatform Platform)
-		{
-			IEnumerable<UnrealPlatformGroup> Groups = UEBuildPlatform.GetPlatformGroups(BuildHostPlatform.Current.Platform);
-			foreach (UnrealPlatformGroup Group in Groups)
-			{
-				// Desktop includes Linux, Mac and Windows.
-				if (UEBuildPlatform.IsPlatformInGroup(Platform, Group) && Group != UnrealPlatformGroup.Desktop)
-				{
-					return true;
-				}
-			}
-
-			return false;
-		}
-
 		private void SerializeTarget(FileReference OutputFile, UEBuildTarget BuildTarget, PlatformProjectGeneratorCollection PlatformProjectGenerators, JsonWriterStyle Minimize, bool bBuildByDefault, ILogger Logger)
 		{
 			DirectoryReference.CreateDirectory(OutputFile.Directory);
diff --git a/Engine/Source/Programs/UnrealBuildTool/ProjectFiles/Xcode/XcodeProject.cs b/Engine/Source/Programs/UnrealBuildTool/ProjectFiles/Xcode/XcodeProject.cs
--- a/Engine/Source/Programs/UnrealBuildTool/ProjectFiles/Xcode/XcodeProject.cs
+++ b/Engine/Source/Programs/UnrealBuildTool/ProjectFiles/Xcode/XcodeProject.cs
@@ -3252,9 +3252,8 @@ namespace UnrealBuildTool.XcodeProjectXcconfig
 
 				if (FileReference.Exists(TemplateProject))
 				{
-					// @todo hahahaah
+					// Merging was removed in CL 23517683 in 2022-12, when adding support for 'modern' Xcode projects.
 					continue;
-					//bSuccess = MergeIntoTemplateProject(PBXFilePath, RootProject, TemplateProject, Logger);
 				}
 				else
 				{
@@ -3325,177 +3324,5 @@ namespace UnrealBuildTool.XcodeProjectXcconfig
 				}
 			}
 		}
-
-		private string Plist(string Command)
-		{
-			return XcodeUtils.Plist(Command);
-		}
-
-		bool MergeIntoTemplateProject(FileReference PBXProjFilePath, XcodeProject RootProject, FileReference TemplateProject, ILogger Logger)
-		{
-			// activate a file for plist reading/writing here
-			XcodeUtils.SetActivePlistFile(PBXFilePathForPlatform(CurrentPlistPlatform).FullName);
-
-			// copy existing template project to final location
-			if (FileReference.Exists(PBXProjFilePath))
-			{
-				FileReference.Delete(PBXProjFilePath);
-			}
-			DirectoryReference.CreateDirectory(PBXProjFilePath.Directory);
-			FileReference.Copy(TemplateProject, PBXProjFilePath);
-
-			// write the nodes we need to add (Build/Index targets)
-			XcodeRunTarget RunTarget = XcodeProjectNode.GetNodesOfType<XcodeRunTarget>(RootProject).First();
-			XcodeBuildTarget BuildTarget = XcodeProjectNode.GetNodesOfType<XcodeBuildTarget>(RunTarget).First();
-			XcodeIndexTarget IndexTarget = XcodeProjectNode.GetNodesOfType<XcodeIndexTarget>(RootProject).First();
-			XcodeDependency BuildDependency = XcodeProjectNode.GetNodesOfType<XcodeDependency>(RunTarget).First();
-
-			// the runtarget and project need to write out so all of their xcconfigs get written as well,
-			// so write everything to a temp string that is tossed, but all xcconfigs will be done at least
-			StringBuilder Temp = new StringBuilder();
-			XcodeProjectNode.WriteNodeAndReferences(Temp, RootProject!, Logger);
-
-			StringBuilder Content = new StringBuilder();
-			Content.WriteLine(0, "{");
-			// @todo figure out how to deal with platform
-			// FileCollection.Write(Content, Platform);
-			XcodeProjectNode.WriteNodeAndReferences(Content, BuildTarget, Logger);
-			XcodeProjectNode.WriteNodeAndReferences(Content, IndexTarget, Logger);
-			XcodeProjectNode.WriteNodeAndReferences(Content, BuildDependency, Logger);
-			Content.WriteLine(0, "}");
-
-			// write to disk
-			FileReference ImportFile = FileReference.Combine(PBXProjFilePath.Directory, "import.plist");
-			File.WriteAllText(ImportFile.FullName, Content.ToString());
-
-			// cache some standard guids from the template project
-			string ProjectGuid = Plist($"Print :rootObject");
-			string TemplateMainGroupGuid = Plist($"Print :objects:{ProjectGuid}:mainGroup");
-
-			// fixup paths that were relative to original project to be relative to merged project
-			//			List<string> ObjectGuids = PlistObjects();
-			IEnumerable<string> MainGroupChildrenGuids = XcodeUtils.PlistArray($":objects:{TemplateMainGroupGuid}:children");
-
-			string RelativeFromMergedToTemplate = TemplateProject.Directory.ParentDirectory!.MakeRelativeTo(PBXProjFilePath.Directory.ParentDirectory!);
-
-			// look for groups with a 'path' element that is in the main group, so that it and everything will get redirected to new location
-			string? FixedPath;
-			foreach (string ChildGuid in MainGroupChildrenGuids)
-			{
-				string IsA = Plist($"Print :objects:{ChildGuid}:isa");
-				// if a Group has a path
-				if (IsA == "PBXGroup")
-				{
-					if ((FixedPath = XcodeUtils.PlistFixPath($":objects:{ChildGuid}:path", RelativeFromMergedToTemplate)) != null)
-					{
-						// if there wasn't a name before, it will now have a nasty path as the name, so add it now
-						XcodeUtils.PlistSetAdd($":objects:{ChildGuid}:name", Path.GetFileName(FixedPath));
-					}
-				}
-			}
-
-			// and import it into the template
-			Plist($"Merge \"{ImportFile.FullName}\" :objects");
-
-			// get all the targets in the template that are application types
-			IEnumerable<string> AppTargetGuids = XcodeUtils.PlistArray($":objects:{ProjectGuid}:targets")
-				.Where(TargetGuid => (Plist($"Print :objects:{TargetGuid}:productType") == "com.apple.product-type.application"));
-
-			// add a dependency on the build target from the app target(s)
-			foreach (string AppTargetGuid in AppTargetGuids)
-			{
-				Plist($"Add :objects:{AppTargetGuid}:dependencies:0 string {BuildDependency.Guid}");
-			}
-
-			// the BuildDependency object was in the "container" of the generated project, not the merged one, so fix it up now
-			Plist($"Set :objects:{BuildDependency.ProxyGuid}:containerPortal {ProjectGuid}");
-
-			// now add all the non-run targets from the generated
-			foreach (XcodeTarget Target in XcodeProjectNode.GetNodesOfType<XcodeTarget>(RootProject!).Where(x => x.GetType() != typeof(XcodeRunTarget)))
-			{
-				Plist($"Add :objects:{ProjectGuid}:targets:0 string {Target.Guid}");
-			}
-
-			// hook up Xcconfig files to the project and the project configs
-			// @todo how to manage with conflicts already present...
-			//PlistSetAdd($":objects:{ProjectGuid}:baseConfigurationReference", RootProject.Xcconfig!.Guid, "string");
-
-			// re-get the list of targets now that we merged in the other file
-			IEnumerable<string> AllTargetGuids = XcodeUtils.PlistArray($":objects:{ProjectGuid}:targets");
-
-			List<string> NodesToFix = new() { ProjectGuid };
-			NodesToFix.AddRange(AllTargetGuids);
-
-			bool bIsProject = true;
-			foreach (string NodeGuid in NodesToFix)
-			{
-				bool bIsAppTarget = AppTargetGuids.Contains(NodeGuid);
-
-				// get the config list, and from there we can get the configs
-				string ProjectBuildConfigListGuid = Plist($"Print :objects:{NodeGuid}:buildConfigurationList");
-
-				IEnumerable<string> ConfigGuids = XcodeUtils.PlistArray($":objects:{ProjectBuildConfigListGuid}:buildConfigurations");
-				foreach (string ConfigGuid in ConfigGuids)
-				{
-					// find the matching unreal generated project build config to hook up to
-					// for now we assume Release is Development [Editor], but we should make sure the template project has good configs
-					// we have to rename the template config from Release because it won't find the matching config in the build target
-					string ConfigName = Plist($"Print :objects:{ConfigGuid}:name");
-					if (ConfigName == "Release")
-					{
-						ConfigName = "Development";
-						if (UnrealData.bMakeProjectPerTarget && UnrealData.TargetRules.Type == TargetType.Editor)
-						{
-							ConfigName = "Development Editor";
-						}
-						Plist($"Set :objects:{ConfigGuid}:name \"{ConfigName}\"");
-					}
-
-					// if there's a plist path, then it will need to be fixed up
-					XcodeUtils.PlistFixPath($":objects:{ConfigGuid}:buildSettings:INFOPLIST_FILE", RelativeFromMergedToTemplate);
-
-					if (bIsProject)
-					{
-						//Console.WriteLine("Looking for " + ConfigName);
-						XcodeBuildConfig Config = RootProject!.ProjectBuildConfigs.BuildConfigs.First(x => x.Info.DisplayName == ConfigName);
-						XcodeUtils.PlistSetAdd($":objects:{ConfigGuid}:baseConfigurationReference", Config.Xcconfig!.Guid, "string");
-					}
-
-					// the Build target used some ini settings to compile, and Run target must match, so we override a few settings, at
-					// whatever level they were already specified at (Projet and/or Target)
-					XcodeUtils.PlistSetUpdate($":objects:{ConfigGuid}:buildSettings:MACOSX_DEPLOYMENT_TARGET", MacToolChainSettings.MinMacDeploymentVersion(UnrealData.TargetRules.Type));
-					if (UnrealData.IOSProjectSettings != null)
-					{
-						XcodeUtils.PlistSetUpdate($":objects:{ConfigGuid}:buildSettings:IPHONEOS_DEPLOYMENT_TARGET", UnrealData.IOSProjectSettings.RuntimeVersion);
-					}
-					if (UnrealData.TVOSProjectSettings != null)
-					{
-						XcodeUtils.PlistSetUpdate($":objects:{ConfigGuid}:buildSettings:TVOS_DEPLOYMENT_TARGET", UnrealData.TVOSProjectSettings.RuntimeVersion);
-					}
-				}
-
-				bIsProject = false;
-			}
-
-			// now we need to merge the main groups together
-			string GeneratedMainGroupGuid = FileCollection.MainGroupGuid;
-			int Index = 0;
-			while (true)
-			{
-				// we copy to a high index to put the copied entries at the end in the same order
-				string Output = Plist($"Copy :objects:{GeneratedMainGroupGuid}:children:{Index} :objects:{TemplateMainGroupGuid}:children:100000000");
-
-				// loop until error
-				if (!String.IsNullOrEmpty(Output))
-				{
-					break;
-				}
-				Index++;
-			}
-			// and remove the one we copied from
-			Plist($"Delete :objects:{GeneratedMainGroupGuid}");
-
-			return true;
-		}
 	}
 }
diff --git a/Engine/Source/Programs/UnrealBuildTool/System/Utils.cs b/Engine/Source/Programs/UnrealBuildTool/System/Utils.cs
--- a/Engine/Source/Programs/UnrealBuildTool/System/Utils.cs
+++ b/Engine/Source/Programs/UnrealBuildTool/System/Utils.cs
@@ -1923,18 +1923,6 @@ namespace UnrealBuildTool
 			return (Encoding.UTF8.GetByteCount(Str) != Str.Length) ? new UTF8Encoding(false) : Encoding.ASCII;
 		}
 
-		/// <summary>
-		/// Determines the appropriate encoding for a list of strings: either ASCII or UTF-8.
-		/// If the string length is equivalent to the encoded length, then no non-ASCII characters were present in the string.
-		/// Don't write BOM as it messes with clang when loading response files.
-		/// </summary>
-		/// <param name="Strings">The string to test.</param>
-		/// <returns>Either System.Text.Encoding.ASCII or System.Text.Encoding.UTF8, depending on whether or not the strings contains non-ASCII characters.</returns>
-		private static Encoding GetEncodingForStrings(IEnumerable<string> Strings)
-		{
-			return Strings.Any(S => Encoding.UTF8.GetByteCount(S) != S.Length) ? new UTF8Encoding(false) : Encoding.ASCII;
-		}
-
 		/// <summary>
 		/// Creates a hard link at Path pointing to PathToTarget.
 		/// If that fails, copies the file at PathToTarget to Path if it's different, using Length and LastWriteTimeUtc.
