Bucket SxS manifest intermediates in the same way that they are bucketed as build outputs.

diff --git a/Engine/Source/Programs/UnrealBuildTool/Configuration/UEBuildTarget.cs b/Engine/Source/Programs/UnrealBuildTool/Configuration/UEBuildTarget.cs
--- a/Engine/Source/Programs/UnrealBuildTool/Configuration/UEBuildTarget.cs
+++ b/Engine/Source/Programs/UnrealBuildTool/Configuration/UEBuildTarget.cs
@@ -3163,7 +3163,7 @@ namespace UnrealBuildTool
 			}
 
 			// Allow the toolchain to modify the final output items
-			TargetToolChain.FinalizeOutput(Rules, MakefileBuilder);
+			TargetToolChain.FinalizeOutput(Rules, MakefileBuilder, ExeDir);
 
 			// Get all the regular build products
 			List<KeyValuePair<FileReference, BuildProductType>> BuildProducts = new List<KeyValuePair<FileReference, BuildProductType>>();
diff --git a/Engine/Source/Programs/UnrealBuildTool/Platform/Clang/ClangToolChain.cs b/Engine/Source/Programs/UnrealBuildTool/Platform/Clang/ClangToolChain.cs
--- a/Engine/Source/Programs/UnrealBuildTool/Platform/Clang/ClangToolChain.cs
+++ b/Engine/Source/Programs/UnrealBuildTool/Platform/Clang/ClangToolChain.cs
@@ -293,7 +293,7 @@ namespace UnrealBuildTool
 			CompileActionWeight = Target.ClangCompileActionWeight;
 		}
 
-		public override void FinalizeOutput(ReadOnlyTargetRules Target, TargetMakefileBuilder MakefileBuilder)
+		public override void FinalizeOutput(ReadOnlyTargetRules Target, TargetMakefileBuilder MakefileBuilder, DirectoryReference ExeDir)
 		{
 			if (Target.bPrintToolChainTimingInfo && Target.bParseTimingInfoForTracing)
 			{
diff --git a/Engine/Source/Programs/UnrealBuildTool/Platform/Mac/MacToolChain.cs b/Engine/Source/Programs/UnrealBuildTool/Platform/Mac/MacToolChain.cs
--- a/Engine/Source/Programs/UnrealBuildTool/Platform/Mac/MacToolChain.cs
+++ b/Engine/Source/Programs/UnrealBuildTool/Platform/Mac/MacToolChain.cs
@@ -1211,9 +1211,9 @@ namespace UnrealBuildTool
 
 		private readonly List<FileItem> DebugInfoFiles = new List<FileItem>();
 
-		public override void FinalizeOutput(ReadOnlyTargetRules Target, TargetMakefileBuilder MakefileBuilder)
+		public override void FinalizeOutput(ReadOnlyTargetRules Target, TargetMakefileBuilder MakefileBuilder, DirectoryReference ExeDir)
 		{
-			base.FinalizeOutput(Target, MakefileBuilder);
+			base.FinalizeOutput(Target, MakefileBuilder, ExeDir);
 
 			TargetMakefile Makefile = MakefileBuilder.Makefile;
 
diff --git a/Engine/Source/Programs/UnrealBuildTool/Platform/Windows/PVSToolChain.cs b/Engine/Source/Programs/UnrealBuildTool/Platform/Windows/PVSToolChain.cs
--- a/Engine/Source/Programs/UnrealBuildTool/Platform/Windows/PVSToolChain.cs
+++ b/Engine/Source/Programs/UnrealBuildTool/Platform/Windows/PVSToolChain.cs
@@ -953,7 +953,7 @@ namespace UnrealBuildTool
 			throw new BuildException("Unable to link with PVS toolchain.");
 		}
 
-		public override void FinalizeOutput(ReadOnlyTargetRules target, TargetMakefileBuilder makefileBuilder)
+		public override void FinalizeOutput(ReadOnlyTargetRules target, TargetMakefileBuilder makefileBuilder, DirectoryReference ExeDir)
 		{
 			string outputFileExtension = OutputFileExtension;
 			FileReference outputFile = target.ProjectFile == null
diff --git a/Engine/Source/Programs/UnrealBuildTool/Platform/Windows/VCToolChain.cs b/Engine/Source/Programs/UnrealBuildTool/Platform/Windows/VCToolChain.cs
--- a/Engine/Source/Programs/UnrealBuildTool/Platform/Windows/VCToolChain.cs
+++ b/Engine/Source/Programs/UnrealBuildTool/Platform/Windows/VCToolChain.cs
@@ -1,6 +1,7 @@
 // Copyright Epic Games, Inc. All Rights Reserved.
 
 using System;
+using System.Collections.Concurrent;
 using System.Collections.Generic;
 using System.IO;
 using System.Linq;
@@ -2481,7 +2482,7 @@ namespace UnrealBuildTool
 			return ParseTimingInfoAction;
 		}
 
-		public override void FinalizeOutput(ReadOnlyTargetRules Target, TargetMakefileBuilder MakefileBuilder)
+		public override void FinalizeOutput(ReadOnlyTargetRules Target, TargetMakefileBuilder MakefileBuilder, DirectoryReference ExeDir)
 		{
 			if (Target.bPGOOptimize)
 			{
@@ -2602,10 +2603,12 @@ namespace UnrealBuildTool
 				}
 			}
 
-			OptionalUpdateSxSManifest(Target, MakefileBuilder);
+			OptionalUpdateSxSManifest(Target, MakefileBuilder, ExeDir);
 		}
 
-		protected void OptionalUpdateSxSManifest(ReadOnlyTargetRules Target, TargetMakefileBuilder MakefileBuilder)
+		private static readonly ConcurrentDictionary<FileReference, object> ManifestLocks = [];
+
+		protected void OptionalUpdateSxSManifest(ReadOnlyTargetRules Target, TargetMakefileBuilder MakefileBuilder, DirectoryReference ExeDir)
 		{
 			if (Target.WindowsPlatform.ManifestFile != null && !Target.bShouldCompileAsDLL && !Target.Architecture.bIsX64)
 			{
@@ -2621,46 +2624,50 @@ namespace UnrealBuildTool
 
 				// load or create the intermediate SxS manifest
 				bool bDirty = false;
-				FileReference IntermediateManifest = GetSxSIntermediateManifest(Target);
+				FileReference IntermediateManifest = GetSxSIntermediateManifest(Target, ExeDir);
 				FileReference TemplateManifest = FileReference.Combine(Unreal.EngineDirectory, "Build", "Windows", "Resources", AssemblyName, AssemblyName + ".manifest");
-				XDocument XmlManifest;
 
-				if (FileReference.Exists(IntermediateManifest))
-				{
-					XmlManifest = XDocument.Load(IntermediateManifest.FullName);
-				}
-				else
+				object Lock = ManifestLocks.GetOrAdd(IntermediateManifest, _ => new object());
+				lock (Lock)
 				{
-					XmlManifest = XDocument.Load(TemplateManifest.FullName);
-					bDirty = true;
-				}
+					XDocument XmlManifest;
 
-				// update & save the intermediate SxS manifest
-				if (XmlManifest.Root != null && XmlManifest.Root.Name.LocalName == "assembly")
-				{
-					XNamespace XmlNs = XmlManifest.Root.Name.Namespace;
-					HashSet<string> ExistingDLLNames = [.. XmlManifest.Root.Elements(XmlNs + "file").Select(X => X.Attribute("name")!.Value)];
+					if (FileReference.Exists(IntermediateManifest))
+					{
+						XmlManifest = XDocument.Load(IntermediateManifest.FullName);
+					}
+					else
+					{
+						XmlManifest = XDocument.Load(TemplateManifest.FullName);
+						bDirty = true;
+					}
 
-					foreach (FileItem DLL in DLLs)
+					// update & save the intermediate SxS manifest
+					if (XmlManifest.Root != null && XmlManifest.Root.Name.LocalName == "assembly")
 					{
-						string DLLName = DLL.Location.MakeRelativeTo(ManifestDir);
-						if (!ExistingDLLNames.Contains(DLLName))
+						XNamespace XmlNs = XmlManifest.Root.Name.Namespace;
+						HashSet<string> ExistingDLLNames = [.. XmlManifest.Root.Elements(XmlNs + "file").Select(X => X.Attribute("name")!.Value)];
+
+						foreach (FileItem DLL in DLLs)
 						{
-							XmlManifest.Root.Add(new XElement(XmlNs + "file", new XAttribute("name", DLLName)));
-							bDirty = true;
+							string DLLName = DLL.Location.MakeRelativeTo(ManifestDir);
+							if (!ExistingDLLNames.Contains(DLLName))
+							{
+								XmlManifest.Root.Add(new XElement(XmlNs + "file", new XAttribute("name", DLLName)));
+								bDirty = true;
+							}
 						}
 					}
-				}
-				if (bDirty)
-				{
-					DirectoryReference.CreateDirectory(IntermediateManifest.Directory);
-					XmlManifest.Save(IntermediateManifest.FullName);
+					if (bDirty)
+					{
+						DirectoryReference.CreateDirectory(IntermediateManifest.Directory);
+						XmlManifest.Save(IntermediateManifest.FullName);
+					}
 				}
 				// include the intermediate SxS manifest as a dependency
 				MakefileBuilder.Makefile.InternalDependencies.Add( FileItem.GetItemByFileReference(IntermediateManifest) );
 			}
 		}
-		
 
 		public override void PrepareRuntimeDependencies(List<RuntimeDependency> RuntimeDependencies, Dictionary<FileReference, FileReference> TargetFileToSourceFile, DirectoryReference ExeDir)
 		{
@@ -2724,7 +2731,7 @@ namespace UnrealBuildTool
 			if (Target.WindowsPlatform.ManifestFile != null && !Target.bShouldCompileAsDLL && !Target.Architecture.bIsX64)
 			{
 				string AssemblyName = Target.Architecture.ToString().ToLower();
-				FileReference SrcFile = GetSxSIntermediateManifest(Target);
+				FileReference SrcFile = GetSxSIntermediateManifest(Target, ExeDir);
 				FileReference DstFile = FileReference.Combine(ExeDir, AssemblyName, AssemblyName + ".manifest");
 
 				TargetFileToSourceFile[DstFile] = SrcFile;
@@ -3994,21 +4001,40 @@ namespace UnrealBuildTool
 
 			if (Target.WindowsPlatform.ManifestFile != null && Binary.Type == UEBuildBinaryType.Executable && !Target.Architecture.bIsX64)
 			{
-				FileReference IntermediateManifest = GetSxSIntermediateManifest(Target);
+				FileReference IntermediateManifest = GetSxSIntermediateManifest(Target, Binary.OutputDir);
 				BuildProducts.TryAdd(IntermediateManifest, BuildProductType.BuildResource);
 			}
 		}
 
-		private static FileReference GetSxSIntermediateManifest(ReadOnlyTargetRules Target)
-		{
-			DirectoryReference BaseDir = Target.ProjectFile?.Directory ?? Unreal.EngineDirectory;
-			string PlatformIntermediateDirectory = UEBuildTarget.GetPlatformIntermediateFolder(Target.Platform, Target.Architectures, false);
-			string TargetIntermediateName = UEBuildTarget.GetTargetIntermediateFolderName(Target.Name, Target.IntermediateEnvironment);
-			string TargetConfiguration = Target.Configuration.ToString();
-			DirectoryReference IntermediateDirectory = DirectoryReference.Combine(BaseDir, PlatformIntermediateDirectory, TargetIntermediateName, TargetConfiguration);
+		/// <summary>
+		/// Get the intermediate location to store the side-by-side assembly manifest file for the given target.
+		/// </summary>
+		/// <see href="https://learn.microsoft.com/en-us/windows/win32/SbsCs/about-private-assemblies-" />
+		/// <remarks>
+		/// In theory, each target could have different dependencies, so ideally each target would have their own private
+		/// assembly with its own copy of the dependencies. Unfortunately, the destination of the dependencies is scattered
+		/// across various different build files (e.g. `Engine/Source/ThirdParty/Intel/TBB/IntelTBB.Build.cs`) which each
+		/// specify that the dependencies should be output to a folder named for the target architecture. This means that
+		/// it's impractical to output a unique assembly per target.
+		/// 
+		/// Instead, we can have binaries that are output to the same directory share a manifest, which ensures that all of
+		/// all dependencies will be satisfied (though each application may get a few extra DLLs that they don't load).
+		/// We only ever add to the manifests, because it's not safe to delete things - we can't force all targets sharing an
+		/// output directory to build at the same time, and if targets A and B share an output directory, A has dep A' and
+		/// B has dep B', then we would make the manifest unusable for B if we cleared it of B' when building A and not B.
+		/// This does mean that we never remove dependencies that have been removed from the targets, but this is the safer
+		/// end of the compromise.
+		/// </remarks>
+		private static FileReference GetSxSIntermediateManifest(ReadOnlyTargetRules Target, DirectoryReference ExeDirectory)
+		{
+			DirectoryReference BaseDirectory = Target.ProjectFile?.Directory ?? Unreal.EngineDirectory;
+			DirectoryReference RootManifestDirectory = DirectoryReference.Combine(BaseDirectory, "Intermediate", "SxSManifest");
+
+			string RelativeOutputDirectory = ExeDirectory.MakeRelativeTo(BaseDirectory);
+			DirectoryReference ManifestDirectory = DirectoryReference.Combine(RootManifestDirectory, RelativeOutputDirectory);
 
 			string ManifestName = $"{Target.Architecture}.manifest";
-			return FileReference.Combine(IntermediateDirectory, ManifestName);
+			return FileReference.Combine(ManifestDirectory, ManifestName);
 		}
 	}
 }
diff --git a/Engine/Source/Programs/UnrealBuildTool/ToolChain/UEToolChain.cs b/Engine/Source/Programs/UnrealBuildTool/ToolChain/UEToolChain.cs
--- a/Engine/Source/Programs/UnrealBuildTool/ToolChain/UEToolChain.cs
+++ b/Engine/Source/Programs/UnrealBuildTool/ToolChain/UEToolChain.cs
@@ -329,7 +329,7 @@ namespace UnrealBuildTool
 		{
 		}
 
-		public virtual void FinalizeOutput(ReadOnlyTargetRules Target, TargetMakefileBuilder MakefileBuilder)
+		public virtual void FinalizeOutput(ReadOnlyTargetRules Target, TargetMakefileBuilder MakefileBuilder, DirectoryReference ExeDir)
 		{
 		}
 
