[UBA]
* Added prebuild target validation reporting an error if 20 or more prebuild target loops have happened. This should catch scenario where targets add prebuild target of itself

#Synced-CL 48164513

diff --git a/Engine/Source/Programs/UnrealBuildTool/Modes/BuildMode.cs b/Engine/Source/Programs/UnrealBuildTool/Modes/BuildMode.cs
--- a/Engine/Source/Programs/UnrealBuildTool/Modes/BuildMode.cs
+++ b/Engine/Source/Programs/UnrealBuildTool/Modes/BuildMode.cs
@@ -1103,8 +1103,16 @@ namespace UnrealBuildTool
 				fileHashCache.Mount(DirectoryReference.Combine(Unreal.WritableEngineDirectory, "Intermediate", "UnzippedFrameworks"), logger);
 			}
 
+			int prebuildTargetsSanityCounter = 0;
+
 			while (descriptors.Count != 0) // descriptors list will be reallocated for PrebuildTargets. hence the while-loop
 			{
+				if (prebuildTargetsSanityCounter++ > 20)
+				{
+					logger.LogError("Error: Prebuild targets has looped 20 times. Is there a prebuild target pointing back to itself somewhere? ({Descriptors})", String.Join(' ', descriptors.Select(d => d.Name)));
+					throw new CompilationResultException(CompilationResult.Unknown);
+				}
+
 				List<(TargetDescriptor descriptor, TargetMakefile makefile)> outTargets = new();
 
 				int descriptorCount = descriptors.Count;
