[UBT]
* Changed code finding files with markup to use ParallelQueue instead of ThreadPoolWorkQueue. Reason is because ThreadPoolWorkQueue is using things it shouldn't and ParallelQueue is more efficient and compliant

#Synced-CL 48175084

diff --git a/Engine/Source/Programs/UnrealBuildTool/System/TargetMakefile.cs b/Engine/Source/Programs/UnrealBuildTool/System/TargetMakefile.cs
--- a/Engine/Source/Programs/UnrealBuildTool/System/TargetMakefile.cs
+++ b/Engine/Source/Programs/UnrealBuildTool/System/TargetMakefile.cs
@@ -1094,12 +1094,13 @@ namespace UnrealBuildTool
 
 				// Find the set of files that contain reflection markup
 				ConcurrentBag<FileItem> NewFilesWithMarkupBag = new ConcurrentBag<FileItem>();
-				using (ThreadPoolWorkQueue Queue = new ThreadPoolWorkQueue())
+				ParallelQueue Queue = new();
 				{
 					foreach (DirectoryItem SourceDirectory in Makefile.SourceDirectories)
 					{
 						Queue.Enqueue(() => FindFilesWithMarkup(SourceDirectory, MetadataCache!, ExcludedFolderNames, NewFilesWithMarkupBag, Queue));
 					}
+					Queue.Drain();
 				}
 
 				// Check whether the list has changed
@@ -1183,7 +1184,7 @@ namespace UnrealBuildTool
 		/// <param name="ExcludedFolderNames">Set of folder names to ignore when recursing the directory tree</param>
 		/// <param name="FilesWithMarkup">Receives the set of files which contain reflection markup</param>
 		/// <param name="Queue">Queue to add sub-tasks to</param>
-		static void FindFilesWithMarkup(DirectoryItem Directory, SourceFileMetadataCache MetadataCache, IReadOnlySet<string> ExcludedFolderNames, ConcurrentBag<FileItem> FilesWithMarkup, ThreadPoolWorkQueue Queue)
+		static void FindFilesWithMarkup(DirectoryItem Directory, SourceFileMetadataCache MetadataCache, IReadOnlySet<string> ExcludedFolderNames, ConcurrentBag<FileItem> FilesWithMarkup, ParallelQueue Queue)
 		{
 			if (Directory.TryGetFile(".ubtignore", out FileItem? OutFile))
 			{
