UnrealBuildTool: Fix incorrect TargetMetadata extension

#rnx
#Synced-CL 47605596

diff --git a/Engine/Source/Programs/UnrealBuildTool/Modes/BuildMode.cs b/Engine/Source/Programs/UnrealBuildTool/Modes/BuildMode.cs
--- a/Engine/Source/Programs/UnrealBuildTool/Modes/BuildMode.cs
+++ b/Engine/Source/Programs/UnrealBuildTool/Modes/BuildMode.cs
@@ -838,7 +838,7 @@ namespace UnrealBuildTool
 				using ITimelineEvent MetadataTimer = Timeline.ScopeEvent("Writing Target Metadata");
 				foreach (TargetMakefile Makefile in Makefiles)
 				{
-					FileReference TargetInfoFile = FileReference.Combine(Makefile.ProjectIntermediateDirectory, "TargetMetadata.dat");
+					FileReference TargetInfoFile = FileReference.Combine(Makefile.ProjectIntermediateDirectory, "TargetMetadata.json");
 					if (FileReference.Exists(TargetInfoFile))
 					{
 						List<string> ArgumentsList = new List<string>()
