UnrealBuildTool: Fix incorrect string collection comparasion

#rnx
#Synced-CL 47605596

diff --git a/Engine/Source/Programs/UnrealBuildTool/System/FileHashCache.cs b/Engine/Source/Programs/UnrealBuildTool/System/FileHashCache.cs
--- a/Engine/Source/Programs/UnrealBuildTool/System/FileHashCache.cs
+++ b/Engine/Source/Programs/UnrealBuildTool/System/FileHashCache.cs
@@ -336,16 +336,7 @@ namespace UnrealBuildTool
 					else if (data is IEnumerable<string> contentLines)
 					{
 						string[] oldLines = FileReference.ReadAllLines(fileItem.Location, new UTF8Encoding(false));
-						int index = 0;
-						isEqual = true;
-						foreach (string str in contentLines)
-						{
-							if (index >= oldLines.Length || str != oldLines[index++])
-							{
-								isEqual = false;
-								break;
-							}
-						}
+						isEqual = oldLines.SequenceEqual(contentLines, StringComparer.Ordinal);
 					}
 
 					if (isEqual)
@@ -651,4 +642,4 @@ namespace UnrealBuildTool
 		private readonly bool _errorOnRewrite;
 		private bool _hasErrors;
 	}
-}
\ No newline at end of file
+}
