UnrealBuildTool: Disable UHT input cache

#rnx
#Synced-CL 47605596

diff --git a/Engine/Source/Programs/UnrealBuildTool/Configuration/UEBuildTarget.cs b/Engine/Source/Programs/UnrealBuildTool/Configuration/UEBuildTarget.cs
--- a/Engine/Source/Programs/UnrealBuildTool/Configuration/UEBuildTarget.cs
+++ b/Engine/Source/Programs/UnrealBuildTool/Configuration/UEBuildTarget.cs
@@ -2830,7 +2830,7 @@ namespace UnrealBuildTool
 				Makefile.UHTTargetSettings.GenerateEnumMaxValues = true;
 			}
 
-			bool bEnabledUHTInputCache = BuildConfiguration.bEnableUHTInputCache || Unreal.IsBuildMachine();
+			bool bEnabledUHTInputCache = BuildConfiguration.bEnableUHTInputCache;
 			Makefile.UHTTargetSettings.EnableInputCacheRead = bEnabledUHTInputCache && !BuildConfiguration.bForceHeaderGeneration;
 			Makefile.UHTTargetSettings.EnableInputCacheWrite = bEnabledUHTInputCache;
 
