UnrealBuildTool: Undo suppression of makefile diagnostic messages

#rnx
#Synced-CL 47605596

diff --git a/Engine/Source/Programs/UnrealBuildTool/Modes/BuildMode.cs b/Engine/Source/Programs/UnrealBuildTool/Modes/BuildMode.cs
--- a/Engine/Source/Programs/UnrealBuildTool/Modes/BuildMode.cs
+++ b/Engine/Source/Programs/UnrealBuildTool/Modes/BuildMode.cs
@@ -1019,7 +1019,7 @@ namespace UnrealBuildTool
 			{
 				foreach (string Diagnostic in Makefiles[OutputIndices[0]].Diagnostics)
 				{
-					Logger.LogDebug("{Diagnostic}", Diagnostic);
+					Logger.LogInformation("{Diagnostic}", Diagnostic);
 				}
 			}
 			else
@@ -1028,7 +1028,7 @@ namespace UnrealBuildTool
 				{
 					foreach (string Diagnostic in Makefiles[OutputIndex].Diagnostics)
 					{
-						Logger.LogDebug("{Name}: {Diagnostic}", TargetDescriptors[OutputIndex].Name, Diagnostic);
+						Logger.LogInformation("{Name}: {Diagnostic}", TargetDescriptors[OutputIndex].Name, Diagnostic);
 					}
 				}
 			}
