Fix or selectively disable issue IDE0060 (Remove unused parameter) in UBT. Configure the editorconfig to ignore unused parameters in public APIs.

#Synced-CL 48208542

diff --git a/Engine/Platforms/VisionOS/Source/Programs/UnrealBuildTool/UEDeployVisionOS.cs b/Engine/Platforms/VisionOS/Source/Programs/UnrealBuildTool/UEDeployVisionOS.cs
--- a/Engine/Platforms/VisionOS/Source/Programs/UnrealBuildTool/UEDeployVisionOS.cs
+++ b/Engine/Platforms/VisionOS/Source/Programs/UnrealBuildTool/UEDeployVisionOS.cs
@@ -105,7 +105,7 @@ namespace UnrealBuildTool
 			Text.AppendLine("\t<key>CFBundleSignature</key>");
 			Text.AppendLine("\t<string>????</string>");
 			Text.AppendLine("\t<key>CFBundleVersion</key>");
-			Text.AppendLine(string.Format("\t<string>{0}</string>", VersionUtilities.UpdateBundleVersion(OldPListData, InEngineDir)));
+			Text.AppendLine(string.Format("\t<string>{0}</string>", VersionUtilities.UpdateBundleVersion(OldPListData)));
 			Text.AppendLine("\t<key>CFBundleShortVersionString</key>");
 			Text.AppendLine(string.Format("\t<string>{0}</string>", BundleShortVersion));
 			Text.AppendLine("\t<key>LSRequiresIPhoneOS</key>");
diff --git a/Engine/Platforms/VisionOS/Source/Programs/UnrealBuildTool/VisionOSExports.cs b/Engine/Platforms/VisionOS/Source/Programs/UnrealBuildTool/VisionOSExports.cs
--- a/Engine/Platforms/VisionOS/Source/Programs/UnrealBuildTool/VisionOSExports.cs
+++ b/Engine/Platforms/VisionOS/Source/Programs/UnrealBuildTool/VisionOSExports.cs
@@ -29,7 +29,7 @@ namespace UnrealBuildTool
 		public static bool PrepForUATPackageOrDeploy(UnrealTargetConfiguration Config, FileReference ProjectFile, string InProjectName, DirectoryReference InProjectDirectory, FileReference Executable, DirectoryReference InEngineDir, bool bForDistribution, string CookFlavor, bool bIsDataDeploy, bool bCreateStubIPA, FileReference BuildReceiptFileName, ILogger Logger)
 		{
 			TargetReceipt Receipt = TargetReceipt.Read(BuildReceiptFileName);
-			return new UEDeployVisionOS(Logger).PrepForUATPackageOrDeploy(Config, ProjectFile, InProjectName, InProjectDirectory.FullName, Executable, InEngineDir.FullName, bForDistribution, CookFlavor, bIsDataDeploy, bCreateStubIPA, Receipt);
+			return new UEDeployVisionOS(Logger).PrepForUATPackageOrDeploy(Config, ProjectFile, InProjectName, InProjectDirectory.FullName, Executable, InEngineDir.FullName, bCreateStubIPA, Receipt);
 		}
 
 		/// <summary>
diff --git a/Engine/Source/Programs/UnrealBuildTool/.editorconfig b/Engine/Source/Programs/UnrealBuildTool/.editorconfig
--- a/Engine/Source/Programs/UnrealBuildTool/.editorconfig
+++ b/Engine/Source/Programs/UnrealBuildTool/.editorconfig
@@ -150,6 +150,8 @@ dotnet_diagnostic.IDE0059.severity = warning
 
 # IDE0060: Remove unused parameter
 dotnet_diagnostic.IDE0060.severity = warning
+# Don't warn about public methods that are shipped to licensees
+dotnet_code_quality_unused_parameters = non_public
 
 # CA1012: Abstract types should not have public constructors
 dotnet_diagnostic.CA1012.severity = warning
diff --git a/Engine/Source/Programs/UnrealBuildTool/Configuration/Rules/RulesAssembly.cs b/Engine/Source/Programs/UnrealBuildTool/Configuration/Rules/RulesAssembly.cs
--- a/Engine/Source/Programs/UnrealBuildTool/Configuration/Rules/RulesAssembly.cs
+++ b/Engine/Source/Programs/UnrealBuildTool/Configuration/Rules/RulesAssembly.cs
@@ -852,7 +852,7 @@ namespace UnrealBuildTool
 			}
 
 			// Some platforms may *require* monolithic compilation...
-			if (Rules.LinkType != TargetLinkType.Monolithic && UEBuildPlatform.PlatformRequiresMonolithicBuilds(Rules.Platform, Rules.Configuration))
+			if (Rules.LinkType != TargetLinkType.Monolithic && UEBuildPlatform.PlatformRequiresMonolithicBuilds(Rules.Platform))
 			{
 				throw new CompilationResultException(CompilationResult.RulesError, KnownLogEvents.RulesAssembly, "{RulesName}: {RulesPlatform} does not support modular builds", Rules.Name, Rules.Platform);
 			}
diff --git a/Engine/Source/Programs/UnrealBuildTool/Configuration/UEBuildBinary.cs b/Engine/Source/Programs/UnrealBuildTool/Configuration/UEBuildBinary.cs
--- a/Engine/Source/Programs/UnrealBuildTool/Configuration/UEBuildBinary.cs
+++ b/Engine/Source/Programs/UnrealBuildTool/Configuration/UEBuildBinary.cs
@@ -823,7 +823,7 @@ namespace UnrealBuildTool
 			// Gather debug visualizers for every relevant module 
 			foreach (UEBuildModule VisitedModule in LinkEnvironmentVisitedModules)
 			{
-				VisitedModule.LinkDebuggerVisualizers(BinaryLinkEnvironment.DebuggerVisualizerFiles, ToolChain, Logger);
+				VisitedModule.LinkDebuggerVisualizers(BinaryLinkEnvironment.DebuggerVisualizerFiles, ToolChain);
 			}
 
 			// Set the link output file.
@@ -846,7 +846,7 @@ namespace UnrealBuildTool
 			BinaryLinkEnvironment.bIsBuildingLibrary = IsBuildingLibrary(Type);
 
 			// Compile any resource files
-			CompileResourceFiles(Target, ToolChain, BinaryCompileEnvironment, BinaryLinkEnvironment, [.. CompiledModules.SelectMany(x => x.ResourceFiles)], Graph, Logger);
+			CompileResourceFiles(ToolChain, BinaryCompileEnvironment, BinaryLinkEnvironment, [.. CompiledModules.SelectMany(x => x.ResourceFiles)], Graph);
 
 			// Setup object export stripping
 			if (Target.bStripExports)
@@ -871,7 +871,7 @@ namespace UnrealBuildTool
 			return BinaryLinkEnvironment;
 		}
 
-		private void CompileResourceFiles(ReadOnlyTargetRules Target, UEToolChain ToolChain, CppCompileEnvironment BinaryCompileEnvironment, LinkEnvironment BinaryLinkEnvironment, IEnumerable<FileItem> InputFiles, IActionGraphBuilder Graph, ILogger Logger)
+		private void CompileResourceFiles(UEToolChain ToolChain, CppCompileEnvironment BinaryCompileEnvironment, LinkEnvironment BinaryLinkEnvironment, IEnumerable<FileItem> InputFiles, IActionGraphBuilder Graph)
 		{
 			// @todo: This should be in some Windows code somewhere...
 			if (!BinaryLinkEnvironment.Platform.IsInGroup(UnrealPlatformGroup.Windows))
diff --git a/Engine/Source/Programs/UnrealBuildTool/Configuration/UEBuildModule.cs b/Engine/Source/Programs/UnrealBuildTool/Configuration/UEBuildModule.cs
--- a/Engine/Source/Programs/UnrealBuildTool/Configuration/UEBuildModule.cs
+++ b/Engine/Source/Programs/UnrealBuildTool/Configuration/UEBuildModule.cs
@@ -1052,7 +1052,7 @@ namespace UnrealBuildTool
 			return [];
 		}
 
-		public IEnumerable<FileItem> CopyDebuggerVisualizers(UEToolChain ToolChain, IActionGraphBuilder Graph, ILogger Logger)
+		public IEnumerable<FileItem> CopyDebuggerVisualizers(UEToolChain ToolChain, IActionGraphBuilder Graph)
 		{
 			List<FileItem> Results = new();
 			foreach (FileItem NatvisSourceFile in NatvisFiles)
@@ -1066,7 +1066,7 @@ namespace UnrealBuildTool
 			return Results;
 		}
 
-		public void LinkDebuggerVisualizers(List<FileItem> OutFiles, UEToolChain ToolChain, ILogger Logger)
+		public void LinkDebuggerVisualizers(List<FileItem> OutFiles, UEToolChain ToolChain)
 		{
 			foreach (FileItem NatvisSourceFile in NatvisFiles)
 			{
diff --git a/Engine/Source/Programs/UnrealBuildTool/Configuration/UEBuildModuleCPP.cs b/Engine/Source/Programs/UnrealBuildTool/Configuration/UEBuildModuleCPP.cs
--- a/Engine/Source/Programs/UnrealBuildTool/Configuration/UEBuildModuleCPP.cs
+++ b/Engine/Source/Programs/UnrealBuildTool/Configuration/UEBuildModuleCPP.cs
@@ -658,7 +658,7 @@ namespace UnrealBuildTool
 			}
 
 			// Configure the precompiled headers for this module
-			CppCompileEnvironment = SetupPrecompiledHeaders(Target, ToolChain, CppCompileEnvironment, LinkInputFiles, Graph, Logger);
+			CppCompileEnvironment = SetupPrecompiledHeaders(Target, ToolChain, CppCompileEnvironment, LinkInputFiles, Graph);
 			if (CppCompileEnvironment.PerArchPrecompiledHeaderFiles != null)
 			{
 				foreach (UnrealArch Arch in CppCompileEnvironment.PerArchPrecompiledHeaderFiles.Keys)
@@ -993,13 +993,13 @@ namespace UnrealBuildTool
 					out List<FileItem> NormalGeneratedFiles, out List<FileItem> AdaptiveGeneratedFiles, NumIncludedBytesPerUnityCPP);
 				ModuleCompileEnvironment.FileInlineSourceMap = new(GeneratedCPPCompileEnvironment.FileInlineSourceMap);
 				CppCompileEnvironment.FileInlineSourceMap = new(GeneratedCPPCompileEnvironment.FileInlineSourceMap);
-				LinkInputFiles.AddRange(CompileFilesWithToolChain(Target, ToolChain, GeneratedCPPCompileEnvironment, ModuleCompileEnvironment, NormalGeneratedFiles, AdaptiveGeneratedFiles, Graph, Logger).ObjectFiles);
+				LinkInputFiles.AddRange(CompileFilesWithToolChain(Target, ToolChain, GeneratedCPPCompileEnvironment, ModuleCompileEnvironment, NormalGeneratedFiles, AdaptiveGeneratedFiles, Graph).ObjectFiles);
 			}
 			else
 			{
-				Unity.GetAdaptiveFiles(Target, GeneratedFileItems, InputFiles.HeaderFiles, GeneratedCPPCompileEnvironment, WorkingSet, (Rules.ShortName ?? Name) + ".gen", IntermediateDirectory, Graph,
+				Unity.GetAdaptiveFiles(Target, GeneratedFileItems, InputFiles.HeaderFiles, GeneratedCPPCompileEnvironment, WorkingSet, IntermediateDirectory, Graph,
 					out List<FileItem> NormalGeneratedFiles, out List<FileItem> AdaptiveGeneratedFiles);
-				LinkInputFiles.AddRange(CompileFilesWithToolChain(Target, ToolChain, GeneratedCPPCompileEnvironment, ModuleCompileEnvironment, NormalGeneratedFiles, AdaptiveGeneratedFiles, Graph, Logger).ObjectFiles);
+				LinkInputFiles.AddRange(CompileFilesWithToolChain(Target, ToolChain, GeneratedCPPCompileEnvironment, ModuleCompileEnvironment, NormalGeneratedFiles, AdaptiveGeneratedFiles, Graph).ObjectFiles);
 			}
 
 			// Compile Swift - which may generate a header that the cpp files need below
@@ -1042,18 +1042,18 @@ namespace UnrealBuildTool
 					Unity.GenerateUnitySource(Target, CPPFiles, InputFiles.HeaderFiles, CppCompileEnvironment, WorkingSet, Rules.ShortName ?? Name, IntermediateDirectory, Graph, SourceFileToUnityFile,
 						out List<FileItem> NormalFiles, out List<FileItem> AdaptiveFiles, NumIncludedBytesPerUnityCPP);
 					ModuleCompileEnvironment.FileInlineSourceMap = new(CppCompileEnvironment.FileInlineSourceMap);
-					LinkInputFiles.AddRange(CompileFilesWithToolChain(Target, ToolChain, CppCompileEnvironment, ModuleCompileEnvironment, NormalFiles, AdaptiveFiles, Graph, Logger).ObjectFiles);
+					LinkInputFiles.AddRange(CompileFilesWithToolChain(Target, ToolChain, CppCompileEnvironment, ModuleCompileEnvironment, NormalFiles, AdaptiveFiles, Graph).ObjectFiles);
 				}
 				else
 				{
-					Unity.GetAdaptiveFiles(Target, CPPFiles, InputFiles.HeaderFiles, CppCompileEnvironment, WorkingSet, Rules.ShortName ?? Name, IntermediateDirectory, Graph,
+					Unity.GetAdaptiveFiles(Target, CPPFiles, InputFiles.HeaderFiles, CppCompileEnvironment, WorkingSet, IntermediateDirectory, Graph,
 						out List<FileItem> NormalFiles, out List<FileItem> AdaptiveFiles);
 					if (!NormalFiles.Where(file => !file.HasExtension(".gen.cpp") && !GeneratedCPPCompileEnvironment.FileMatchesExtraGeneratedCPPTypes(file.FullName)).Any())
 					{
 						NormalFiles = CPPFiles;
 						AdaptiveFiles.RemoveAll(new HashSet<FileItem>(NormalFiles).Contains);
 					}
-					LinkInputFiles.AddRange(CompileFilesWithToolChain(Target, ToolChain, CppCompileEnvironment, ModuleCompileEnvironment, NormalFiles, AdaptiveFiles, Graph, Logger).ObjectFiles);
+					LinkInputFiles.AddRange(CompileFilesWithToolChain(Target, ToolChain, CppCompileEnvironment, ModuleCompileEnvironment, NormalFiles, AdaptiveFiles, Graph).ObjectFiles);
 				}
 			}
 
@@ -1075,7 +1075,7 @@ namespace UnrealBuildTool
 				{
 					Unity.GenerateUnitySource(Target, InputFiles.CFiles, InputFiles.HeaderFiles, CCompileEnvironment, WorkingSet, Rules.ShortName ?? Name, IntermediateDirectory, Graph, SourceFileToUnityFile,
 						out List<FileItem> NormalFiles, out List<FileItem> AdaptiveFiles, NumIncludedBytesPerUnityCPP);
-					LinkInputFiles.AddRange(CompileFilesWithToolChain(Target, ToolChain, CCompileEnvironment, ModuleCompileEnvironment, NormalFiles, AdaptiveFiles, Graph, Logger).ObjectFiles);
+					LinkInputFiles.AddRange(CompileFilesWithToolChain(Target, ToolChain, CCompileEnvironment, ModuleCompileEnvironment, NormalFiles, AdaptiveFiles, Graph).ObjectFiles);
 				}
 				else
 				{
@@ -1094,7 +1094,7 @@ namespace UnrealBuildTool
 				{
 					Unity.GenerateUnitySource(Target, InputFiles.CCFiles, InputFiles.HeaderFiles, CCCompileEnvironment, WorkingSet, Rules.ShortName ?? Name, IntermediateDirectory, Graph, SourceFileToUnityFile,
 						out List<FileItem> NormalFiles, out List<FileItem> AdaptiveFiles, NumIncludedBytesPerUnityCPP);
-					LinkInputFiles.AddRange(CompileFilesWithToolChain(Target, ToolChain, CCCompileEnvironment, CppCompileEnvironment, NormalFiles, AdaptiveFiles, Graph, Logger).ObjectFiles);
+					LinkInputFiles.AddRange(CompileFilesWithToolChain(Target, ToolChain, CCCompileEnvironment, CppCompileEnvironment, NormalFiles, AdaptiveFiles, Graph).ObjectFiles);
 				}
 				else
 				{
@@ -1113,7 +1113,7 @@ namespace UnrealBuildTool
 				{
 					Unity.GenerateUnitySource(Target, InputFiles.CXXFiles, InputFiles.HeaderFiles, CXXCompileEnvironment, WorkingSet, Rules.ShortName ?? Name, IntermediateDirectory, Graph, SourceFileToUnityFile,
 						out List<FileItem> NormalFiles, out List<FileItem> AdaptiveFiles, NumIncludedBytesPerUnityCPP);
-					LinkInputFiles.AddRange(CompileFilesWithToolChain(Target, ToolChain, CXXCompileEnvironment, CppCompileEnvironment, NormalFiles, AdaptiveFiles, Graph, Logger).ObjectFiles);
+					LinkInputFiles.AddRange(CompileFilesWithToolChain(Target, ToolChain, CXXCompileEnvironment, CppCompileEnvironment, NormalFiles, AdaptiveFiles, Graph).ObjectFiles);
 				}
 				else
 				{
@@ -1132,7 +1132,7 @@ namespace UnrealBuildTool
 				{
 					Unity.GenerateUnitySource(Target, InputFiles.MFiles, InputFiles.HeaderFiles, MCompileEnvironment, WorkingSet, Rules.ShortName ?? Name, IntermediateDirectory, Graph, SourceFileToUnityFile,
 						out List<FileItem> NormalFiles, out List<FileItem> AdaptiveFiles, NumIncludedBytesPerUnityCPP);
-					LinkInputFiles.AddRange(CompileFilesWithToolChain(Target, ToolChain, MCompileEnvironment, ModuleCompileEnvironment, NormalFiles, AdaptiveFiles, Graph, Logger).ObjectFiles);
+					LinkInputFiles.AddRange(CompileFilesWithToolChain(Target, ToolChain, MCompileEnvironment, ModuleCompileEnvironment, NormalFiles, AdaptiveFiles, Graph).ObjectFiles);
 				}
 				else
 				{
@@ -1151,7 +1151,7 @@ namespace UnrealBuildTool
 				{
 					Unity.GenerateUnitySource(Target, InputFiles.MMFiles, InputFiles.HeaderFiles, MMCompileEnvironment, WorkingSet, Rules.ShortName ?? Name, IntermediateDirectory, Graph, SourceFileToUnityFile,
 						out List<FileItem> NormalFiles, out List<FileItem> AdaptiveFiles, NumIncludedBytesPerUnityCPP);
-					LinkInputFiles.AddRange(CompileFilesWithToolChain(Target, ToolChain, MMCompileEnvironment, CppCompileEnvironment, NormalFiles, AdaptiveFiles, Graph, Logger).ObjectFiles);
+					LinkInputFiles.AddRange(CompileFilesWithToolChain(Target, ToolChain, MMCompileEnvironment, CppCompileEnvironment, NormalFiles, AdaptiveFiles, Graph).ObjectFiles);
 				}
 				else
 				{
@@ -1979,8 +1979,7 @@ namespace UnrealBuildTool
 			CppCompileEnvironment ModuleCompileEnvironment,
 			List<FileItem> NormalFiles,
 			List<FileItem> AdaptiveFiles,
-			IActionGraphBuilder Graph,
-			ILogger Logger)
+			IActionGraphBuilder Graph)
 		{
 			bool isEngine = Rules.File.IsUnderDirectory(Unreal.EngineDirectory) || Rules.bTreatAsEngineModule || Rules.PrivatePCHHeaderFile == null;
 			bool bAdaptiveUnityDisablesPCH = isEngine ? Target.bAdaptiveUnityDisablesPCH : Target.bAdaptiveUnityDisablesPCHForProject;
@@ -2138,8 +2137,7 @@ namespace UnrealBuildTool
 		/// <param name="CompileEnvironment">The current compile environment</param>
 		/// <param name="LinkInputFiles">List of files that will be linked for the target</param>
 		/// <param name="Graph">List of build actions</param>
-		/// <param name="Logger">Logger for output</param>
-		CppCompileEnvironment SetupPrecompiledHeaders(ReadOnlyTargetRules Target, UEToolChain? ToolChain, CppCompileEnvironment CompileEnvironment, List<FileItem> LinkInputFiles, IActionGraphBuilder Graph, ILogger Logger)
+		CppCompileEnvironment SetupPrecompiledHeaders(ReadOnlyTargetRules Target, UEToolChain? ToolChain, CppCompileEnvironment CompileEnvironment, List<FileItem> LinkInputFiles, IActionGraphBuilder Graph)
 		{
 			if (Target.bUsePCHFiles && Rules.PCHUsage != ModuleRules.PCHUsageMode.NoPCHs)
 			{
@@ -2564,7 +2562,7 @@ namespace UnrealBuildTool
 		{
 			builder ??= new NullActionGraphBuilder(Logger);
 			CppCompileEnvironment CompileEnvironment = CreateModuleCompileEnvironment(Target, BaseCompileEnvironment, Logger);
-			CompileEnvironment = SetupPrecompiledHeaders(Target, null, CompileEnvironment, new List<FileItem>(), builder, Logger);
+			CompileEnvironment = SetupPrecompiledHeaders(Target, null, CompileEnvironment, new List<FileItem>(), builder);
 			CreateHeaderForDefinitions(CompileEnvironment, IntermediateDirectory, null, builder);
 			return CompileEnvironment;
 		}
diff --git a/Engine/Source/Programs/UnrealBuildTool/Configuration/UEBuildPlatform.cs b/Engine/Source/Programs/UnrealBuildTool/Configuration/UEBuildPlatform.cs
--- a/Engine/Source/Programs/UnrealBuildTool/Configuration/UEBuildPlatform.cs
+++ b/Engine/Source/Programs/UnrealBuildTool/Configuration/UEBuildPlatform.cs
@@ -1061,9 +1061,8 @@ namespace UnrealBuildTool
 		/// Return whether the given platform requires a monolithic build
 		/// </summary>
 		/// <param name="InPlatform">The platform of interest</param>
-		/// <param name="InConfiguration">The configuration of interest</param>
 		/// <returns></returns>
-		public static bool PlatformRequiresMonolithicBuilds(UnrealTargetPlatform InPlatform, UnrealTargetConfiguration InConfiguration)
+		public static bool PlatformRequiresMonolithicBuilds(UnrealTargetPlatform InPlatform)
 		{
 			// Some platforms require monolithic builds...
 			UEBuildPlatform? BuildPlatform;
@@ -1145,7 +1144,7 @@ namespace UnrealBuildTool
 			return Configurations;
 		}
 
-		protected static bool DoProjectSettingsMatchDefault(UnrealTargetPlatform Platform, DirectoryReference ProjectDirectoryName, string Section, string[]? BoolKeys, string[]? IntKeys, string[]? StringKeys, ILogger Logger)
+		protected static bool DoProjectSettingsMatchDefault(UnrealTargetPlatform Platform, DirectoryReference ProjectDirectoryName, string Section, string[]? BoolKeys, string[]? IntKeys, string[]? StringKeys)
 		{
 			ConfigHierarchy ProjIni = ConfigCache.ReadHierarchy(ConfigHierarchyType.Engine, ProjectDirectoryName, Platform);
 			ConfigHierarchy DefaultIni = ConfigCache.ReadHierarchy(ConfigHierarchyType.Engine, (DirectoryReference?)null, Platform);
@@ -1213,7 +1212,7 @@ namespace UnrealBuildTool
 			};
 
 			return DoProjectSettingsMatchDefault(Platform, ProjectDirectoryName, "/Script/BuildSettings.BuildSettings",
-				BoolKeys, null, null, Logger);
+				BoolKeys, null, null);
 		}
 
 		/// <summary>
diff --git a/Engine/Source/Programs/UnrealBuildTool/Configuration/UEBuildTarget.cs b/Engine/Source/Programs/UnrealBuildTool/Configuration/UEBuildTarget.cs
--- a/Engine/Source/Programs/UnrealBuildTool/Configuration/UEBuildTarget.cs
+++ b/Engine/Source/Programs/UnrealBuildTool/Configuration/UEBuildTarget.cs
@@ -1343,7 +1343,7 @@ namespace UnrealBuildTool
 			UEBuildTarget Target;
 			using (GlobalTracer.Instance.BuildSpan("UEBuildTarget constructor").StartActive())
 			{
-				Target = new UEBuildTarget(Descriptor, new ReadOnlyTargetRules(RulesObject), RulesAssembly, Logger);
+				Target = new UEBuildTarget(Descriptor, new ReadOnlyTargetRules(RulesObject), RulesAssembly);
 			}
 			using (GlobalTracer.Instance.BuildSpan("UEBuildTarget.PreBuildSetup()").StartActive())
 			{
@@ -1424,7 +1424,7 @@ namespace UnrealBuildTool
 			{
 				try
 				{
-					ValidateSharedEnvironment(RulesAssembly, Descriptor.Name, Descriptor.AdditionalArguments, RulesObject, Logger);
+					ValidateSharedEnvironment(RulesAssembly, Descriptor.Name, Descriptor.AdditionalArguments, RulesObject);
 				}
 				catch (Exception)
 				{
@@ -1503,7 +1503,7 @@ namespace UnrealBuildTool
 		/// <summary>
 		/// Validates that the build environment matches the shared build environment, by comparing the TargetRules instance to the vanilla target rules for the current target type.
 		/// </summary>
-		static void ValidateSharedEnvironment(RulesAssembly RulesAssembly, string ThisTargetName, CommandLineArguments Arguments, TargetRules ThisRules, ILogger Logger)
+		static void ValidateSharedEnvironment(RulesAssembly RulesAssembly, string ThisTargetName, CommandLineArguments Arguments, TargetRules ThisRules)
 		{
 			Dictionary<string, (string?, string?)> PropNamesThatRequiredUnique = new();
 			string? BaseTargetName;
@@ -1881,8 +1881,7 @@ namespace UnrealBuildTool
 		/// <param name="InDescriptor">Target descriptor</param>
 		/// <param name="InRules">The target rules, as created by RulesCompiler.</param>
 		/// <param name="InRulesAssembly">The chain of rules assemblies that this target was created with</param>
-		/// <param name="Logger">Logger for output</param>
-		internal UEBuildTarget(TargetDescriptor InDescriptor, ReadOnlyTargetRules InRules, RulesAssembly InRulesAssembly, ILogger Logger)
+		internal UEBuildTarget(TargetDescriptor InDescriptor, ReadOnlyTargetRules InRules, RulesAssembly InRulesAssembly)
 		{
 			ProjectFile = InDescriptor.ProjectFile;
 			TargetName = InRules.Name;
@@ -2588,7 +2587,7 @@ namespace UnrealBuildTool
 		{
 			try
 			{
-				UEBuildTarget incompleteTarget = new UEBuildTarget(Descriptor, new ReadOnlyTargetRules(RulesObject), RulesAssembly, Logger);
+				UEBuildTarget incompleteTarget = new UEBuildTarget(Descriptor, new ReadOnlyTargetRules(RulesObject), RulesAssembly);
 				UEToolChain targetToolChain = incompleteTarget.CreateToolchain(incompleteTarget.Platform, Logger);
 				TargetReceipt incompleteTargetReceipt = incompleteTarget.PrepareReceipt(targetToolChain, [], []);
 
@@ -3082,7 +3081,7 @@ namespace UnrealBuildTool
 			// Copy debugger visualizer files for each module to their intermediate directory 
 			foreach (UEBuildModule Module in Modules.Values)
 			{
-				IEnumerable<FileItem> Items = Module.CopyDebuggerVisualizers(TargetToolChain, MakefileBuilder, Logger);
+				IEnumerable<FileItem> Items = Module.CopyDebuggerVisualizers(TargetToolChain, MakefileBuilder);
 				MakefileBuilder.AddOutputFiles(Items);
 			}
 
@@ -5396,7 +5395,7 @@ namespace UnrealBuildTool
 				ModuleConfiguration = UnrealTargetConfiguration.Development;
 			}
 
-			List<FileReference> OutputPaths = MakeBinaryPaths(OutputDirectory, AppName + Rules.DecoratedSeparator + MergedName, Platform, Configuration, UEBuildBinaryType.DynamicLinkLibrary, Rules.Architectures, Rules.UndecoratedConfiguration, bCompileMonolithic && ProjectFile != null, Rules.ExeBinariesSubFolder, ProjectFile, Rules);
+			List<FileReference> OutputPaths = MakeBinaryPaths(OutputDirectory, AppName + Rules.DecoratedSeparator + MergedName, Platform, Configuration, UEBuildBinaryType.DynamicLinkLibrary, Rules.Architectures, Rules.UndecoratedConfiguration, Rules.ExeBinariesSubFolder, ProjectFile, Rules);
 
 			// Create the binary
 			UEBuildBinary Binary = new UEBuildBinary(
@@ -5472,12 +5471,11 @@ namespace UnrealBuildTool
 		/// <param name="Architectures">Architectures being built</param>
 		/// <param name="BinaryType">The type of binary we're compiling</param>
 		/// <param name="UndecoratedConfiguration">The configuration which doesn't have a "-{Platform}-{Configuration}" suffix added to the binary</param>
-		/// <param name="bIncludesGameModules">Whether this executable contains game modules</param>
 		/// <param name="ExeSubFolder">Subfolder for executables. May be null.</param>
 		/// <param name="ProjectFile">The project file containing the target being built</param>
 		/// <param name="Rules">Rules for the target being built</param>
 		/// <returns>List of executable paths for this target</returns>
-		public List<FileReference> MakeBinaryPaths(DirectoryReference BaseDirectory, string BinaryName, UnrealTargetPlatform Platform, UnrealTargetConfiguration Configuration, UEBuildBinaryType BinaryType, UnrealArchitectures Architectures, UnrealTargetConfiguration UndecoratedConfiguration, bool bIncludesGameModules, string ExeSubFolder, FileReference? ProjectFile, ReadOnlyTargetRules Rules)
+		public List<FileReference> MakeBinaryPaths(DirectoryReference BaseDirectory, string BinaryName, UnrealTargetPlatform Platform, UnrealTargetConfiguration Configuration, UEBuildBinaryType BinaryType, UnrealArchitectures Architectures, UnrealTargetConfiguration UndecoratedConfiguration, string ExeSubFolder, FileReference? ProjectFile, ReadOnlyTargetRules Rules)
 		{
 			FileReference BinaryFile;
 			if (Rules.OutputFile != null)
@@ -6089,7 +6087,7 @@ namespace UnrealBuildTool
 			}
 
 			bool bCompileAsDLL = Rules.bShouldCompileAsDLL && bCompileMonolithic;
-			List<FileReference> OutputPaths = MakeBinaryPaths(OutputDirectory, bCompileMonolithic ? TargetName : AppName, Platform, Configuration, bCompileAsDLL ? UEBuildBinaryType.DynamicLinkLibrary : UEBuildBinaryType.Executable, Rules.Architectures, Rules.UndecoratedConfiguration, bCompileMonolithic && ProjectFile != null, Rules.ExeBinariesSubFolder, ProjectFile, Rules);
+			List<FileReference> OutputPaths = MakeBinaryPaths(OutputDirectory, bCompileMonolithic ? TargetName : AppName, Platform, Configuration, bCompileAsDLL ? UEBuildBinaryType.DynamicLinkLibrary : UEBuildBinaryType.Executable, Rules.Architectures, Rules.UndecoratedConfiguration, Rules.ExeBinariesSubFolder, ProjectFile, Rules);
 
 			// Create the binary
 			UEBuildBinary Binary = new UEBuildBinary(
diff --git a/Engine/Source/Programs/UnrealBuildTool/Executors/Experimental/FASTBuild.cs b/Engine/Source/Programs/UnrealBuildTool/Executors/Experimental/FASTBuild.cs
--- a/Engine/Source/Programs/UnrealBuildTool/Executors/Experimental/FASTBuild.cs
+++ b/Engine/Source/Programs/UnrealBuildTool/Executors/Experimental/FASTBuild.cs
@@ -1248,7 +1248,7 @@ namespace UnrealBuildTool
 			AddText("}\n\n");
 		}
 
-		private void AddExecAction(LinkedAction Action, IEnumerable<LinkedAction> DependencyActions, ILogger Logger)
+		private void AddExecAction(LinkedAction Action, IEnumerable<LinkedAction> DependencyActions)
 		{
 			AddText($"Exec('{ActionToActionString(Action)}')\n{{\n");
 			AddText($"\t.ExecExecutable = '{Action.CommandPath.FullName}' \n");
@@ -1300,7 +1300,7 @@ namespace UnrealBuildTool
 					}
 					else
 					{
-						AddExecAction(Action, DependencyActions, Logger);
+						AddExecAction(Action, DependencyActions);
 					}
 				}
 
diff --git a/Engine/Source/Programs/UnrealBuildTool/Executors/UnrealBuildAccelerator/UBAAgentCoordinatorHorde.cs b/Engine/Source/Programs/UnrealBuildTool/Executors/UnrealBuildAccelerator/UBAAgentCoordinatorHorde.cs
--- a/Engine/Source/Programs/UnrealBuildTool/Executors/UnrealBuildAccelerator/UBAAgentCoordinatorHorde.cs
+++ b/Engine/Source/Programs/UnrealBuildTool/Executors/UnrealBuildAccelerator/UBAAgentCoordinatorHorde.cs
@@ -818,7 +818,7 @@ namespace UnrealBuildTool
 			_hordeSessionTask = UBAHordeSession.TryCreateHordeSessionAsync(_hordeConfig, _ubaConfig, executor.Crypto, _ubaConfig.bStrict, _logger, cancellationToken);
 			await _hordeSessionTask;
 			_updateStatus!(0, 6, _hordeSessionTask.Result.Item2, EpicGames.UBA.LogEntryType.Info, null);
-			executor.AgentCoordinatorInitialized(this, _hordeSessionTask.IsCompletedSuccessfully && _hordeSessionTask.Result.Item1 != null);
+			executor.AgentCoordinatorInitialized(_hordeSessionTask.IsCompletedSuccessfully && _hordeSessionTask.Result.Item1 != null);
 		}
 
 		public UnrealBuildAcceleratorCacheConfig? RequestCacheServer(CancellationToken cancellationToken)
diff --git a/Engine/Source/Programs/UnrealBuildTool/Executors/UnrealBuildAccelerator/UBAExecutor.cs b/Engine/Source/Programs/UnrealBuildTool/Executors/UnrealBuildAccelerator/UBAExecutor.cs
--- a/Engine/Source/Programs/UnrealBuildTool/Executors/UnrealBuildAccelerator/UBAExecutor.cs
+++ b/Engine/Source/Programs/UnrealBuildTool/Executors/UnrealBuildAccelerator/UBAExecutor.cs
@@ -473,7 +473,7 @@ namespace UnrealBuildTool
 			return BitConverter.ToString(bytes).Replace("-", "", StringComparison.OrdinalIgnoreCase).ToLowerInvariant(); // "1234567890abcdef1234567890abcdef";
 		}
 
-		public void AgentCoordinatorInitialized(IUBAAgentCoordinator coordinator, bool successful)
+		public void AgentCoordinatorInitialized(bool successful)
 		{
 			if (successful)
 			{
diff --git a/Engine/Source/Programs/UnrealBuildTool/Modes/BuildMode.cs b/Engine/Source/Programs/UnrealBuildTool/Modes/BuildMode.cs
--- a/Engine/Source/Programs/UnrealBuildTool/Modes/BuildMode.cs
+++ b/Engine/Source/Programs/UnrealBuildTool/Modes/BuildMode.cs
@@ -326,19 +326,18 @@ namespace UnrealBuildTool
 
 			DirectoryReference ProjectIntermediateDirectory = DirectoryReference.Combine(ProjectDirectory, PlatformIntermediateFolder, TargetDescriptor.Name, TargetDescriptor.Configuration.ToString());
 
-			return WriteInitScripts(TargetDescriptor, BuildHostPlatform.Current.Platform, ProjectIntermediateDirectory, "Init", InitCommandBatches);
+			return WriteInitScripts(TargetDescriptor, ProjectIntermediateDirectory, "Init", InitCommandBatches);
 		}
 
 		/// <summary>
-		/// Write scripts containing the custom build steps for the given host platform
+		/// Write scripts containing the custom build steps for the host platform
 		/// </summary>
 		/// <param name="TargetDescriptor">The current target</param>
-		/// <param name="HostPlatform">The current host platform</param>
 		/// <param name="Directory">The output directory for the scripts</param>
 		/// <param name="FilePrefix">Bare prefix for all the created script files</param>
 		/// <param name="CommandBatches">List of custom build steps</param>
 		/// <returns>List of created script files</returns>
-		private static FileReference[] WriteInitScripts(TargetDescriptor TargetDescriptor, UnrealTargetPlatform HostPlatform, DirectoryReference Directory, string FilePrefix, List<string[]> CommandBatches)
+		private static FileReference[] WriteInitScripts(TargetDescriptor TargetDescriptor, DirectoryReference Directory, string FilePrefix, List<string[]> CommandBatches)
 		{
 			List<FileReference> ScriptFiles = new List<FileReference>();
 			foreach (string[] CommandBatch in CommandBatches)
@@ -407,10 +406,9 @@ namespace UnrealBuildTool
 		/// <param name="WriteOutdatedActionsFile">Files to write the list of outdated actions to (rather than building them)</param>
 		/// <param name="cppDependencyCache"></param>
 		/// <param name="Logger">Logger for output</param>
-		/// <param name="bSkipPreBuildTargets">If true then only the current target descriptors will be built.</param>
 		/// <param name="ActionTypeFilter">If specified, only actions of this type will be run</param>
 		/// <returns>Result from the compilation</returns>
-		public static async Task BuildAsync(List<TargetDescriptor> TargetDescriptors, BuildConfiguration BuildConfiguration, ISourceFileWorkingSet WorkingSet, BuildOptions Options, FileReference? WriteOutdatedActionsFile, CppDependencyCache cppDependencyCache, ILogger Logger, bool bSkipPreBuildTargets = false, string? ActionTypeFilter = null)
+		public static async Task BuildAsync(List<TargetDescriptor> TargetDescriptors, BuildConfiguration BuildConfiguration, ISourceFileWorkingSet WorkingSet, BuildOptions Options, FileReference? WriteOutdatedActionsFile, CppDependencyCache cppDependencyCache, ILogger Logger, string? ActionTypeFilter = null)
 		{
 			List<Task> tasks = [];
 			try
@@ -1166,7 +1164,7 @@ namespace UnrealBuildTool
 					(TargetRules rules, RulesAssembly assembly) = UEBuildTarget.CreateRulesObject(descriptor, skipRulesCompile, forceRulesCompile, usePrecompiled, descriptor.IntermediateEnvironment, logger);
 
 					// Create the target. Ctor does not do much so we can create it regardless if we are building it or not
-					UEBuildTarget buildTarget = new(descriptor, new ReadOnlyTargetRules(rules), assembly, logger);
+					UEBuildTarget buildTarget = new(descriptor, new ReadOnlyTargetRules(rules), assembly);
 					buildTargets[index] = buildTarget;
 
 					// Queue up cppdependencies loading
diff --git a/Engine/Source/Programs/UnrealBuildTool/Modes/QueryMode.cs b/Engine/Source/Programs/UnrealBuildTool/Modes/QueryMode.cs
--- a/Engine/Source/Programs/UnrealBuildTool/Modes/QueryMode.cs
+++ b/Engine/Source/Programs/UnrealBuildTool/Modes/QueryMode.cs
@@ -168,7 +168,7 @@ namespace UnrealBuildTool
 			{
 				case QueryType.Capabilities:
 					Logger.LogInformation("QueryCapabilities");
-					return QueryCapabilities(Arguments, Logger, ResponseOptions);
+					return QueryCapabilities(Logger, ResponseOptions);
 				case QueryType.AvailableTargets:
 					Logger.LogInformation("QueryAvailableTargets");
 					return QueryAvailableTargetsAsync(Arguments, Logger, ResponseOptions);
@@ -195,7 +195,7 @@ namespace UnrealBuildTool
 			await JsonSerializer.SerializeAsync(Stream, Value, JsonOptions);
 		}
 
-		private async Task<int> QueryCapabilities(CommandLineArguments Arguments, ILogger Logger, JsonSerializerOptions JsonOptions)
+		private async Task<int> QueryCapabilities(ILogger Logger, JsonSerializerOptions JsonOptions)
 		{
 			var Reply = new
 			{
diff --git a/Engine/Source/Programs/UnrealBuildTool/Modes/ServerMode.cs b/Engine/Source/Programs/UnrealBuildTool/Modes/ServerMode.cs
--- a/Engine/Source/Programs/UnrealBuildTool/Modes/ServerMode.cs
+++ b/Engine/Source/Programs/UnrealBuildTool/Modes/ServerMode.cs
@@ -217,21 +217,21 @@ namespace UnrealBuildTool
 						break;
 					case CommandType.QueryFile:
 						{
-							QueryFileResponse response = Handle_QueryFile(JsonSerializer.Deserialize<QueryFileCommand>(CommandString, jsonOptions), Logger);
+							QueryFileResponse response = Handle_QueryFile(JsonSerializer.Deserialize<QueryFileCommand>(CommandString, jsonOptions));
 							response.Sequence = command.Sequence;
 							Console.WriteLine(JsonSerializer.Serialize(response, responseOptions));
 						}
 						break;
 					case CommandType.GetBrowseConfiguration:
 						{
-							GetBrowseConfigurationResponse response = Handle_GetBrowseConfiguration(JsonSerializer.Deserialize<GetBrowseConfigurationCommand>(CommandString, jsonOptions), Logger);
+							GetBrowseConfigurationResponse response = Handle_GetBrowseConfiguration(JsonSerializer.Deserialize<GetBrowseConfigurationCommand>(CommandString, jsonOptions));
 							response.Sequence = command.Sequence;
 							Console.WriteLine(JsonSerializer.Serialize(response, responseOptions));
 						}
 						break;
 					case CommandType.GetCompileSettings:
 						{
-							GetCompileSettingsResponse response = Handle_GetCompileSettings(JsonSerializer.Deserialize<GetCompileSettingsCommand>(CommandString, jsonOptions), Logger);
+							GetCompileSettingsResponse response = Handle_GetCompileSettings(JsonSerializer.Deserialize<GetCompileSettingsCommand>(CommandString, jsonOptions));
 							response.Sequence = command.Sequence;
 							Console.WriteLine(JsonSerializer.Serialize(response, responseOptions));
 						}
@@ -402,7 +402,7 @@ namespace UnrealBuildTool
 			return Reply;
 		}
 
-		private QueryFileResponse Handle_QueryFile(QueryFileCommand Command, ILogger Logger)
+		private QueryFileResponse Handle_QueryFile(QueryFileCommand Command)
 		{
 			QueryFileResponse FailResponse = new QueryFileResponse { Found = false };
 			if (CurrentTargetIntellisenseInfo == null)
@@ -423,7 +423,8 @@ namespace UnrealBuildTool
 				return FailResponse;
 			}
 		}
-		private GetBrowseConfigurationResponse Handle_GetBrowseConfiguration(GetBrowseConfigurationCommand Command, ILogger Logger)
+
+		private GetBrowseConfigurationResponse Handle_GetBrowseConfiguration(GetBrowseConfigurationCommand Command)
 		{
 			GetBrowseConfigurationResponse Response = new GetBrowseConfigurationResponse { Success = false, Paths = new List<string>() };
 			if (CurrentTargetIntellisenseInfo == null)
@@ -452,7 +453,7 @@ namespace UnrealBuildTool
 			return Response;
 		}
 
-		private GetCompileSettingsResponse Handle_GetCompileSettings(GetCompileSettingsCommand Command, ILogger Logger)
+		private GetCompileSettingsResponse Handle_GetCompileSettings(GetCompileSettingsCommand Command)
 		{
 			GetCompileSettingsResponse FailResponse = new GetCompileSettingsResponse { };
 			try
diff --git a/Engine/Source/Programs/UnrealBuildTool/Platform/Android/AndroidAARHandler.cs b/Engine/Source/Programs/UnrealBuildTool/Platform/Android/AndroidAARHandler.cs
--- a/Engine/Source/Programs/UnrealBuildTool/Platform/Android/AndroidAARHandler.cs
+++ b/Engine/Source/Programs/UnrealBuildTool/Platform/Android/AndroidAARHandler.cs
@@ -25,7 +25,7 @@ namespace UnrealBuildTool
 				Dependencies = new List<string>();
 			}
 
-			public void AddDependency(string InBaseName, string InVersion)
+			public void AddDependency(string InBaseName)
 			{
 				Dependencies.Add(InBaseName);  // + "-" + InVersion);  will replace version with latest
 			}
@@ -432,7 +432,7 @@ namespace UnrealBuildTool
 				if (DepType == "aar")
 				{
 					// Add dependency
-					AAREntry.AddDependency(DepArtifactId, DepVersion);
+					AAREntry.AddDependency(DepArtifactId);
 
 					AddNewAAR(DepGroupId, DepArtifactId, DepVersion, Logger);
 				}
diff --git a/Engine/Source/Programs/UnrealBuildTool/Platform/Android/AndroidProjectGenerator.cs b/Engine/Source/Programs/UnrealBuildTool/Platform/Android/AndroidProjectGenerator.cs
--- a/Engine/Source/Programs/UnrealBuildTool/Platform/Android/AndroidProjectGenerator.cs
+++ b/Engine/Source/Programs/UnrealBuildTool/Platform/Android/AndroidProjectGenerator.cs
@@ -27,7 +27,7 @@ namespace UnrealBuildTool
 		private readonly Task<bool> AGDEInstalledTask;
 
 		public AndroidProjectGenerator(CommandLineArguments Arguments, ILogger Logger)
-			: base(Arguments, Logger)
+			: base(Logger)
 		{
 			AGDEInstalledTask = Task.Run(() =>
 			{
diff --git a/Engine/Source/Programs/UnrealBuildTool/Platform/Android/AndroidToolChain.cs b/Engine/Source/Programs/UnrealBuildTool/Platform/Android/AndroidToolChain.cs
--- a/Engine/Source/Programs/UnrealBuildTool/Platform/Android/AndroidToolChain.cs
+++ b/Engine/Source/Programs/UnrealBuildTool/Platform/Android/AndroidToolChain.cs
@@ -814,7 +814,7 @@ namespace UnrealBuildTool
 			return Result;
 		}
 
-		static string GetArArguments(LinkEnvironment LinkEnvironment)
+		static string GetArArguments()
 		{
 			string Result = "";
 
@@ -1013,7 +1013,7 @@ namespace UnrealBuildTool
 			return CompileEnvironment;
 		}
 
-		void GenerateEmptyLinkFunctionsForRemovedModules(List<FileItem> SourceFiles, UnrealArch Arch, string ModuleName, DirectoryReference OutputDirectory, IActionGraphBuilder Graph, ILogger Logger)
+		void GenerateEmptyLinkFunctionsForRemovedModules(List<FileItem> SourceFiles, UnrealArch Arch, string ModuleName, DirectoryReference OutputDirectory, IActionGraphBuilder Graph)
 		{
 			// Only add to Launch module
 			if (!ModuleName.Equals("Launch", StringComparison.Ordinal))
@@ -1073,12 +1073,12 @@ namespace UnrealBuildTool
 			return StripPath.Replace("-ar", "-strip");
 		}
 
-		public static string GetReadElfExecutablePath(UnrealArch UnrealArch)
+		public static string GetReadElfExecutablePath()
 		{
 			return ReadElfPath!;
 		}
 
-		public static string GetObjDumpExecutablePath(UnrealArch UnrealArch)
+		public static string GetObjDumpExecutablePath()
 		{
 			return ObjDumpPath!;
 		}
@@ -1101,7 +1101,7 @@ namespace UnrealBuildTool
 				// Directly added NDK files for NDK extensions
 				ModifiedInputFiles.Add(FileItem.GetItemByPath(GetNativeGluePath()));
 				// Deal with dynamic modules removed by architecture
-				GenerateEmptyLinkFunctionsForRemovedModules(ModifiedInputFiles, CompileEnvironment.Architecture, ModuleName, OutputDir, Graph, Logger);
+				GenerateEmptyLinkFunctionsForRemovedModules(ModifiedInputFiles, CompileEnvironment.Architecture, ModuleName, OutputDir, Graph);
 
 				HasHandledLaunchModule.Add(CompileEnvironment.Architecture);
 			}
@@ -1172,7 +1172,7 @@ namespace UnrealBuildTool
 			LinkAction.WorkingDirectory = LinkEnvironment.IntermediateDirectory!;
 
 			// Get link arguments.
-			LinkAction.CommandArguments = LinkEnvironment.bIsBuildingLibrary ? GetArArguments(LinkEnvironment) : GetLinkArguments(LinkEnvironment, Arch, Graph);
+			LinkAction.CommandArguments = LinkEnvironment.bIsBuildingLibrary ? GetArArguments() : GetLinkArguments(LinkEnvironment, Arch, Graph);
 
 			// Add the output file as a production of the link action.
 			FileItem OutputFile;
@@ -1422,7 +1422,7 @@ namespace UnrealBuildTool
 			*/
 		}
 
-		public static void OutputReceivedDataEventHandler(object Sender, DataReceivedEventArgs Line, ILogger Logger)
+		public static void OutputReceivedDataEventHandler(DataReceivedEventArgs Line, ILogger Logger)
 		{
 			if ((Line != null) && (Line.Data != null))
 			{
diff --git a/Engine/Source/Programs/UnrealBuildTool/Platform/Android/UEBuildAndroid.cs b/Engine/Source/Programs/UnrealBuildTool/Platform/Android/UEBuildAndroid.cs
--- a/Engine/Source/Programs/UnrealBuildTool/Platform/Android/UEBuildAndroid.cs
+++ b/Engine/Source/Programs/UnrealBuildTool/Platform/Android/UEBuildAndroid.cs
@@ -119,7 +119,7 @@ namespace UnrealBuildTool
 
 		}
 
-		public override UnrealArchitectures ActiveArchitectures(FileReference? ProjectFile, string? TargetName) => GetProjectArchitectures(ProjectFile, false);
+		public override UnrealArchitectures ActiveArchitectures(FileReference? ProjectFile, string? TargetName) => GetProjectArchitectures(ProjectFile);
 
 		public override string GetFolderNameForArchitecture(UnrealArch Architecture)
 		{
@@ -128,7 +128,7 @@ namespace UnrealBuildTool
 
 		private static UnrealArchitectures? CachedActiveArchitectures = null;
 		private static FileReference? CachedActiveArchesProject = null;
-		private UnrealArchitectures GetProjectArchitectures(FileReference? ProjectFile, bool bGetAllSupported)
+		private UnrealArchitectures GetProjectArchitectures(FileReference? ProjectFile)
 		{
 			if (CachedActiveArchitectures == null || ProjectFile != CachedActiveArchesProject)
 			{
@@ -309,7 +309,7 @@ namespace UnrealBuildTool
 
 			// look up Android specific settings
 			if (!DoProjectSettingsMatchDefault(Platform, ProjectPath, "/Script/AndroidRuntimeSettings.AndroidRuntimeSettings",
-				BoolKeys, null, StringKeys, Logger))
+				BoolKeys, null, StringKeys))
 			{
 				return false;
 			}
@@ -555,7 +555,7 @@ namespace UnrealBuildTool
 			LinkEnvironment.SystemLibraries.Add("atomic");
 		}
 
-		private bool UseTegraGraphicsDebugger(ReadOnlyTargetRules Target)
+		private bool UseTegraGraphicsDebugger(ReadOnlyTargetRules _)
 		{
 			// Disable for now
 			return false;
diff --git a/Engine/Source/Programs/UnrealBuildTool/Platform/Android/UEDeployAndroid.cs b/Engine/Source/Programs/UnrealBuildTool/Platform/Android/UEDeployAndroid.cs
--- a/Engine/Source/Programs/UnrealBuildTool/Platform/Android/UEDeployAndroid.cs
+++ b/Engine/Source/Programs/UnrealBuildTool/Platform/Android/UEDeployAndroid.cs
@@ -1193,7 +1193,7 @@ namespace UnrealBuildTool
 			}
 		}
 
-		public void WriteCrashlyticsResources(string UEBuildPath, string PackageName, string ApplicationDisplayName, bool bIsEmbedded, UnrealArch UnrealArch)
+		public void WriteCrashlyticsResources(string UEBuildPath, string PackageName, bool bIsEmbedded, UnrealArch UnrealArch)
 		{
 			System.DateTime CurrentDateTime = System.DateTime.Now;
 			string BuildID = Guid.NewGuid().ToString();
@@ -1275,10 +1275,10 @@ namespace UnrealBuildTool
 			Utils.RunLocalProcessAndLogOutput(StartInfo, Logger);
 		}
 
-		private List<string> GetSODependencies(string SOFilename, UnrealArch UnrealArch)
+		private List<string> GetSODependencies(string SOFilename)
 		{
 			List<string> DependencyList = new List<string>();
-			string ReadElfExecutable = AndroidToolChain.GetReadElfExecutablePath(UnrealArch);
+			string ReadElfExecutable = AndroidToolChain.GetReadElfExecutablePath();
 			string? Output = Utils.RunLocalProcessAndReturnStdOut(ReadElfExecutable, "--dynamic \"" + SOFilename + "\"");
 			if (Output != null)
 			{
@@ -1318,9 +1318,9 @@ namespace UnrealBuildTool
 			return DependencyList;
 		}
 
-		private static bool IsSO16KbCompatible(string FileName, UnrealArch UnrealArch)
+		private static bool IsSO16KbCompatible(string FileName)
 		{
-			string ReadElfExecutable = AndroidToolChain.GetReadElfExecutablePath(UnrealArch);
+			string ReadElfExecutable = AndroidToolChain.GetReadElfExecutablePath();
 			string StdOutput = Utils.RunLocalProcessAndReturnStdOut(ReadElfExecutable, $"--elf-output-style=JSON --program-headers \"{FileName}\"");
 
 			ReadElfOutput Output = default;
@@ -1411,7 +1411,7 @@ namespace UnrealBuildTool
 			public string Name => !String.IsNullOrEmpty(MaybeName) ? MaybeName : MaybeValueAsString;
 		}
 
-		private void CheckSOInFolderAre16KbCompatible(string FolderPath, UnrealArch UnrealArch)
+		private void CheckSOInFolderAre16KbCompatible(string FolderPath)
 		{
 			try
 			{
@@ -1419,7 +1419,7 @@ namespace UnrealBuildTool
 
 				foreach (string SOFile in SOFiles)
 				{
-					if (!IsSO16KbCompatible(SOFile, UnrealArch))
+					if (!IsSO16KbCompatible(SOFile))
 					{
 						Logger.LogWarning("Library '{SOFile}' is NOT compatible with 16kb page size alignment.", SOFile);
 					}
@@ -1431,10 +1431,10 @@ namespace UnrealBuildTool
 			}
 		}
 
-		private void CheckSOInGradleFolderAre16KbCompatible(string GradlePath, UnrealArch UnrealArch)
+		private void CheckSOInGradleFolderAre16KbCompatible(string GradlePath)
 		{
-			CheckSOInFolderAre16KbCompatible($"{GradlePath}/jniLibs/arm64-v8a", UnrealArch);
-			CheckSOInFolderAre16KbCompatible($"{GradlePath}/libs/arm64-v8a", UnrealArch);
+			CheckSOInFolderAre16KbCompatible($"{GradlePath}/jniLibs/arm64-v8a");
+			CheckSOInFolderAre16KbCompatible($"{GradlePath}/libs/arm64-v8a");
 		}
 
 		private static void CopySO(string FinalName, string SourceName)
@@ -1473,7 +1473,7 @@ namespace UnrealBuildTool
 			return "windows-x86_64";
 		}
 
-		private static void CopySTL(AndroidToolChain ToolChain, string UnrealBuildPath, UnrealArch UnrealArch, string NDKArch, bool bForDistribution)
+		private static void CopySTL(string UnrealBuildPath, string NDKArch)
 		{
 			// copy it in!
 			string SourceSTLSOName = Environment.ExpandEnvironmentVariables("%NDKROOT%/sources/cxx-stl/llvm-libc++/libs/") + NDKArch + "/libc++_shared.so";
@@ -1488,7 +1488,7 @@ namespace UnrealBuildTool
 			CopySO(FinalSTLSOName, SourceSTLSOName);
 		}
 
-		private void CopyPSOService(string UnrealBuildPath, string PreBuiltPath, UnrealArch UnrealArch, string NDKArch)
+		private void CopyPSOService(string UnrealBuildPath, string PreBuiltPath, string NDKArch)
 		{
 			// copy it in!
 			string SourceSOName = Path.Combine(PreBuiltPath, "PSOService/Android/Release/" + NDKArch) + "/libpsoservice.so";
@@ -1497,7 +1497,7 @@ namespace UnrealBuildTool
 			CopySO(FinalSOName, SourceSOName);
 		}
 
-		private void CopySwappy(string UnrealBuildPath, UnrealArch UnrealArch, string NDKArch, UnrealTargetConfiguration Configuration)
+		private void CopySwappy(string UnrealBuildPath, string NDKArch, UnrealTargetConfiguration Configuration)
 		{
 			string BuildConfig = Configuration == UnrealTargetConfiguration.Debug ? "debug" : "release";
 			string SourceSOName = Path.Combine(Unreal.EngineDirectory.ToString(), "Source", "ThirdParty", "GoogleGameSDK", "gamesdk", "libs", BuildConfig, NDKArch) + "/libswappy.so";
@@ -1506,7 +1506,7 @@ namespace UnrealBuildTool
 			CopySO(FinalSOName, SourceSOName);
 		}
 
-		private void CopyGfxDebugger(string UnrealBuildPath, UnrealArch UnrealArch, string NDKArch)
+		private void CopyGfxDebugger(string UnrealBuildPath, string NDKArch)
 		{
 			string AndroidGraphicsDebugger;
 			ConfigHierarchy Ini = GetConfigCacheIni(ConfigHierarchyType.Engine);
@@ -1570,7 +1570,7 @@ namespace UnrealBuildTool
 			Logger.LogInformation("bSupportsVulkan: {SupportsVulkan}", (bSupportsVulkan ? "true" : "false"));
 		}
 
-		void CopyVulkanValidationLayers(string UnrealBuildPath, UnrealArch UnrealArch, string NDKArch, string Configuration)
+		void CopyVulkanValidationLayers(string UnrealBuildPath, string NDKArch, string Configuration)
 		{
 			ConfigHierarchy Ini = GetConfigCacheIni(ConfigHierarchyType.Engine);
 			Ini.GetBool("/Script/AndroidRuntimeSettings.AndroidRuntimeSettings", "bSupportsVulkan", out bool bSupportsVulkan);
@@ -1630,7 +1630,7 @@ namespace UnrealBuildTool
 			}
 		}
 
-		void CopyClangSanitizerLib(AndroidToolChain ToolChain, string UnrealBuildPath, UnrealArch UnrealArch, string NDKArch, AndroidToolChain.ClangSanitizer Sanitizer)
+		void CopyClangSanitizerLib(AndroidToolChain ToolChain, string UnrealBuildPath, string NDKArch, AndroidToolChain.ClangSanitizer Sanitizer)
 		{
 			string Architecture = "-aarch64";
 			switch (NDKArch)
@@ -1986,7 +1986,7 @@ namespace UnrealBuildTool
 			}
 		}
 
-		static void FilterStdOutErr(object sender, DataReceivedEventArgs e, ILogger Logger)
+		static void FilterStdOutErr(DataReceivedEventArgs e, ILogger Logger)
 		{
 			if (e.Data != null)
 			{
@@ -2055,8 +2055,8 @@ namespace UnrealBuildTool
 
 			Process Proc = new Process();
 			Proc.StartInfo = StartInfo;
-			Proc.OutputDataReceived += (s, e) => FilterStdOutErr(s, e, Logger);
-			Proc.ErrorDataReceived += (s, e) => FilterStdOutErr(s, e, Logger);
+			Proc.OutputDataReceived += (_, e) => FilterStdOutErr(e, Logger);
+			Proc.ErrorDataReceived += (_, e) => FilterStdOutErr(e, Logger);
 
 			Proc.Start();
 			Proc.BeginOutputReadLine();
@@ -2275,7 +2275,7 @@ namespace UnrealBuildTool
 		}
 
 		private bool CheckDependencies(UnrealArchitectures Architectures, string ProjectName, string ProjectDirectory, string IntermediateAndroidPath,
-			string UnrealBuildFilesPath, string GameBuildFilesPath, string EngineDirectory, List<string> SettingsFiles, string CookFlavor,
+			string UnrealBuildFilesPath, string GameBuildFilesPath, List<string> SettingsFiles, string CookFlavor,
 			string OutputPath, bool bMakeSeparateApks, bool bPackageDataInsideApk, bool bDontBundleLibrariesInAPK)
 		{
 			// check all input files (.so, java files, .ini files, etc)
@@ -2733,7 +2733,7 @@ namespace UnrealBuildTool
 			return EngineMajorVersion + "." + EngineMinorVersion + "." + EnginePatchVersion;
 		}
 
-		private string GenerateManifest(AndroidToolChain ToolChain, string ProjectName, TargetType InTargetType, string EngineDirectory, bool bIsForDistribution, bool bPackageDataInsideApk, string GameBuildFilesPath, bool bHasOBBFiles, bool bDisableVerifyOBBOnStartUp, UnrealArch UnrealArch, string CookFlavor, bool bUseExternalFilesDir, string Configuration, int SDKLevelInt, bool bIsEmbedded, bool bEnableBundle, bool bUsesZenStreaming)
+		private string GenerateManifest(AndroidToolChain ToolChain, string ProjectName, TargetType InTargetType, bool bIsForDistribution, bool bPackageDataInsideApk, string GameBuildFilesPath, bool bHasOBBFiles, bool bDisableVerifyOBBOnStartUp, UnrealArch UnrealArch, string CookFlavor, bool bUseExternalFilesDir, string Configuration, int SDKLevelInt, bool bIsEmbedded, bool bEnableBundle, bool bUsesZenStreaming)
 		{
 			// Read the engine version
 			string EngineVersion = ReadEngineVersion();
@@ -2741,7 +2741,7 @@ namespace UnrealBuildTool
 			int StoreVersion = GetStoreVersion(UnrealArch);
 
 			string Arch = GetNDKArch(UnrealArch);
-			GetMinTargetSDKVersions(ToolChain, UnrealArch, UPL!, Arch, bEnableBundle, out int MinSDKVersion, out int TargetSDKVersion, out _);
+			GetMinTargetSDKVersions(ToolChain, UPL!, Arch, bEnableBundle, out int MinSDKVersion, out int TargetSDKVersion, out _);
 
 			// get project version from ini
 			ConfigHierarchy GameIni = GetConfigCacheIni(ConfigHierarchyType.Game);
@@ -3640,7 +3640,7 @@ namespace UnrealBuildTool
 			return true;
 		}
 
-		private void GetMinTargetSDKVersions(AndroidToolChain ToolChain, UnrealArch Arch, UnrealPluginLanguage UPL, string NDKArch, bool bEnableBundle, out int MinSDKVersion, out int TargetSDKVersion, out int NDKLevelInt)
+		private void GetMinTargetSDKVersions(AndroidToolChain ToolChain, UnrealPluginLanguage UPL, string NDKArch, bool bEnableBundle, out int MinSDKVersion, out int TargetSDKVersion, out int NDKLevelInt)
 		{
 			ConfigHierarchy Ini = GetConfigCacheIni(ConfigHierarchyType.Engine);
 			Ini.GetInt32("/Script/AndroidRuntimeSettings.AndroidRuntimeSettings", "MinSDKVersion", out MinSDKVersion);
@@ -3988,7 +3988,7 @@ namespace UnrealBuildTool
 			return true;
 		}
 
-		private void CreateGradlePropertiesFiles(AndroidToolChain ToolChain, UnrealArch Arch, int MinSDKVersion, int TargetSDKVersion, string CompileSDKVersion, string BuildToolsVersion, string PackageName,
+		private void CreateGradlePropertiesFiles(AndroidToolChain ToolChain, int MinSDKVersion, int TargetSDKVersion, string CompileSDKVersion, string BuildToolsVersion, string PackageName,
 			string? DestApkName, string NDKArch, string UnrealBuildFilesPath, string GameBuildFilesPath, string UnrealBuildGradleAppPath, string UnrealBuildPath, string UnrealBuildGradlePath,
 			bool bForDistribution, bool bIsEmbedded, List<string> OBBFiles, bool bJavaPerfettoMarkers)
 		{
@@ -4329,7 +4329,7 @@ dependencies.implementation('androidx.tracing:tracing:1.2.0')
 		};
 
 		private void MakeApk(AndroidToolChain ToolChain, string ProjectName, TargetType InTargetType, string ProjectDirectory, string OutputPath, string EngineDirectory, bool bForDistribution, string CookFlavor,
-			UnrealTargetConfiguration Configuration, bool bMakeSeparateApks, bool bIncrementalPackage, bool bDisallowPackagingDataInApk, bool bDisallowExternalFilesDir, bool bSkipGradleBuild, bool bIsArchive, bool bIsFromUAT)
+			UnrealTargetConfiguration Configuration, bool bMakeSeparateApks, bool bDisallowPackagingDataInApk, bool bDisallowExternalFilesDir, bool bSkipGradleBuild, bool bIsArchive, bool bIsFromUAT)
 		{
 			Logger.LogInformation("");
 			Logger.LogInformation("===={Time}====PREPARING TO MAKE APK=================================================================", DateTime.Now.ToString());
@@ -4694,7 +4694,7 @@ dependencies.implementation('androidx.tracing:tracing:1.2.0')
 				JavaFiles.AddRange(from t in templates select t.DestinationFile);
 
 				bBuildSettingsMatch = CheckDependencies(Architectures, ProjectName, ProjectDirectory, IntermediateAndroidPath,
-					UnrealBuildFilesPath, GameBuildFilesPath, EngineDirectory, JavaFiles, CookFlavor, OutputPath,
+					UnrealBuildFilesPath, GameBuildFilesPath, JavaFiles, CookFlavor, OutputPath,
 					bMakeSeparateApks, bPackageDataInsideApk, bDontBundleLibrariesInAPK);
 
 			}
@@ -4769,7 +4769,7 @@ dependencies.implementation('androidx.tracing:tracing:1.2.0')
 				// LINQ query is subsequently used multiple times, so must copy results
 				BuildList = (from Arch in Architectures.Architectures
 							let CompileSDKLevelForArch = NDKArchToCompileSDK[GetNDKArch(Arch)]
-							let manifest = GenerateManifest(ToolChain, ProjectName, InTargetType, EngineDirectory, bForDistribution, bPackageDataInsideApk, GameBuildFilesPath, bRequiresOBB, bDisableVerifyOBBOnStartUp, Arch, CookFlavor, bUseExternalFilesDir, Configuration.ToString(), CompileSDKLevelForArch, bSkipGradleBuild, bEnableBundle, bUsesZenStreaming)
+							let manifest = GenerateManifest(ToolChain, ProjectName, InTargetType, bForDistribution, bPackageDataInsideApk, GameBuildFilesPath, bRequiresOBB, bDisableVerifyOBBOnStartUp, Arch, CookFlavor, bUseExternalFilesDir, Configuration.ToString(), CompileSDKLevelForArch, bSkipGradleBuild, bEnableBundle, bUsesZenStreaming)
 							select Tuple.Create(Arch, manifest)).ToArray();
 			}
 			else
@@ -4778,7 +4778,7 @@ dependencies.implementation('androidx.tracing:tracing:1.2.0')
 				BuildList = (from Arch in Architectures.Architectures
 							let CompileSDKLevelForArch = NDKArchToCompileSDK[GetNDKArch(Arch)]
 							let manifestFile = Path.Combine(IntermediateAndroidPath, Arch + "_AndroidManifest.xml")
-							let manifest = GenerateManifest(ToolChain, ProjectName, InTargetType, EngineDirectory, bForDistribution, bPackageDataInsideApk, GameBuildFilesPath, bRequiresOBB, bDisableVerifyOBBOnStartUp, Arch, CookFlavor, bUseExternalFilesDir, Configuration.ToString(), CompileSDKLevelForArch, bSkipGradleBuild, bEnableBundle, bUsesZenStreaming)
+							let manifest = GenerateManifest(ToolChain, ProjectName, InTargetType, bForDistribution, bPackageDataInsideApk, GameBuildFilesPath, bRequiresOBB, bDisableVerifyOBBOnStartUp, Arch, CookFlavor, bUseExternalFilesDir, Configuration.ToString(), CompileSDKLevelForArch, bSkipGradleBuild, bEnableBundle, bUsesZenStreaming)
 							let OldManifest = File.Exists(manifestFile) ? File.ReadAllText(manifestFile) : ""
 							where manifest != OldManifest
 							select Tuple.Create(Arch, manifest)).ToArray();
@@ -5342,8 +5342,8 @@ popd
 				}
 
 				// update GameActivity.java and GameApplication.java if out of date
-				UpdateGameActivity(Arch, NDKArch, EngineDirectory, UnrealBuildPath, bDontBundleLibrariesInAPK, bJavaPerfettoMarkers, Configuration, bHasSplashActivity);
-				UpdateGameApplication(Arch, NDKArch, EngineDirectory, UnrealBuildPath, bJavaPerfettoMarkers, Configuration, SourceSOName);
+				UpdateGameActivity(NDKArch, EngineDirectory, UnrealBuildPath, bDontBundleLibrariesInAPK, bJavaPerfettoMarkers, Configuration, bHasSplashActivity);
+				UpdateGameApplication(NDKArch, EngineDirectory, UnrealBuildPath, bJavaPerfettoMarkers, Configuration, SourceSOName);
 
 				bool bSkipLibCpp = false;
 				Ini.GetBool("/Script/AndroidRuntimeSettings.AndroidRuntimeSettings", "bSkipLibCpp", out bSkipLibCpp);
@@ -5351,16 +5351,16 @@ popd
 				{
 					// after ndk-build is called, we can now copy in the stl .so (ndk-build deletes old files)
 					// copy libc++_shared.so to library
-					CopySTL(ToolChain, UnrealBuildPath, Arch, NDKArch, bForDistribution);
+					CopySTL(UnrealBuildPath, NDKArch);
 				}
-				CopySwappy(UnrealBuildPath, Arch, NDKArch, Configuration);
-				CopyPSOService(UnrealBuildPath, UnrealPreBuiltFilesPath, Arch, NDKArch);
-				CopyGfxDebugger(UnrealBuildPath, Arch, NDKArch);
-				CopyVulkanValidationLayers(UnrealBuildPath, Arch, NDKArch, Configuration.ToString());
+				CopySwappy(UnrealBuildPath, NDKArch, Configuration);
+				CopyPSOService(UnrealBuildPath, UnrealPreBuiltFilesPath, NDKArch);
+				CopyGfxDebugger(UnrealBuildPath, NDKArch);
+				CopyVulkanValidationLayers(UnrealBuildPath, NDKArch, Configuration.ToString());
 
 				if (Sanitizer != AndroidToolChain.ClangSanitizer.None)
 				{
-					CopyClangSanitizerLib(ToolChain, UnrealBuildPath, Arch, NDKArch, Sanitizer);
+					CopyClangSanitizerLib(ToolChain, UnrealBuildPath, NDKArch, Sanitizer);
 
 					if (bEnableScudoMemoryTracing)
 					{
@@ -5415,11 +5415,11 @@ popd
 				int MinSDKVersion = 0;
 				int TargetSDKVersion = 0;
 				int NDKLevelInt = 0;
-				GetMinTargetSDKVersions(ToolChain, Arch, UPL, NDKArch, bEnableBundle, out MinSDKVersion, out TargetSDKVersion, out NDKLevelInt);
+				GetMinTargetSDKVersions(ToolChain, UPL, NDKArch, bEnableBundle, out MinSDKVersion, out TargetSDKVersion, out NDKLevelInt);
 
 				// move JavaLibs into subprojects
 				string JavaLibsDir = Path.Combine(UnrealBuildPath, "JavaLibs");
-				PrepareJavaLibsForGradle(JavaLibsDir, UnrealBuildGradlePath, MinSDKVersion.ToString(), TargetSDKVersion.ToString(), CompileSDKVersion, BuildToolsVersion, NDKArch);
+				PrepareJavaLibsForGradle(JavaLibsDir, UnrealBuildGradlePath, MinSDKVersion.ToString(), TargetSDKVersion.ToString(), NDKArch);
 
 				// Create local.properties
 				string LocalPropertiesFilename = Path.Combine(UnrealBuildGradlePath, "local.properties");
@@ -5428,10 +5428,10 @@ popd
 				LocalProperties.AppendLine(String.Format("sdk.dir={0}", Environment.GetEnvironmentVariable("ANDROID_HOME")!.Replace("\\", "/")));
 				File.WriteAllText(LocalPropertiesFilename, LocalProperties.ToString());
 
-				CreateGradlePropertiesFiles(ToolChain, Arch, MinSDKVersion, TargetSDKVersion, CompileSDKVersion, BuildToolsVersion, PackageName, DestApkName, NDKArch,
+				CreateGradlePropertiesFiles(ToolChain, MinSDKVersion, TargetSDKVersion, CompileSDKVersion, BuildToolsVersion, PackageName, DestApkName, NDKArch,
 					UnrealBuildFilesPath, GameBuildFilesPath, UnrealBuildGradleAppPath, UnrealBuildPath, UnrealBuildGradlePath, bForDistribution, bSkipGradleBuild, RequiredOBBFiles, bJavaPerfettoMarkers);
 
-				CheckSOInGradleFolderAre16KbCompatible(UnrealBuildGradleMainPath, Arch);
+				CheckSOInGradleFolderAre16KbCompatible(UnrealBuildGradleMainPath);
 
 				if (!bSkipGradleBuild && !bOnlyGradleGeneration)
 				{
@@ -6079,7 +6079,7 @@ popd
 			string RelativeEnginePath = Unreal.EngineDirectory.MakeRelativeTo(DirectoryReference.GetCurrentDirectory());
 
 			MakeApk(ToolChain, TargetName, Receipt.TargetType, ProjectDirectory.FullName, BaseSoName, RelativeEnginePath, bForDistribution: false, CookFlavor: "", Configuration: Receipt.Configuration,
-				bMakeSeparateApks: ShouldMakeSeparateApks(), bIncrementalPackage: true, bDisallowPackagingDataInApk: false, bDisallowExternalFilesDir: true, bSkipGradleBuild: bShouldCompileAsDll, bIsArchive: false, bIsFromUAT: false);
+				bMakeSeparateApks: ShouldMakeSeparateApks(), bDisallowPackagingDataInApk: false, bDisallowExternalFilesDir: true, bSkipGradleBuild: bShouldCompileAsDll, bIsArchive: false, bIsFromUAT: false);
 
 			// if we made any non-standard .apk files, the generated debugger settings may be wrong
 			if (ShouldMakeSeparateApks() && (OutputPaths.Count > 1 || !OutputPaths[0].FullName.Contains("-armv7")))
@@ -6156,12 +6156,12 @@ popd
 			SavePackageInfo(ProjectName, ProjectDirectory.FullName, Type, bSkipGradleBuild);
 
 			MakeApk(ToolChain, ProjectName, Type, ProjectDirectory.FullName, ExecutablePath, EngineDirectory, bForDistribution: bForDistribution, CookFlavor: CookFlavor, Configuration: Configuration,
-				bMakeSeparateApks: ShouldMakeSeparateApks(), bIncrementalPackage: false, bDisallowPackagingDataInApk: bIsDataDeploy, bDisallowExternalFilesDir: bIsDataDeploy,
+				bMakeSeparateApks: ShouldMakeSeparateApks(), bDisallowPackagingDataInApk: bIsDataDeploy, bDisallowExternalFilesDir: bIsDataDeploy,
 				bSkipGradleBuild: bSkipGradleBuild, bIsArchive: bIsArchive, bIsFromUAT: true);
 			return true;
 		}
 
-		public static void OutputReceivedDataEventHandler(object Sender, DataReceivedEventArgs Line, ILogger Logger)
+		public static void OutputReceivedDataEventHandler(DataReceivedEventArgs Line, ILogger Logger)
 		{
 			if ((Line != null) && (Line.Data != null))
 			{
@@ -6218,7 +6218,7 @@ popd
 			return "empty";
 		}
 
-		private void UpdateGameActivity(UnrealArch UnrealArch, string NDKArch, string EngineDir, string UnrealBuildPath, bool bDontBundleLibrariesInAPK,
+		private void UpdateGameActivity(string NDKArch, string EngineDir, string UnrealBuildPath, bool bDontBundleLibrariesInAPK,
 				bool bJavaPerfettoMarkers, UnrealTargetConfiguration Configuration, bool bHasSplashActivity)
 		{
 			string SourceFilename = Path.Combine(EngineDir, "Build", "Android", "Java", "src", "com", "epicgames", "unreal", "GameActivity.java.template");
@@ -6437,7 +6437,7 @@ import java.util.Collection;
 			}
 		}
 
-		private void UpdateGameApplication(UnrealArch UnrealArch, string NDKArch, string EngineDir, string UnrealBuildPath,
+		private void UpdateGameApplication(string NDKArch, string EngineDir, string UnrealBuildPath,
 				bool bJavaPerfettoMarkers, UnrealTargetConfiguration Configuration, string SOName)
 		{
 			string SourceFilename = Path.Combine(EngineDir, "Build", "Android", "Java", "src", "com", "epicgames", "unreal", "GameApplication.java.template");
@@ -6450,7 +6450,7 @@ import java.util.Collection;
 			}
 
 			// add dependencies to early load during SplashActivity (optional)
-			List<string> UnrealDependencies = GetSODependencies(SOName, UnrealArch);
+			List<string> UnrealDependencies = GetSODependencies(SOName);
 			string GameApplicationEarlyLoad = "";
 			if (UnrealDependencies.Count > 0)
 			{
@@ -6659,7 +6659,7 @@ import java.util.Collection;
 			return AARHandler;
 		}
 
-		private void PrepareJavaLibsForGradle(string JavaLibsDir, string UnrealBuildGradlePath, string InMinSdkVersion, string InTargetSdkVersion, string CompileSDKVersion, string BuildToolsVersion, string NDKArch)
+		private void PrepareJavaLibsForGradle(string JavaLibsDir, string UnrealBuildGradlePath, string InMinSdkVersion, string InTargetSdkVersion, string NDKArch)
 		{
 			StringBuilder SettingsGradleContent = new StringBuilder();
 			StringBuilder ProjectDependencyContent = new StringBuilder();
diff --git a/Engine/Source/Programs/UnrealBuildTool/Platform/Apple/AppleProjectGenerator.cs b/Engine/Source/Programs/UnrealBuildTool/Platform/Apple/AppleProjectGenerator.cs
--- a/Engine/Source/Programs/UnrealBuildTool/Platform/Apple/AppleProjectGenerator.cs
+++ b/Engine/Source/Programs/UnrealBuildTool/Platform/Apple/AppleProjectGenerator.cs
@@ -16,8 +16,8 @@ namespace UnrealBuildTool
 		private static readonly XcrunRunner AppleHelper = new XcrunRunner();
 
 		/// <inheritdoc/>
-		protected AppleProjectGenerator(CommandLineArguments Arguments, ILogger Logger)
-			: base(Arguments, Logger)
+		protected AppleProjectGenerator(ILogger Logger)
+			: base(Logger)
 		{
 		}
 
@@ -63,7 +63,7 @@ namespace UnrealBuildTool
 					throw new InvalidOperationException("xcrun can be run only for Apple's platforms");
 				}
 
-				string SDKPath = GetSDKPath(Architecture, Platform, Logger);
+				string SDKPath = GetSDKPath(Architecture, Platform);
 				if (!CachedIncludePaths.ContainsKey(SDKPath))
 				{
 					CalculateSystemIncludePaths(SDKPath, Logger);
@@ -142,7 +142,7 @@ namespace UnrealBuildTool
 				}
 			}
 
-			private string GetSDKPath(UnrealArch Architecture, UnrealTargetPlatform Platform, ILogger Logger)
+			private string GetSDKPath(UnrealArch Architecture, UnrealTargetPlatform Platform)
 			{
 				if (Platform == UnrealTargetPlatform.Mac)
 				{
diff --git a/Engine/Source/Programs/UnrealBuildTool/Platform/Apple/AppleToolChain.cs b/Engine/Source/Programs/UnrealBuildTool/Platform/Apple/AppleToolChain.cs
--- a/Engine/Source/Programs/UnrealBuildTool/Platform/Apple/AppleToolChain.cs
+++ b/Engine/Source/Programs/UnrealBuildTool/Platform/Apple/AppleToolChain.cs
@@ -160,7 +160,7 @@ namespace UnrealBuildTool
 				{
 					// even in modern, we have to try to unzip in UBT, so that commandline builds extract the frameworks _before_ building (the Xcode finalize
 					// happens after building). We still need the unzip in Xcode pre-build phase, because Xcode wants the frameworks around for its build graph
-					ExtractFramework(Framework, Graph, OutputDir, Logger);
+					ExtractFramework(Framework, Graph);
 					FrameworkTokenFiles.Add(Framework.ExtractedTokenFile!);
 				}
 			}
@@ -258,7 +258,7 @@ namespace UnrealBuildTool
 			return DestFile;
 		}
 
-		protected FileItem ExtractFramework(UEBuildFramework Framework, IActionGraphBuilder Graph, DirectoryReference IntermediateDirectory, ILogger Logger)
+		protected FileItem ExtractFramework(UEBuildFramework Framework, IActionGraphBuilder Graph)
 		{
 			if (Framework.ZipFile == null)
 			{
diff --git a/Engine/Source/Programs/UnrealBuildTool/Platform/Apple/AppleToolChainSettings.cs b/Engine/Source/Programs/UnrealBuildTool/Platform/Apple/AppleToolChainSettings.cs
--- a/Engine/Source/Programs/UnrealBuildTool/Platform/Apple/AppleToolChainSettings.cs
+++ b/Engine/Source/Programs/UnrealBuildTool/Platform/Apple/AppleToolChainSettings.cs
@@ -99,7 +99,7 @@ namespace UnrealBuildTool
 			SDKVersion = ApplePlatformSDK.GetPlatformSDKVersion(OSPrefix) ?? "";
 			SDKVersionFloat = ApplePlatformSDK.GetPlatformSDKVersionFloat(OSPrefix);
 
-			TestXcode(bVerbose, Logger);
+			TestXcode(bVerbose);
 		}
 
 		/// <summary>
@@ -138,9 +138,8 @@ namespace UnrealBuildTool
 		/// Find the Xcode developer directory
 		/// </summary>
 		/// <param name="bVerbose"></param>
-		/// <param name="Logger"></param>
 		/// <exception cref="BuildException"></exception>
-		private static void TestXcode(bool bVerbose, ILogger Logger)
+		private static void TestXcode(bool bVerbose)
 		{
 			DirectoryReference XcodeDeveloperDir = ApplePlatformSDK.DeveloperDir;
 			// make sure we get a full path
diff --git a/Engine/Source/Programs/UnrealBuildTool/Platform/Apple/RemoteMac.cs b/Engine/Source/Programs/UnrealBuildTool/Platform/Apple/RemoteMac.cs
--- a/Engine/Source/Programs/UnrealBuildTool/Platform/Apple/RemoteMac.cs
+++ b/Engine/Source/Programs/UnrealBuildTool/Platform/Apple/RemoteMac.cs
@@ -756,9 +756,8 @@ namespace UnrealBuildTool
 		/// Prepare the remotely built iOS or tvOS client to de debugged on Mac
 		/// </summary>
 		/// <param name="CookedDataDirectory">Where the content has been previously cooked</param>
-		/// <param name="ProjectFile">The project file</param>
 		/// <param name="Logger">The logger</param>
-		public void PrepareToDebug(string CookedDataDirectory, FileReference ProjectFile, ILogger Logger)
+		public void PrepareToDebug(string CookedDataDirectory, ILogger Logger)
 		{
 			Logger.LogInformation("[Remote] Uploading the default uprojectdirs ...");
 			UploadFile(new FileReference("Default.uprojectdirs"), Logger);
diff --git a/Engine/Source/Programs/UnrealBuildTool/Platform/Clang/ClangToolChain.cs b/Engine/Source/Programs/UnrealBuildTool/Platform/Clang/ClangToolChain.cs
--- a/Engine/Source/Programs/UnrealBuildTool/Platform/Clang/ClangToolChain.cs
+++ b/Engine/Source/Programs/UnrealBuildTool/Platform/Clang/ClangToolChain.cs
@@ -1396,7 +1396,7 @@ namespace UnrealBuildTool
 			FileItem TargetFile = GetCompileArguments_FileType(CompileEnvironment, SourceFile, OutputDir, FileArguments, CompileAction, Result);
 
 			// Creates the path to the response file using the name of the output file and creates its contents.
-			FileReference ResponseFileName = GetResponseFileName(CompileEnvironment, TargetFile);
+			FileReference ResponseFileName = GetResponseFileName(TargetFile);
 			List<string> ResponseFileContents = ExpandResponseFileContents(FileArguments);
 
 			if (OperatingSystem.IsWindows())
diff --git a/Engine/Source/Programs/UnrealBuildTool/Platform/IOS/IOSExports.cs b/Engine/Source/Programs/UnrealBuildTool/Platform/IOS/IOSExports.cs
--- a/Engine/Source/Programs/UnrealBuildTool/Platform/IOS/IOSExports.cs
+++ b/Engine/Source/Programs/UnrealBuildTool/Platform/IOS/IOSExports.cs
@@ -82,7 +82,7 @@ namespace UnrealBuildTool
 		public static bool PrepForUATPackageOrDeploy(UnrealTargetConfiguration Config, FileReference ProjectFile, string InProjectName, DirectoryReference InProjectDirectory, FileReference Executable, DirectoryReference InEngineDir, bool bForDistribution, string CookFlavor, bool bIsDataDeploy, bool bCreateStubIPA, FileReference BuildReceiptFileName, ILogger Logger)
 		{
 			TargetReceipt Receipt = TargetReceipt.Read(BuildReceiptFileName);
-			return new UEDeployIOS(Logger).PrepForUATPackageOrDeploy(Config, ProjectFile, InProjectName, InProjectDirectory.FullName, Executable, InEngineDir.FullName, bForDistribution, CookFlavor, bIsDataDeploy, bCreateStubIPA, Receipt);
+			return new UEDeployIOS(Logger).PrepForUATPackageOrDeploy(Config, ProjectFile, InProjectName, InProjectDirectory.FullName, Executable, InEngineDir.FullName, bCreateStubIPA, Receipt);
 		}
 
 		/// <summary>
@@ -202,7 +202,7 @@ namespace UnrealBuildTool
 		public static void PrepareRemoteMacForDebugging(string ClientPlatform, FileReference ProjectFile, ILogger Logger)
 		{
 			RemoteMac Remote = new RemoteMac(ProjectFile, Logger);
-			Remote.PrepareToDebug(ClientPlatform, ProjectFile, Logger);
+			Remote.PrepareToDebug(ClientPlatform, Logger);
 		}
 
 		/// <summary>
diff --git a/Engine/Source/Programs/UnrealBuildTool/Platform/IOS/IOSProjectGenerator.cs b/Engine/Source/Programs/UnrealBuildTool/Platform/IOS/IOSProjectGenerator.cs
--- a/Engine/Source/Programs/UnrealBuildTool/Platform/IOS/IOSProjectGenerator.cs
+++ b/Engine/Source/Programs/UnrealBuildTool/Platform/IOS/IOSProjectGenerator.cs
@@ -16,8 +16,10 @@ namespace UnrealBuildTool
 		/// </summary>
 		/// <param name="Arguments">Command line arguments passed to the project generator</param>
 		/// <param name="Logger">Logger for output</param>
+#pragma warning disable IDE0060 // Remove unused parameter - constructor is found by reflection in GenerateProjectFilesMode.ExecuteAsync
 		public IOSProjectGenerator(CommandLineArguments Arguments, ILogger Logger)
-			: base(Arguments, Logger)
+			: base(Logger)
+#pragma warning restore IDE0060
 		{
 		}
 
diff --git a/Engine/Source/Programs/UnrealBuildTool/Platform/IOS/IOSToolChain.cs b/Engine/Source/Programs/UnrealBuildTool/Platform/IOS/IOSToolChain.cs
--- a/Engine/Source/Programs/UnrealBuildTool/Platform/IOS/IOSToolChain.cs
+++ b/Engine/Source/Programs/UnrealBuildTool/Platform/IOS/IOSToolChain.cs
@@ -443,7 +443,7 @@ namespace UnrealBuildTool
 			{
 				if (Framework.ZipFile != null)
 				{
-					FileItem ExtractedTokenFile = ExtractFramework(Framework, Graph, LinkEnvironment.IntermediateDirectory!, Logger);
+					FileItem ExtractedTokenFile = ExtractFramework(Framework, Graph);
 					LinkAction.PrerequisiteItems.Add(ExtractedTokenFile);
 				}
 			}
@@ -606,10 +606,9 @@ namespace UnrealBuildTool
 		/// Generates debug info for a given executable
 		/// </summary>
 		/// <param name="Executable">FileItem describing the executable to generate debug info for</param>
-		/// <param name="bIsForLTOBuild">Was this build made with LTO enabled?</param>
 		/// <param name="Graph">List of actions to be executed. Additional actions will be added to this list.</param>
 		/// <param name="Logger">Logger for output</param>
-		private List<FileItem> GenerateDebugInfo(FileItem Executable, bool bIsForLTOBuild, IActionGraphBuilder Graph, ILogger Logger)
+		private List<FileItem> GenerateDebugInfo(FileItem Executable, IActionGraphBuilder Graph, ILogger Logger)
 		{
 			// Make a file item for the source and destination files
 			string FullDestPathRoot = GetdSYMPath(Executable);
@@ -1029,7 +1028,7 @@ namespace UnrealBuildTool
 			// For IOS/tvOS, generate the dSYM file if needed or requested
 			if (Target!.IOSPlatform.bGeneratedSYM)
 			{
-				List<FileItem> Files = GenerateDebugInfo(Executable, BinaryLinkEnvironment.bAllowLTCG, Graph, Logger);
+				List<FileItem> Files = GenerateDebugInfo(Executable, Graph, Logger);
 				foreach (FileItem item in Files)
 				{
 					OutputFiles.Add(item);
@@ -1178,7 +1177,7 @@ namespace UnrealBuildTool
 			//This chunk looks to be required to pipe output to VS giving information on the status of a remote build.
 			public bool OutputReceivedDataEventHandlerEncounteredError = false;
 			public string OutputReceivedDataEventHandlerEncounteredErrorMessage = "";
-			public void OutputReceivedDataEventHandler(object Sender, DataReceivedEventArgs Line, ILogger Logger)
+			public void OutputReceivedDataEventHandler(DataReceivedEventArgs Line, ILogger Logger)
 			{
 				if ((Line != null) && (Line.Data != null))
 				{
@@ -1196,7 +1195,7 @@ namespace UnrealBuildTool
 				}
 			}
 
-			public void OutputReceivedDataEventLogger(object Sender, DataReceivedEventArgs Line, ILogger Logger)
+			public void OutputReceivedDataEventLogger(DataReceivedEventArgs Line, ILogger Logger)
 			{
 				if ((Line != null) && (Line.Data != null))
 				{
@@ -1205,7 +1204,7 @@ namespace UnrealBuildTool
 			}
 		}
 
-		private static void GenerateCrashlyticsData(string DsymZip, string ProjectDir, string ProjectName, ILogger Logger)
+		private static void GenerateCrashlyticsData(string DsymZip, ILogger Logger)
 		{
 			Logger.LogInformation("Generating and uploading Crashlytics Data");
 
@@ -1231,8 +1230,8 @@ namespace UnrealBuildTool
 
 				ProcessOutput Output = new ProcessOutput();
 
-				FabricProcess.OutputDataReceived += new DataReceivedEventHandler((s, e) => Output.OutputReceivedDataEventHandler(s, e, Logger));
-				FabricProcess.ErrorDataReceived += new DataReceivedEventHandler((s, e) => Output.OutputReceivedDataEventHandler(s, e, Logger));
+				FabricProcess.OutputDataReceived += new DataReceivedEventHandler((_, e) => Output.OutputReceivedDataEventHandler(e, Logger));
+				FabricProcess.ErrorDataReceived += new DataReceivedEventHandler((_, e) => Output.OutputReceivedDataEventHandler(e, Logger));
 
 				Utils.RunLocalProcess(FabricProcess);
 				if (Output.OutputReceivedDataEventHandlerEncounteredError)
@@ -1247,7 +1246,7 @@ namespace UnrealBuildTool
 			return FileReference.Combine(Executable.Directory, "Payload", TargetName + ".app", TargetName);
 		}
 
-		private static void WriteEntitlements(IOSPostBuildSyncTarget Target, ILogger Logger)
+		private static void WriteEntitlements(IOSPostBuildSyncTarget Target)
 		{
 			string AppName = Target.TargetName;
 			FileReference? MobileProvisionFile;
@@ -1291,7 +1290,7 @@ namespace UnrealBuildTool
 			FileReference.WriteAllText(MissingPlistPath, Plist);
 		}
 
-		private static void GenerateFrameworkWrapperIfNonexistent(IOSPostBuildSyncTarget Target, ILogger Logger)
+		private static void GenerateFrameworkWrapperIfNonexistent(IOSPostBuildSyncTarget Target)
 		{
 			DirectoryReference ProjectDirectory = GetActualProjectDirectory(Target.ProjectFile);
 			DirectoryReference WrapperDirectory = DirectoryReference.Combine(ProjectDirectory, "Wrapper");
@@ -1366,7 +1365,7 @@ namespace UnrealBuildTool
 
 					PathToDsymZip = Path.Combine(Target.OutputPath.Directory.ParentDirectory!.FullName, Target.OutputPath.GetFileName() + ".dSYM.zip");
 
-					GenerateFrameworkWrapperIfNonexistent(Target, Logger);
+					GenerateFrameworkWrapperIfNonexistent(Target);
 
 					// do not perform any of the .app creation below
 					bPerformFullAppCreation = false;
@@ -1377,7 +1376,7 @@ namespace UnrealBuildTool
 
 			if (!Target.bSkipCrashlytics)
 			{
-				GenerateCrashlyticsData(PathToDsymZip, Target.ProjectDirectory.FullName, AppName, Logger);
+				GenerateCrashlyticsData(PathToDsymZip, Logger);
 			}
 
 			// only make the app if needed
@@ -1588,12 +1587,12 @@ namespace UnrealBuildTool
 				// ensure the plist, entitlements, and provision files are properly copied
 				UEDeployIOS DeployHandler = (Target.Platform == UnrealTargetPlatform.IOS ? new UEDeployIOS(Logger) : new UEDeployTVOS(Logger));
 				DeployHandler.ForDistribution = Target.bForDistribution;
-				DeployHandler.PrepTargetForDeployment(Target.ProjectFile, Target.TargetName, Target.OutputPath, Target.Platform, Target.Configuration, Target.UPLScripts, Target.bCreateStubIPA, BundleID, Target.bBuildAsFramework);
+				DeployHandler.PrepTargetForDeployment(Target.ProjectFile, Target.TargetName, Target.OutputPath, Target.Configuration, Target.UPLScripts, Target.bCreateStubIPA, BundleID, Target.bBuildAsFramework);
 
 				Logger.LogInformation("Executing {Script}", SignProjectScript);
 
 				// write the entitlements file (building remotely)
-				WriteEntitlements(Target, Logger);
+				WriteEntitlements(Target);
 
 				Process SignProcess = new Process();
 				SignProcess.StartInfo.WorkingDirectory = RemoteShadowDirectoryMac;
@@ -1603,8 +1602,8 @@ namespace UnrealBuildTool
 
 				ProcessOutput Output = new ProcessOutput();
 
-				SignProcess.OutputDataReceived += new DataReceivedEventHandler((s, e) => Output.OutputReceivedDataEventHandler(s, e, Logger));
-				SignProcess.ErrorDataReceived += new DataReceivedEventHandler((s, e) => Output.OutputReceivedDataEventHandler(s, e, Logger));
+				SignProcess.OutputDataReceived += new DataReceivedEventHandler((_, e) => Output.OutputReceivedDataEventHandler(e, Logger));
+				SignProcess.ErrorDataReceived += new DataReceivedEventHandler((_, e) => Output.OutputReceivedDataEventHandler(e, Logger));
 
 				Output.OutputReceivedDataEventHandlerEncounteredError = false;
 				Output.OutputReceivedDataEventHandlerEncounteredErrorMessage = "";
@@ -1630,8 +1629,8 @@ namespace UnrealBuildTool
 
 					ProcessOutput CleanOutput = new ProcessOutput();
 
-					SignProcess.OutputDataReceived += new DataReceivedEventHandler((s, e) => CleanOutput.OutputReceivedDataEventLogger(s, e, Logger));
-					SignProcess.ErrorDataReceived += new DataReceivedEventHandler((s, e) => CleanOutput.OutputReceivedDataEventLogger(s, e, Logger));
+					SignProcess.OutputDataReceived += new DataReceivedEventHandler((_, e) => CleanOutput.OutputReceivedDataEventLogger(e, Logger));
+					SignProcess.ErrorDataReceived += new DataReceivedEventHandler((_, e) => CleanOutput.OutputReceivedDataEventLogger(e, Logger));
 
 					Utils.RunLocalProcess(CleanProcess);
 				}
@@ -1667,10 +1666,10 @@ namespace UnrealBuildTool
 				// ensure the plist, entitlements, and provision files are properly copied
 				UEDeployIOS DeployHandler = (Target.Platform == UnrealTargetPlatform.IOS ? new UEDeployIOS(Logger) : new UEDeployTVOS(Logger));
 				DeployHandler.ForDistribution = Target.bForDistribution;
-				DeployHandler.PrepTargetForDeployment(Target.ProjectFile, Target.TargetName, Target.OutputPath, Target.Platform, Target.Configuration, Target.UPLScripts, Target.bCreateStubIPA, BundleID, Target.bBuildAsFramework);
+				DeployHandler.PrepTargetForDeployment(Target.ProjectFile, Target.TargetName, Target.OutputPath, Target.Configuration, Target.UPLScripts, Target.bCreateStubIPA, BundleID, Target.bBuildAsFramework);
 
 				// write the entitlements file (building on Mac)
-				WriteEntitlements(Target, Logger);
+				WriteEntitlements(Target);
 			}
 
 			{
diff --git a/Engine/Source/Programs/UnrealBuildTool/Platform/IOS/UEBuildIOS.cs b/Engine/Source/Programs/UnrealBuildTool/Platform/IOS/UEBuildIOS.cs
--- a/Engine/Source/Programs/UnrealBuildTool/Platform/IOS/UEBuildIOS.cs
+++ b/Engine/Source/Programs/UnrealBuildTool/Platform/IOS/UEBuildIOS.cs
@@ -896,7 +896,7 @@ namespace UnrealBuildTool
 
 			// look up iOS specific settings
 			if (!DoProjectSettingsMatchDefault(Platform, ProjectDirectoryName, "/Script/IOSRuntimeSettings.IOSRuntimeSettings",
-					BoolKeys, null, StringKeys, Logger))
+					BoolKeys, null, StringKeys))
 			{
 				return false;
 			}
diff --git a/Engine/Source/Programs/UnrealBuildTool/Platform/IOS/UEDeployIOS.cs b/Engine/Source/Programs/UnrealBuildTool/Platform/IOS/UEDeployIOS.cs
--- a/Engine/Source/Programs/UnrealBuildTool/Platform/IOS/UEDeployIOS.cs
+++ b/Engine/Source/Programs/UnrealBuildTool/Platform/IOS/UEDeployIOS.cs
@@ -148,7 +148,7 @@ namespace UnrealBuildTool
 			/// Updates the minor version in the CFBundleVersion key of the specified PList if this is a new package.
 			/// Also updates the key EpicAppVersion with the bundle version and the current date/time (no year)
 			/// </summary>
-			public static string UpdateBundleVersion(string OldPList, string EngineDirectory)
+			public static string UpdateBundleVersion(string OldPList)
 			{
 				string CFBundleVersion;
 				if (!Unreal.IsBuildMachine())
@@ -217,7 +217,7 @@ namespace UnrealBuildTool
 			return MinVersionToReturn;
 		}
 
-		public static void WritePlistFile(FileReference PlistFile, DirectoryReference? ProjectLocation, UnrealPluginLanguage? UPL, UnrealTargetConfiguration Config, string GameName, bool bIsUnrealGame, string ProjectName, ILogger Logger)
+		public static void WritePlistFile(FileReference PlistFile, DirectoryReference? ProjectLocation, UnrealPluginLanguage? UPL, UnrealTargetConfiguration Config, string GameName, bool bIsUnrealGame, string ProjectName)
 		{
 			ConfigHierarchy Ini = ConfigCache.ReadHierarchy(ConfigHierarchyType.Engine, ProjectLocation, UnrealTargetPlatform.IOS);
 			// required capabilities
@@ -478,7 +478,7 @@ namespace UnrealBuildTool
 
 				if (!AppleExports.UseModernXcode(ProjectFile))
 				{
-					return GenerateLegacyIOSPList(ProjectFile, Config, ProjectDirectory, bIsUnrealGame, GameName, bIsClient, ProjectName, InEngineDir, AppDirectory, UPL, BundleID, bBuildAsFramework, Logger);
+					return GenerateLegacyIOSPList(ProjectFile, ProjectDirectory, bIsUnrealGame, GameName, bIsClient, ProjectName, InEngineDir, AppDirectory, UPL, BundleID, bBuildAsFramework, Logger);
 				}
 
 				// generate the Info.plist for future use
@@ -490,7 +490,7 @@ namespace UnrealBuildTool
 				VersionUtilities.BuildDirectory = BuildDirectory;
 				VersionUtilities.GameName = GameName;
 
-				WritePlistFile(new FileReference(PListFile), DirRef, UPL, Config, GameName, bIsUnrealGame, ProjectName, Logger);
+				WritePlistFile(new FileReference(PListFile), DirRef, UPL, Config, GameName, bIsUnrealGame, ProjectName);
 
 				if (BuildHostPlatform.Current.Platform == UnrealTargetPlatform.Mac && !bBuildAsFramework)
 				{
@@ -512,7 +512,7 @@ namespace UnrealBuildTool
 			}
 		}
 
-		public static bool GenerateLegacyIOSPList(FileReference? ProjectFile, UnrealTargetConfiguration Config, string ProjectDirectory, bool bIsUnrealGame, string GameName, bool bIsClient, string ProjectName, string InEngineDir, string AppDirectory, UnrealPluginLanguage? UPL, string? BundleID, bool bBuildAsFramework, ILogger Logger)
+		public static bool GenerateLegacyIOSPList(FileReference? ProjectFile, string ProjectDirectory, bool bIsUnrealGame, string GameName, bool bIsClient, string ProjectName, string InEngineDir, string AppDirectory, UnrealPluginLanguage? UPL, string? BundleID, bool bBuildAsFramework, ILogger Logger)
 		{
 			// generate the Info.plist for future use
 			string BuildDirectory = ProjectDirectory + "/Build/IOS";
@@ -689,7 +689,7 @@ namespace UnrealBuildTool
 			Text.AppendLine("\t<key>CFBundleSignature</key>");
 			Text.AppendLine("\t<string>????</string>");
 			Text.AppendLine("\t<key>CFBundleVersion</key>");
-			Text.AppendLine(String.Format("\t<string>{0}</string>", VersionUtilities.UpdateBundleVersion(OldPListData, InEngineDir)));
+			Text.AppendLine(String.Format("\t<string>{0}</string>", VersionUtilities.UpdateBundleVersion(OldPListData)));
 			Text.AppendLine("\t<key>CFBundleShortVersionString</key>");
 			Text.AppendLine(String.Format("\t<string>{0}</string>", BundleShortVersion));
 			Text.AppendLine("\t<key>LSRequiresIPhoneOS</key>");
@@ -1029,12 +1029,12 @@ namespace UnrealBuildTool
 			}
 		}
 
-		public bool PrepForUATPackageOrDeploy(UnrealTargetConfiguration Config, FileReference ProjectFile, string InProjectName, string InProjectDirectory, FileReference Executable, string InEngineDir, bool bForDistribution, string CookFlavor, bool bIsDataDeploy, bool bCreateStubIPA, TargetReceipt Receipt)
+		public bool PrepForUATPackageOrDeploy(UnrealTargetConfiguration Config, FileReference ProjectFile, string InProjectName, string InProjectDirectory, FileReference Executable, string InEngineDir, bool bCreateStubIPA, TargetReceipt Receipt)
 		{
 			List<string> UPLScripts = CollectPluginDataPaths(Receipt.AdditionalProperties, Logger);
 			VersionNumber? SdkVersion = GetSdkVersion(Receipt);
 			bool bBuildAsFramework = GetCompileAsDll(Receipt);
-			return PrepForUATPackageOrDeploy(Config, ProjectFile, InProjectName, InProjectDirectory, Executable, InEngineDir, bForDistribution, CookFlavor, bIsDataDeploy, bCreateStubIPA, UPLScripts, "", bBuildAsFramework);
+			return PrepForUATPackageOrDeploy(Config, ProjectFile, InProjectName, InProjectDirectory, Executable, InEngineDir, bCreateStubIPA, UPLScripts, "", bBuildAsFramework);
 		}
 
 		void CopyAllProvisions(string ProvisionDir, ILogger Logger)
@@ -1086,7 +1086,7 @@ namespace UnrealBuildTool
 			}
 		}
 
-		public bool PrepForUATPackageOrDeploy(UnrealTargetConfiguration Config, FileReference? ProjectFile, string InProjectName, string InProjectDirectory, FileReference Executable, string InEngineDir, bool bForDistribution, string CookFlavor, bool bIsDataDeploy, bool bCreateStubIPA, List<string> UPLScripts, string? BundleID, bool bBuildAsFramework)
+		public bool PrepForUATPackageOrDeploy(UnrealTargetConfiguration Config, FileReference? ProjectFile, string InProjectName, string InProjectDirectory, FileReference Executable, string InEngineDir, bool bCreateStubIPA, List<string> UPLScripts, string? BundleID, bool bBuildAsFramework)
 		{
 
 			if (BuildHostPlatform.Current.Platform != UnrealTargetPlatform.Mac)
@@ -1290,10 +1290,10 @@ namespace UnrealBuildTool
 		{
 			List<string> UPLScripts = CollectPluginDataPaths(Receipt.AdditionalProperties, Logger);
 			bool bBuildAsFramework = GetCompileAsDll(Receipt);
-			return PrepTargetForDeployment(Receipt.ProjectFile, Receipt.TargetName, Receipt.BuildProducts.First(x => x.Type == BuildProductType.Executable).Path, Receipt.Platform, Receipt.Configuration, UPLScripts, false, "", bBuildAsFramework);
+			return PrepTargetForDeployment(Receipt.ProjectFile, Receipt.TargetName, Receipt.BuildProducts.First(x => x.Type == BuildProductType.Executable).Path, Receipt.Configuration, UPLScripts, false, "", bBuildAsFramework);
 		}
 
-		public bool PrepTargetForDeployment(FileReference? ProjectFile, string TargetName, FileReference Executable, UnrealTargetPlatform Platform, UnrealTargetConfiguration Configuration, List<string> UPLScripts, bool bCreateStubIPA, string? BundleID, bool bBuildAsFramework)
+		public bool PrepTargetForDeployment(FileReference? ProjectFile, string TargetName, FileReference Executable, UnrealTargetConfiguration Configuration, List<string> UPLScripts, bool bCreateStubIPA, string? BundleID, bool bBuildAsFramework)
 		{
 			string GameName = TargetName;
 			string ProjectDirectory = (DirectoryReference.FromFile(ProjectFile) ?? Unreal.EngineDirectory).FullName;
@@ -1301,7 +1301,7 @@ namespace UnrealBuildTool
 
 			if (BuildHostPlatform.Current.Platform == UnrealTargetPlatform.Mac && Environment.GetEnvironmentVariable("UBT_NO_POST_DEPLOY") != "true")
 			{
-				return PrepForUATPackageOrDeploy(Configuration, ProjectFile, GameName, ProjectDirectory, Executable, "../../Engine", bForDistribution, "", false, bCreateStubIPA, UPLScripts, BundleID, bBuildAsFramework);
+				return PrepForUATPackageOrDeploy(Configuration, ProjectFile, GameName, ProjectDirectory, Executable, "../../Engine", bCreateStubIPA, UPLScripts, BundleID, bBuildAsFramework);
 			}
 			else
 			{
diff --git a/Engine/Source/Programs/UnrealBuildTool/Platform/Linux/LinuxProjectGenerator.cs b/Engine/Source/Programs/UnrealBuildTool/Platform/Linux/LinuxProjectGenerator.cs
--- a/Engine/Source/Programs/UnrealBuildTool/Platform/Linux/LinuxProjectGenerator.cs
+++ b/Engine/Source/Programs/UnrealBuildTool/Platform/Linux/LinuxProjectGenerator.cs
@@ -21,8 +21,10 @@ namespace UnrealBuildTool
 		/// </summary>
 		/// <param name="Arguments">Command line arguments passed to the project generator</param>
 		/// <param name="Logger">Logger for output</param>
+#pragma warning disable IDE0060 // Remove unused parameter - constructor is found by reflection in GenerateProjectFilesMode.ExecuteAsync
 		public LinuxProjectGenerator(CommandLineArguments Arguments, ILogger Logger)
-			: base(Arguments, Logger)
+			: base(Logger)
+#pragma warning restore IDE0060
 		{
 		}
 
diff --git a/Engine/Source/Programs/UnrealBuildTool/Platform/Linux/LinuxToolChain.cs b/Engine/Source/Programs/UnrealBuildTool/Platform/Linux/LinuxToolChain.cs
--- a/Engine/Source/Programs/UnrealBuildTool/Platform/Linux/LinuxToolChain.cs
+++ b/Engine/Source/Programs/UnrealBuildTool/Platform/Linux/LinuxToolChain.cs
@@ -52,8 +52,11 @@ namespace UnrealBuildTool
 		protected LinuxToolChainInfo LinuxInfo => (Info as LinuxToolChainInfo)!;
 
 		public LinuxToolChain(UnrealArch InArchitecture, LinuxPlatformSDK InSDK, ClangToolChainOptions InOptions, ILogger InLogger)
-			: this(UnrealTargetPlatform.Linux, InArchitecture, InSDK, InOptions, InLogger)
+			: base(InOptions, InLogger)
 		{
+			Architecture = InArchitecture;
+			PlatformSDK = InSDK;
+
 			// prevent unknown clangs since the build is likely to fail on too old or too new compilers
 			if ((CompilerVersionLessThan(20, 0, 0) || CompilerVersionGreaterOrEqual(21, 0, 0)) && !Options.HasFlag(ClangToolChainOptions.UseAutoRTFMCompiler))
 			{
@@ -61,13 +64,6 @@ namespace UnrealBuildTool
 			}
 		}
 
-		public LinuxToolChain(UnrealTargetPlatform InPlatform, UnrealArch InArchitecture, LinuxPlatformSDK InSDK, ClangToolChainOptions InOptions, ILogger InLogger)
-			: base(InOptions, InLogger)
-		{
-			Architecture = InArchitecture;
-			PlatformSDK = InSDK;
-		}
-
 		protected override ClangToolChainInfo GetToolChainInfo()
 		{
 			DirectoryReference? MultiArchRoot = PlatformSDK.GetSDKLocation();
@@ -770,7 +766,7 @@ namespace UnrealBuildTool
 			}
 		}
 
-		static string GetArchiveArguments(LinkEnvironment LinkEnvironment)
+		static string GetArchiveArguments()
 		{
 			return " rcs";
 		}
@@ -903,7 +899,7 @@ namespace UnrealBuildTool
 		/// <summary>
 		/// Creates an action to archive all the .o files into single .a file
 		/// </summary>
-		public FileItem CreateArchiveAndIndex(LinkEnvironment LinkEnvironment, IActionGraphBuilder Graph, ILogger Logger)
+		public FileItem CreateArchiveAndIndex(LinkEnvironment LinkEnvironment, IActionGraphBuilder Graph)
 		{
 			// Create an archive action
 			Action ArchiveAction = Graph.CreateAction(ActionType.Link);
@@ -920,7 +916,7 @@ namespace UnrealBuildTool
 			ArchiveAction.ProducedItems.Add(OutputFile);
 			ArchiveAction.CommandDescription = "Archive";
 			ArchiveAction.StatusDescription = Path.GetFileName(OutputFile.AbsolutePath);
-			ArchiveAction.CommandArguments += $"{GetArchiveArguments(LinkEnvironment)} \"{OutputFile.AbsolutePath}\"";
+			ArchiveAction.CommandArguments += $"{GetArchiveArguments()} \"{OutputFile.AbsolutePath}\"";
 
 			// Add the input files to a response file, and pass the response file on the command-line.
 			List<string> InputFileNames = [];
@@ -967,7 +963,7 @@ namespace UnrealBuildTool
 
 			if (LinkEnvironment.bIsBuildingLibrary || bBuildImportLibraryOnly)
 			{
-				return CreateArchiveAndIndex(LinkEnvironment, Graph, Logger);
+				return CreateArchiveAndIndex(LinkEnvironment, Graph);
 			}
 
 			// Create an action that invokes the linker.
diff --git a/Engine/Source/Programs/UnrealBuildTool/Platform/Mac/MacProjectGenerator.cs b/Engine/Source/Programs/UnrealBuildTool/Platform/Mac/MacProjectGenerator.cs
--- a/Engine/Source/Programs/UnrealBuildTool/Platform/Mac/MacProjectGenerator.cs
+++ b/Engine/Source/Programs/UnrealBuildTool/Platform/Mac/MacProjectGenerator.cs
@@ -16,8 +16,10 @@ namespace UnrealBuildTool
 		/// </summary>
 		/// <param name="Arguments">Command line arguments passed to the project generator</param>
 		/// <param name="Logger">Logger for output</param>
+#pragma warning disable IDE0060 // Remove unused parameter - constructor is found by reflection in GenerateProjectFilesMode.ExecuteAsync
 		public MacProjectGenerator(CommandLineArguments Arguments, ILogger Logger)
-			: base(Arguments, Logger)
+			: base(Logger)
+#pragma warning restore IDE0060
 		{
 		}
 
diff --git a/Engine/Source/Programs/UnrealBuildTool/Platform/Mac/MacToolChain.cs b/Engine/Source/Programs/UnrealBuildTool/Platform/Mac/MacToolChain.cs
--- a/Engine/Source/Programs/UnrealBuildTool/Platform/Mac/MacToolChain.cs
+++ b/Engine/Source/Programs/UnrealBuildTool/Platform/Mac/MacToolChain.cs
@@ -848,7 +848,7 @@ namespace UnrealBuildTool
 				{
 					if (Framework.ZipFile != null)
 					{
-						FileItem ExtractedTokenFile = ExtractFramework(Framework, Graph, LinkEnvironment.IntermediateDirectory!, Logger);
+						FileItem ExtractedTokenFile = ExtractFramework(Framework, Graph);
 						LinkAction.PrerequisiteItems.Add(ExtractedTokenFile);
 					}
 
@@ -956,8 +956,7 @@ namespace UnrealBuildTool
 		/// <param name="MachOBinary">FileItem describing the executable or dylib to generate debug info for</param>
 		/// <param name="LinkEnvironment"></param>
 		/// <param name="Graph">List of actions to be executed. Additional actions will be added to this list.</param>
-		/// <param name="Logger">Logger for output</param>
-		public FileItem GenerateDebugInfo(FileItem MachOBinary, LinkEnvironment LinkEnvironment, IActionGraphBuilder Graph, ILogger Logger)
+		public FileItem GenerateDebugInfo(FileItem MachOBinary, LinkEnvironment LinkEnvironment, IActionGraphBuilder Graph)
 		{
 			string BinaryPath = MachOBinary.AbsolutePath;
 			if (BinaryPath.Contains(".app"))
@@ -1037,11 +1036,10 @@ namespace UnrealBuildTool
 		/// <summary>
 		/// Creates app bundle for a given executable
 		/// </summary>
-		/// <param name="Target"></param>
 		/// <param name="LinkEnvironment"></param>
 		/// <param name="Executable">FileItem describing the executable to generate app bundle for</param>
 		/// <param name="Graph">List of actions to be executed. Additional actions will be added to this list.</param>
-		FileItem FinalizeAppBundle(ReadOnlyTargetRules Target, LinkEnvironment LinkEnvironment, FileItem Executable, IActionGraphBuilder Graph)
+		FileItem FinalizeAppBundle(LinkEnvironment LinkEnvironment, FileItem Executable, IActionGraphBuilder Graph)
 		{
 			// Make a file item for the source and destination files
 			string FullDestPath = Executable.AbsolutePath.Substring(0, Executable.AbsolutePath.IndexOf(".app") + 4);
@@ -1252,7 +1250,7 @@ namespace UnrealBuildTool
 			// For Mac, generate the dSYM file if the config file is set to do so
 			if (BinaryLinkEnvironment.bUsePDBFiles == true)
 			{
-				DebugInfoFiles.Add(GenerateDebugInfo(Executable, BinaryLinkEnvironment, Graph, Logger));
+				DebugInfoFiles.Add(GenerateDebugInfo(Executable, BinaryLinkEnvironment, Graph));
 			}
 
 			if ((BinaryLinkEnvironment.bIsBuildingDLL && (Options & ClangToolChainOptions.OutputDylib) == 0) || (BinaryLinkEnvironment.bIsBuildingConsoleApplication && Executable.Name.EndsWith("-Cmd")))
@@ -1265,7 +1263,7 @@ namespace UnrealBuildTool
 			{
 				if (!bUseModernXcode)
 				{
-					OutputFiles.Add(FinalizeAppBundle(Target, BinaryLinkEnvironment, Executable, Graph));
+					OutputFiles.Add(FinalizeAppBundle(BinaryLinkEnvironment, Executable, Graph));
 				}
 			}
 
diff --git a/Engine/Source/Programs/UnrealBuildTool/Platform/TVOS/TVOSExports.cs b/Engine/Source/Programs/UnrealBuildTool/Platform/TVOS/TVOSExports.cs
--- a/Engine/Source/Programs/UnrealBuildTool/Platform/TVOS/TVOSExports.cs
+++ b/Engine/Source/Programs/UnrealBuildTool/Platform/TVOS/TVOSExports.cs
@@ -76,7 +76,7 @@ namespace UnrealBuildTool
 		public static bool PrepForUATPackageOrDeploy(UnrealTargetConfiguration Config, FileReference ProjectFile, string InProjectName, DirectoryReference InProjectDirectory, FileReference Executable, DirectoryReference InEngineDir, bool bForDistribution, string CookFlavor, bool bIsDataDeploy, bool bCreateStubIPA, FileReference BuildReceiptFileName, ILogger Logger)
 		{
 			TargetReceipt Receipt = TargetReceipt.Read(BuildReceiptFileName);
-			return new UEDeployTVOS(Logger).PrepForUATPackageOrDeploy(Config, ProjectFile, InProjectName, InProjectDirectory.FullName, Executable, InEngineDir.FullName, bForDistribution, CookFlavor, bIsDataDeploy, bCreateStubIPA, Receipt);
+			return new UEDeployTVOS(Logger).PrepForUATPackageOrDeploy(Config, ProjectFile, InProjectName, InProjectDirectory.FullName, Executable, InEngineDir.FullName, bCreateStubIPA, Receipt);
 		}
 
 		/// <summary>
diff --git a/Engine/Source/Programs/UnrealBuildTool/Platform/TVOS/UEDeployTVOS.cs b/Engine/Source/Programs/UnrealBuildTool/Platform/TVOS/UEDeployTVOS.cs
--- a/Engine/Source/Programs/UnrealBuildTool/Platform/TVOS/UEDeployTVOS.cs
+++ b/Engine/Source/Programs/UnrealBuildTool/Platform/TVOS/UEDeployTVOS.cs
@@ -111,7 +111,7 @@ namespace UnrealBuildTool
 			Text.AppendLine("\t<key>CFBundleSignature</key>");
 			Text.AppendLine("\t<string>????</string>");
 			Text.AppendLine("\t<key>CFBundleVersion</key>");
-			Text.AppendLine(String.Format("\t<string>{0}</string>", VersionUtilities.UpdateBundleVersion(OldPListData, InEngineDir)));
+			Text.AppendLine(String.Format("\t<string>{0}</string>", VersionUtilities.UpdateBundleVersion(OldPListData)));
 			Text.AppendLine("\t<key>CFBundleShortVersionString</key>");
 			Text.AppendLine(String.Format("\t<string>{0}</string>", BundleShortVersion));
 			Text.AppendLine("\t<key>LSRequiresIPhoneOS</key>");
diff --git a/Engine/Source/Programs/UnrealBuildTool/Platform/Windows/MicrosoftPlatformSDK.cs b/Engine/Source/Programs/UnrealBuildTool/Platform/Windows/MicrosoftPlatformSDK.cs
--- a/Engine/Source/Programs/UnrealBuildTool/Platform/Windows/MicrosoftPlatformSDK.cs
+++ b/Engine/Source/Programs/UnrealBuildTool/Platform/Windows/MicrosoftPlatformSDK.cs
@@ -1066,13 +1066,13 @@ namespace UnrealBuildTool
 					VersionNumber? Version;
 					if (IsValidToolChainDirMSVC(ToolChainDir, UnrealArch.X64, out Version))
 					{
-						DirectoryReference? RedistDir = FindVisualStudioRedistForToolChain(ToolChainDir, OptionalRedistDir, Version);
+						DirectoryReference? RedistDir = FindVisualStudioRedistForToolChain(ToolChainDir, OptionalRedistDir);
 						AddVisualCppToolChain(Version, UnrealArch.X64, ReleaseChannel, ToolChainDir, RedistDir, ToolChains, IsAutoSdk, Logger);
 					}
 
 					if (IsValidToolChainDirMSVC(ToolChainDir, UnrealArch.Arm64, out Version))
 					{
-						DirectoryReference? RedistDir = FindVisualStudioRedistForToolChain(ToolChainDir, OptionalRedistDir, Version);
+						DirectoryReference? RedistDir = FindVisualStudioRedistForToolChain(ToolChainDir, OptionalRedistDir);
 						AddVisualCppToolChain(Version, UnrealArch.Arm64, ReleaseChannel, ToolChainDir, RedistDir, ToolChains, IsAutoSdk, Logger);
 					}
 				}
@@ -1084,9 +1084,8 @@ namespace UnrealBuildTool
 		/// </summary>
 		/// <param name="ToolChainDir"></param>
 		/// <param name="OptionalRedistDir"></param>
-		/// <param name="Version"></param>
 		/// <returns></returns>
-		static DirectoryReference? FindVisualStudioRedistForToolChain(DirectoryReference ToolChainDir, DirectoryReference? OptionalRedistDir, VersionNumber Version)
+		static DirectoryReference? FindVisualStudioRedistForToolChain(DirectoryReference ToolChainDir, DirectoryReference? OptionalRedistDir)
 		{
 			DirectoryReference? RedistDir;
 			if (OptionalRedistDir == null)
@@ -1167,13 +1166,12 @@ namespace UnrealBuildTool
 		/// <summary>
 		/// Add a Clang toolchain
 		/// </summary>
-		/// <param name="Compiler"></param>
 		/// <param name="ToolChainDir"></param>
 		/// <param name="CompilerFile"></param>
 		/// <param name="ToolChains"></param>
 		/// <param name="IsAutoSdk"></param>
 		/// <param name="Logger"></param>
-		static void AddClangToolChain(WindowsCompiler Compiler, DirectoryReference ToolChainDir, FileReference CompilerFile, List<ToolChainInstallation> ToolChains, bool IsAutoSdk, ILogger Logger)
+		static void AddClangToolChain(DirectoryReference ToolChainDir, FileReference CompilerFile, List<ToolChainInstallation> ToolChains, bool IsAutoSdk, ILogger Logger)
 		{
 			if (FileReference.Exists(CompilerFile))
 			{
@@ -1223,8 +1221,8 @@ namespace UnrealBuildTool
 		static void AddClangToolChain(WindowsCompiler Compiler, DirectoryReference ToolChainDir, List<ToolChainInstallation> ToolChains, bool IsAutoSdk, ILogger Logger)
 		{
 			string CompilerFilename = Compiler == WindowsCompiler.ClangRTFM ? "verse-clang-cl.exe" : (Compiler == WindowsCompiler.ClangInstrument ? "instr-clang-cl.exe" : "clang-cl.exe");
-			AddClangToolChain(Compiler, ToolChainDir, FileReference.Combine(ToolChainDir, "bin", CompilerFilename), ToolChains, IsAutoSdk, Logger);
-			AddClangToolChain(Compiler, ToolChainDir, FileReference.Combine(ToolChainDir, "bin-woa64", CompilerFilename), ToolChains, IsAutoSdk, Logger);
+			AddClangToolChain(ToolChainDir, FileReference.Combine(ToolChainDir, "bin", CompilerFilename), ToolChains, IsAutoSdk, Logger);
+			AddClangToolChain(ToolChainDir, FileReference.Combine(ToolChainDir, "bin-woa64", CompilerFilename), ToolChains, IsAutoSdk, Logger);
 		}
 
 		/// <summary>
diff --git a/Engine/Source/Programs/UnrealBuildTool/Platform/Windows/PVSToolChain.cs b/Engine/Source/Programs/UnrealBuildTool/Platform/Windows/PVSToolChain.cs
--- a/Engine/Source/Programs/UnrealBuildTool/Platform/Windows/PVSToolChain.cs
+++ b/Engine/Source/Programs/UnrealBuildTool/Platform/Windows/PVSToolChain.cs
@@ -709,7 +709,7 @@ namespace UnrealBuildTool
 		const string CPP_20 = "c++20";
 		const string CPP_23 = "c++23";
 
-		public static string GetLangStandForCfgFile(CppStandardVersion cppStandard, VersionNumber compilerVersion)
+		public static string GetLangStandForCfgFile(CppStandardVersion cppStandard)
 		{
 			return cppStandard switch
 			{
@@ -802,8 +802,7 @@ namespace UnrealBuildTool
 			bool isVisualCppCompiler = windowsCompiler.IsMSVC();
 			if (!shouldCompileAsC)
 			{
-				VersionNumber compilerVersion = Target.WindowsPlatform.Environment.CompilerVersion;
-				string languageStandardForCfg = GetLangStandForCfgFile(compileEnvironment.CppStandard, compilerVersion);
+				string languageStandardForCfg = GetLangStandForCfgFile(compileEnvironment.CppStandard);
 
 				configFileContents.AppendFormat("std={0}\n", languageStandardForCfg);
 
diff --git a/Engine/Source/Programs/UnrealBuildTool/Platform/Windows/ResXResource/UEResXWriter.cs b/Engine/Source/Programs/UnrealBuildTool/Platform/Windows/ResXResource/UEResXWriter.cs
--- a/Engine/Source/Programs/UnrealBuildTool/Platform/Windows/ResXResource/UEResXWriter.cs
+++ b/Engine/Source/Programs/UnrealBuildTool/Platform/Windows/ResXResource/UEResXWriter.cs
@@ -21,10 +21,10 @@ namespace UnrealBuildTool
 			RootElement = Document.CreateElement("root");
 			Document.AppendChild(RootElement);
 
-			RootElement.AppendChild(CreateEntry("resheader", "resmimetype", "text/microsoft-resx", null));
-			RootElement.AppendChild(CreateEntry("resheader", "version", "2.0", null));
-			RootElement.AppendChild(CreateEntry("resheader", "reader", typeof(UEResXReader).AssemblyQualifiedName!, null));
-			RootElement.AppendChild(CreateEntry("resheader", "writer", typeof(UEResXWriter).AssemblyQualifiedName!, null));
+			RootElement.AppendChild(CreateEntry("resmimetype", "text/microsoft-resx", null));
+			RootElement.AppendChild(CreateEntry("version", "2.0", null));
+			RootElement.AppendChild(CreateEntry("reader", typeof(UEResXReader).AssemblyQualifiedName!, null));
+			RootElement.AppendChild(CreateEntry("writer", typeof(UEResXWriter).AssemblyQualifiedName!, null));
 		}
 
 		/// <summary>
@@ -40,10 +40,10 @@ namespace UnrealBuildTool
 		/// </summary>
 		public void AddResource(string InName, string InValue)
 		{
-			RootElement.AppendChild(CreateEntry("data", InName, InValue, "preserve"));
+			RootElement.AppendChild(CreateEntry(InName, InValue, "preserve"));
 		}
 
-		private XmlNode CreateEntry(string InRootName, string InName, string InValue, string? InSpace)
+		private XmlNode CreateEntry(string InName, string InValue, string? InSpace)
 		{
 			XmlElement Value = Document.CreateElement("value");
 			Value.InnerText = InValue;
diff --git a/Engine/Source/Programs/UnrealBuildTool/Platform/Windows/VCCompileAction.cs b/Engine/Source/Programs/UnrealBuildTool/Platform/Windows/VCCompileAction.cs
--- a/Engine/Source/Programs/UnrealBuildTool/Platform/Windows/VCCompileAction.cs
+++ b/Engine/Source/Programs/UnrealBuildTool/Platform/Windows/VCCompileAction.cs
@@ -508,17 +508,16 @@ namespace UnrealBuildTool
 		/// Writes the response file with the action's arguments
 		/// </summary>
 		/// <param name="graph">The graph builder</param>
-		/// <param name="logger">Logger for output</param>
-		public void WriteResponseFile(IActionGraphBuilder graph, ILogger logger)
+		public void WriteResponseFile(IActionGraphBuilder graph)
 		{
 			if (ResponseFile != null)
 			{
-				graph.CreateIntermediateTextFile(ResponseFile, GetCompilerArguments(logger));
+				graph.CreateIntermediateTextFile(ResponseFile, GetCompilerArguments());
 				Arguments.Clear();
 			}
 		}
 
-		public List<string> GetCompilerArguments(ILogger logger)
+		public List<string> GetCompilerArguments()
 		{
 			List<string> arguments = [];
 
diff --git a/Engine/Source/Programs/UnrealBuildTool/Platform/Windows/VCEnvironment.cs b/Engine/Source/Programs/UnrealBuildTool/Platform/Windows/VCEnvironment.cs
--- a/Engine/Source/Programs/UnrealBuildTool/Platform/Windows/VCEnvironment.cs
+++ b/Engine/Source/Programs/UnrealBuildTool/Platform/Windows/VCEnvironment.cs
@@ -148,10 +148,10 @@ namespace UnrealBuildTool
 			LinkerPath = GetLinkerToolPath(Compiler, Architecture, CompilerDir, ToolChainDir);
 			LibraryManagerPath = GetLibraryLinkerToolPath(Compiler, Architecture, ToolChainDir);
 			ToolchainCompilerPath = GetCompilerToolPath(ToolChain, Architecture, ToolChainDir);
-			ResourceToObjectConverterPath = GetResourceToObjectConverterToolPath(ToolChain, Architecture, ToolChainDir);
+			ResourceToObjectConverterPath = GetResourceToObjectConverterToolPath(Architecture, ToolChainDir);
 
 			// Get the resource compiler path from the Windows SDK
-			ResourceCompilerPath = GetResourceCompilerToolPath(WindowsSdkDir, WindowsSdkVersion, Logger);
+			ResourceCompilerPath = GetResourceCompilerToolPath(WindowsSdkDir, WindowsSdkVersion);
 
 			// Register any executables for cross-architecture support with UBA
 			RegisterCrossArchitectureToolPaths(WindowsSdkDir, WindowsSdkVersion);
@@ -287,7 +287,7 @@ namespace UnrealBuildTool
 		/// <summary>
 		/// Gets the path to the resource to object converter.
 		/// </summary>
-		private FileReference GetResourceToObjectConverterToolPath(WindowsCompiler Compiler, UnrealArch Architecture, DirectoryReference ToochainDir)
+		private FileReference GetResourceToObjectConverterToolPath(UnrealArch Architecture, DirectoryReference ToochainDir)
 		{
 			return FileReference.Combine(GetVCToolPath(ToochainDir, Architecture), "cvtres.exe");
 		}
@@ -295,7 +295,7 @@ namespace UnrealBuildTool
 		/// <summary>
 		/// Gets the path to the resource compiler.
 		/// </summary>
-		private FileReference GetResourceCompilerToolPath(DirectoryReference WindowsSdkDir, VersionNumber WindowsSdkVersion, ILogger Logger)
+		private FileReference GetResourceCompilerToolPath(DirectoryReference WindowsSdkDir, VersionNumber WindowsSdkVersion)
 		{
 			string BuildHostArch = RuntimeInformation.ProcessArchitecture == System.Runtime.InteropServices.Architecture.Arm64 ? "arm64" : "x64";
 			FileReference ResourceCompilerPath = FileReference.Combine(WindowsSdkDir, "bin", WindowsSdkVersion.ToString(), BuildHostArch, "rc.exe");
diff --git a/Engine/Source/Programs/UnrealBuildTool/Platform/Windows/VCSpecificFileAction.cs b/Engine/Source/Programs/UnrealBuildTool/Platform/Windows/VCSpecificFileAction.cs
--- a/Engine/Source/Programs/UnrealBuildTool/Platform/Windows/VCSpecificFileAction.cs
+++ b/Engine/Source/Programs/UnrealBuildTool/Platform/Windows/VCSpecificFileAction.cs
@@ -113,7 +113,7 @@ namespace UnrealBuildTool
 				Action.DependencyListFile = FileItem.GetItemByFileReference(FileReference.Combine(IntermediateDirectory, $"{Filename}.d"));
 			}
 			Action.ResponseFile = FileItem.GetItemByFileReference(FileReference.Combine(IntermediateDirectory, $"{Filename}.rsp"));
-			File.WriteAllLines(Action.ResponseFile.FullName, Action.GetCompilerArguments(Logger));
+			File.WriteAllLines(Action.ResponseFile.FullName, Action.GetCompilerArguments());
 			Action.Arguments.Clear();
 
 			return Action;
diff --git a/Engine/Source/Programs/UnrealBuildTool/Platform/Windows/VCToolChain.cs b/Engine/Source/Programs/UnrealBuildTool/Platform/Windows/VCToolChain.cs
--- a/Engine/Source/Programs/UnrealBuildTool/Platform/Windows/VCToolChain.cs
+++ b/Engine/Source/Programs/UnrealBuildTool/Platform/Windows/VCToolChain.cs
@@ -2203,14 +2203,14 @@ namespace UnrealBuildTool
 				if (CompileEnvironment.bPreprocessOnly)
 				{
 					CompileAction.PreprocessedFile = FileItem.GetItemByFileReference(FileReference.Combine(OutputDir, FileName + ".i"));
-					CompileAction.ResponseFile = FileItem.GetItemByFileReference(GetResponseFileName(CompileEnvironment, CompileAction.PreprocessedFile));
+					CompileAction.ResponseFile = FileItem.GetItemByFileReference(GetResponseFileName(CompileAction.PreprocessedFile));
 				}
 				else if (Target.WindowsPlatform.Compiler.IsClang() && Target.StaticAnalyzer == StaticAnalyzer.Default && CompileEnvironment.PrecompiledHeaderAction != PrecompiledHeaderAction.Create)
 				{
 					// Clang analysis does not actually create an object, use the dependency list as the response filename
 					string DependencyListFilename = FileName + ".d";
 					FileItem DependencyListFile = FileItem.GetItemByFileReference(FileReference.Combine(OutputDir, DependencyListFilename));
-					CompileAction.ResponseFile = FileItem.GetItemByFileReference(GetResponseFileName(CompileEnvironment, DependencyListFile));
+					CompileAction.ResponseFile = FileItem.GetItemByFileReference(GetResponseFileName(DependencyListFile));
 				}
 				else
 				{
@@ -2219,7 +2219,7 @@ namespace UnrealBuildTool
 					FileItem ObjectFile = FileItem.GetItemByFileReference(FileReference.Combine(OutputDir, ObjectLeafFilename));
 
 					CompileAction.ObjectFile = ObjectFile;
-					CompileAction.ResponseFile = FileItem.GetItemByFileReference(GetResponseFileName(CompileEnvironment, ObjectFile));
+					CompileAction.ResponseFile = FileItem.GetItemByFileReference(GetResponseFileName(ObjectFile));
 
 					if (Target.WindowsPlatform.ObjSrcMapFile != null)
 					{
@@ -2364,7 +2364,7 @@ namespace UnrealBuildTool
 
 					VCCompileAction CompileDepsAction = new VCCompileAction(CompileAction);
 					CompileDepsAction.ActionType = ActionType.GatherModuleDependencies;
-					CompileDepsAction.ResponseFile = FileItem.GetItemByFileReference(GetResponseFileName(CompileEnvironment, IfcDepsFile));
+					CompileDepsAction.ResponseFile = FileItem.GetItemByFileReference(GetResponseFileName(IfcDepsFile));
 					CompileDepsAction.ObjectFile = null;
 					CompileDepsAction.DependencyListFile = IfcDepsFile;
 					CompileDepsAction.Arguments.Add($"/sourceDependencies:directives \"{IfcDepsFile.Location}\"");
@@ -2375,7 +2375,7 @@ namespace UnrealBuildTool
 
 					if (!ProjectFileGenerator.bGenerateProjectFiles)
 					{
-						CompileDepsAction.WriteResponseFile(Graph, Logger);
+						CompileDepsAction.WriteResponseFile(Graph);
 					}
 
 					CompileAction.ActionType = ActionType.CompileModuleInterface;
@@ -2427,7 +2427,7 @@ namespace UnrealBuildTool
 
 				if (!ProjectFileGenerator.bGenerateProjectFiles)
 				{
-					CompileAction.WriteResponseFile(Graph, Logger);
+					CompileAction.WriteResponseFile(Graph);
 				}
 
 				// Must be added after response file is created just to make sure it ends up on the command line and not in the response file
@@ -2891,7 +2891,7 @@ namespace UnrealBuildTool
 				Arguments.Add($"\"{NormalizeCommandLinePath(RCFile, RootPaths)}\"");
 
 				// Create a response file for the resource compilier for informational purposes as it cannot be used
-				FileItem ResponseFile = FileItem.GetItemByFileReference(GetResponseFileName(CompileEnvironment, CompiledResourceFile));
+				FileItem ResponseFile = FileItem.GetItemByFileReference(GetResponseFileName(CompiledResourceFile));
 				Graph.CreateIntermediateTextFile(ResponseFile, Arguments);
 				CompileAction.PrerequisiteItems.Add(ResponseFile);
 
diff --git a/Engine/Source/Programs/UnrealBuildTool/Platform/Windows/WindowsProjectGenerator.cs b/Engine/Source/Programs/UnrealBuildTool/Platform/Windows/WindowsProjectGenerator.cs
--- a/Engine/Source/Programs/UnrealBuildTool/Platform/Windows/WindowsProjectGenerator.cs
+++ b/Engine/Source/Programs/UnrealBuildTool/Platform/Windows/WindowsProjectGenerator.cs
@@ -70,8 +70,10 @@ namespace UnrealBuildTool
 		/// </summary>
 		/// <param name="Arguments">Command line arguments passed to the project generator</param>
 		/// <param name="Logger">Logger for output</param>
+#pragma warning disable IDE0060 // Remove unused parameter - constructor is found by reflection in GenerateProjectFilesMode.ExecuteAsync
 		public WindowsProjectGenerator(CommandLineArguments Arguments, ILogger Logger)
-			: base(Arguments, Logger)
+			: base(Logger)
+#pragma warning restore IDE0060
 		{
 			XmlConfig.ApplyTo(this);
 
diff --git a/Engine/Source/Programs/UnrealBuildTool/ProjectFiles/AndroidStudio/AndroidStudioProjectFileGenerator.cs b/Engine/Source/Programs/UnrealBuildTool/ProjectFiles/AndroidStudio/AndroidStudioProjectFileGenerator.cs
--- a/Engine/Source/Programs/UnrealBuildTool/ProjectFiles/AndroidStudio/AndroidStudioProjectFileGenerator.cs
+++ b/Engine/Source/Programs/UnrealBuildTool/ProjectFiles/AndroidStudio/AndroidStudioProjectFileGenerator.cs
@@ -58,7 +58,7 @@ namespace UnrealBuildTool
 		protected override bool WriteProjectFiles(PlatformProjectGeneratorCollection PlatformProjectGenerators, ILogger Logger)
 		{
 			TargetRules ClientTargetRules = DefaultProject!.ProjectTargets.First(Project => Project.TargetRules!.Type == TargetType.Client).TargetRules!;
-			UnrealArchitectures Architectures = VCProjectFileGenerator.GetAndroidProjectArchitectures(OnlyGameProject, true);
+			UnrealArchitectures Architectures = VCProjectFileGenerator.GetAndroidProjectArchitectures(OnlyGameProject);
 			UnrealArch Architecture = Architectures.Architectures.First();
 
 			string ProjectDir = Path.Combine(PrimaryProjectPath.FullName, $"{PrimaryProjectName}.gradle");
diff --git a/Engine/Source/Programs/UnrealBuildTool/ProjectFiles/CMake/CMakefileGenerator.cs b/Engine/Source/Programs/UnrealBuildTool/ProjectFiles/CMake/CMakefileGenerator.cs
--- a/Engine/Source/Programs/UnrealBuildTool/ProjectFiles/CMake/CMakefileGenerator.cs
+++ b/Engine/Source/Programs/UnrealBuildTool/ProjectFiles/CMake/CMakefileGenerator.cs
@@ -308,7 +308,7 @@ namespace UnrealBuildTool
 					string SourceFileRelativeToRoot = CurSourceFile.MakeRelativeTo(Unreal.EngineDirectory);
 
 					// Exclude files/folders on a per-platform basis.
-					if (!IsPathExcludedOnPlatform(SourceFileRelativeToRoot, BuildHostPlatform.Current.Platform))
+					if (!IsPathExcludedOnPlatform(SourceFileRelativeToRoot))
 					{
 						if (SourceFileRelativeToRoot.EndsWith(".cpp"))
 						{
@@ -527,7 +527,7 @@ namespace UnrealBuildTool
 					bWriteIncludes && bWriteDefinitions;
 		}
 
-		private static bool IsPathExcludedOnPlatform(string SourceFileRelativeToRoot, UnrealTargetPlatform targetPlatform)
+		private static bool IsPathExcludedOnPlatform(string SourceFileRelativeToRoot)
 		{
 			if (BuildHostPlatform.Current.Platform == UnrealTargetPlatform.Linux)
 			{
diff --git a/Engine/Source/Programs/UnrealBuildTool/ProjectFiles/PlatformProjectGenerator.cs b/Engine/Source/Programs/UnrealBuildTool/ProjectFiles/PlatformProjectGenerator.cs
--- a/Engine/Source/Programs/UnrealBuildTool/ProjectFiles/PlatformProjectGenerator.cs
+++ b/Engine/Source/Programs/UnrealBuildTool/ProjectFiles/PlatformProjectGenerator.cs
@@ -20,9 +20,8 @@ namespace UnrealBuildTool
 		/// <summary>
 		/// Constructor
 		/// </summary>
-		/// <param name="Arguments">Command line arguments passed to the project generator</param>
 		/// <param name="Logger">Logger for output</param>
-		public PlatformProjectGenerator(CommandLineArguments Arguments, ILogger Logger)
+		public PlatformProjectGenerator(ILogger Logger)
 		{
 			this.Logger = Logger;
 		}
diff --git a/Engine/Source/Programs/UnrealBuildTool/ProjectFiles/PlatformProjectGeneratorCollection.cs b/Engine/Source/Programs/UnrealBuildTool/ProjectFiles/PlatformProjectGeneratorCollection.cs
--- a/Engine/Source/Programs/UnrealBuildTool/ProjectFiles/PlatformProjectGeneratorCollection.cs
+++ b/Engine/Source/Programs/UnrealBuildTool/ProjectFiles/PlatformProjectGeneratorCollection.cs
@@ -112,7 +112,7 @@ namespace UnrealBuildTool
 			return true;
 		}
 
-		public bool PlatformRequiresVSUserFileGeneration(List<UnrealTargetPlatform> InPlatforms, List<UnrealTargetConfiguration> InConfigurations)
+		public bool PlatformRequiresVSUserFileGeneration(List<UnrealTargetPlatform> InPlatforms)
 		{
 			bool bRequiresVSUserFileGeneration = false;
 			foreach (KeyValuePair<UnrealTargetPlatform, PlatformProjectGenerator> Entry in ProjectGeneratorDictionary)
diff --git a/Engine/Source/Programs/UnrealBuildTool/ProjectFiles/ProjectFileGenerator.cs b/Engine/Source/Programs/UnrealBuildTool/ProjectFiles/ProjectFileGenerator.cs
--- a/Engine/Source/Programs/UnrealBuildTool/ProjectFiles/ProjectFileGenerator.cs
+++ b/Engine/Source/Programs/UnrealBuildTool/ProjectFiles/ProjectFileGenerator.cs
@@ -3350,7 +3350,7 @@ namespace UnrealBuildTool
 						UnrealProjectFilePath: CheckProjectFile,
 						SupportedPlatforms: TargetRulesObject.GetSupportedPlatforms().Where(
 												x => UEBuildPlatform.TryGetBuildPlatform(x, out UEBuildPlatform? _) &&
-												(TargetRulesObject.LinkType != TargetLinkType.Modular || !UEBuildPlatform.PlatformRequiresMonolithicBuilds(x, TargetRulesObject.Configuration))
+												(TargetRulesObject.LinkType != TargetLinkType.Modular || !UEBuildPlatform.PlatformRequiresMonolithicBuilds(x))
 												).ToArray(),
 						CreateRulesDelegate: (Platform, Configuration) => RulesAssembly.CreateTargetRules(TargetName, Platform, Configuration, null, CheckProjectFile, new CommandLineArguments(GetTargetArguments(Arguments)), Logger, false, TargetRulesValidationOptions.ValidateNothing, UnrealIntermediateEnvironment.GenerateProjectFiles),
 						CreateRulesDelegateWithArch: (Platform, Configuration, Arch) => RulesAssembly.CreateTargetRules(TargetName, Platform, Configuration, new(Arch), CheckProjectFile, new CommandLineArguments(GetTargetArguments(Arguments)), Logger, false, TargetRulesValidationOptions.ValidateNothing, UnrealIntermediateEnvironment.GenerateProjectFiles)
diff --git a/Engine/Source/Programs/UnrealBuildTool/ProjectFiles/Rider/RiderProjectFileGenerator.cs b/Engine/Source/Programs/UnrealBuildTool/ProjectFiles/Rider/RiderProjectFileGenerator.cs
--- a/Engine/Source/Programs/UnrealBuildTool/ProjectFiles/Rider/RiderProjectFileGenerator.cs
+++ b/Engine/Source/Programs/UnrealBuildTool/ProjectFiles/Rider/RiderProjectFileGenerator.cs
@@ -525,7 +525,7 @@ namespace UnrealBuildTool
 				HashSet<UnrealTargetPlatform> ValidPlatforms;
 				List<VCProjectFileGenerator.VCSolutionConfigCombination> ValidSolutionConfigs;
 				VCProjectFileGenerator.CollectSolutionConfigurations(SupportedConfigurations,
-					SupportedPlatforms, AllProjectFiles, bMakeProjectPerTarget, Logger, out ValidPlatforms, out ValidSolutionConfigs);
+					SupportedPlatforms, AllProjectFiles, bMakeProjectPerTarget, out ValidPlatforms, out ValidSolutionConfigs);
 				FileReference OutputFile = GetRiderProjectLocation("root.json");
 				DirectoryReference.CreateDirectory(OutputFile.Directory);
 				using (JsonWriter Writer = new JsonWriter(OutputFile))
diff --git a/Engine/Source/Programs/UnrealBuildTool/ProjectFiles/VisualStudio/VCProject.cs b/Engine/Source/Programs/UnrealBuildTool/ProjectFiles/VisualStudio/VCProject.cs
--- a/Engine/Source/Programs/UnrealBuildTool/ProjectFiles/VisualStudio/VCProject.cs
+++ b/Engine/Source/Programs/UnrealBuildTool/ProjectFiles/VisualStudio/VCProject.cs
@@ -190,9 +190,8 @@ namespace UnrealBuildTool
 		/// <param name="ProjectTarget">The target that we're checking for a valid platform/config combination</param>
 		/// <param name="Platform">Platform</param>
 		/// <param name="Configuration">Configuration</param>
-		/// <param name="Logger">Logger for output</param>
 		/// <returns>True if this is a valid combination for this project, otherwise false</returns>
-		public static bool IsValidProjectPlatformAndConfiguration(Project ProjectTarget, UnrealTargetPlatform Platform, UnrealTargetConfiguration Configuration, ILogger Logger)
+		public static bool IsValidProjectPlatformAndConfiguration(Project ProjectTarget, UnrealTargetPlatform Platform, UnrealTargetConfiguration Configuration)
 		{
 			if (!ProjectFileGenerator.bIncludeTestAndShippingConfigs)
 			{
@@ -367,7 +366,7 @@ namespace UnrealBuildTool
 				List<Project> MatchingProjectTargets = new List<Project>();
 				foreach (Project ProjectTarget in ProjectTargets)
 				{
-					if (VCProjectFile.IsValidProjectPlatformAndConfiguration(ProjectTarget, SolutionPlatform, SolutionConfiguration, Logger))
+					if (VCProjectFile.IsValidProjectPlatformAndConfiguration(ProjectTarget, SolutionPlatform, SolutionConfiguration))
 					{
 						if (ProjectTarget.TargetRules != null)
 						{
@@ -663,7 +662,7 @@ namespace UnrealBuildTool
 		HashSet<string>? InvalidConfigPlatformNames;
 		List<ProjectConfigAndTargetCombination>? ProjectConfigAndTargetCombinations;
 
-		private void BuildProjectConfigAndTargetCombinations(List<UnrealTargetPlatform> InPlatforms, List<UnrealTargetConfiguration> InConfigurations, PlatformProjectGeneratorCollection PlatformProjectGenerators, ILogger Logger)
+		private void BuildProjectConfigAndTargetCombinations(List<UnrealTargetPlatform> InPlatforms, List<UnrealTargetConfiguration> InConfigurations, PlatformProjectGeneratorCollection PlatformProjectGenerators)
 		{
 			//no need to do this more than once
 			if (ProjectConfigAndTargetCombinations == null)
@@ -700,7 +699,7 @@ namespace UnrealBuildTool
 							{
 								foreach (ProjectTarget ProjectTarget in ProjectTargets.OfType<ProjectTarget>())
 								{
-									if (!IsValidProjectPlatformAndConfiguration(ProjectTarget, Platform, Configuration, Logger))
+									if (!IsValidProjectPlatformAndConfiguration(ProjectTarget, Platform, Configuration))
 									{
 										continue;
 									}
@@ -804,7 +803,7 @@ namespace UnrealBuildTool
 			List<UnrealTargetPlatform> ProjectPlatforms = new List<UnrealTargetPlatform>();
 			List<Tuple<ProjectFile, string>> ProjectFiles = new List<Tuple<ProjectFile, string>>();
 
-			BuildProjectConfigAndTargetCombinations(InPlatforms, InConfigurations, PlatformProjectGenerators, Logger);
+			BuildProjectConfigAndTargetCombinations(InPlatforms, InConfigurations, PlatformProjectGenerators);
 
 			foreach (ProjectConfigAndTargetCombination Combination in ProjectConfigAndTargetCombinations!)
 			{
@@ -948,14 +947,14 @@ namespace UnrealBuildTool
 			VCProjectFileContent.AppendLine("<?xml version=\"1.0\" encoding=\"utf-8\"?>");
 			VCProjectFileContent.AppendLine("<Project DefaultTargets=\"Build\" ToolsVersion=\"{0}\" xmlns=\"http://schemas.microsoft.com/developer/msbuild/2003\">", VCProjectFileGenerator.GetProjectFileToolVersionString(ProjectFileFormat));
 
-			bool bGenerateUserFileContent = PlatformProjectGenerators.PlatformRequiresVSUserFileGeneration(InPlatforms, InConfigurations);
+			bool bGenerateUserFileContent = PlatformProjectGenerators.PlatformRequiresVSUserFileGeneration(InPlatforms);
 			if (bGenerateUserFileContent)
 			{
 				VCUserFileContent.AppendLine("<?xml version=\"1.0\" encoding=\"utf-8\"?>");
 				VCUserFileContent.AppendLine("<Project ToolsVersion=\"{0}\" xmlns=\"http://schemas.microsoft.com/developer/msbuild/2003\">", VCProjectFileGenerator.GetProjectFileToolVersionString(ProjectFileFormat));
 			}
 
-			BuildProjectConfigAndTargetCombinations(InPlatforms, InConfigurations, PlatformProjectGenerators, Logger);
+			BuildProjectConfigAndTargetCombinations(InPlatforms, InConfigurations, PlatformProjectGenerators);
 
 			VCProjectFileContent.AppendLine("  <ItemGroup Label=\"ProjectConfigurations\">");
 
@@ -971,7 +970,7 @@ namespace UnrealBuildTool
 				{
 					continue;
 				}
-				if (Combination.ProjectTarget != null && !IsValidProjectPlatformAndConfiguration(Combination.ProjectTarget, (UnrealTargetPlatform)Combination.Platform, Combination.Configuration, Logger))
+				if (Combination.ProjectTarget != null && !IsValidProjectPlatformAndConfiguration(Combination.ProjectTarget, (UnrealTargetPlatform)Combination.Platform, Combination.Configuration))
 				{
 					continue;
 				}
@@ -1799,7 +1798,7 @@ namespace UnrealBuildTool
 			VCProjectFileContent.AppendLine("  <ImportGroup Label=\"ExtensionTargets\">");
 			VCProjectFileContent.AppendLine("  </ImportGroup>");
 
-			WriteTargets(ProjectPlatforms, PlatformProjectGenerators, VCProjectFileContent, Logger);
+			WriteTargets(ProjectPlatforms, PlatformProjectGenerators, VCProjectFileContent);
 
 			VCProjectFileContent.AppendLine("</Project>");
 
@@ -2022,8 +2021,7 @@ namespace UnrealBuildTool
 		/// <param name="ProjectPlatforms"></param>
 		/// <param name="PlatformProjectGenerators"></param>
 		/// <param name="VCProjectFileContent"></param>
-		/// <param name="Logger"></param>
-		private void WriteTargets(List<UnrealTargetPlatform> ProjectPlatforms, PlatformProjectGeneratorCollection PlatformProjectGenerators, StringBuilder VCProjectFileContent, ILogger Logger)
+		private void WriteTargets(List<UnrealTargetPlatform> ProjectPlatforms, PlatformProjectGeneratorCollection PlatformProjectGenerators, StringBuilder VCProjectFileContent)
 		{
 			foreach (UnrealTargetPlatform Platform in ProjectPlatforms)
 			{
diff --git a/Engine/Source/Programs/UnrealBuildTool/ProjectFiles/VisualStudio/VCProjectFileGenerator.cs b/Engine/Source/Programs/UnrealBuildTool/ProjectFiles/VisualStudio/VCProjectFileGenerator.cs
--- a/Engine/Source/Programs/UnrealBuildTool/ProjectFiles/VisualStudio/VCProjectFileGenerator.cs
+++ b/Engine/Source/Programs/UnrealBuildTool/ProjectFiles/VisualStudio/VCProjectFileGenerator.cs
@@ -261,7 +261,7 @@ namespace UnrealBuildTool
 		}
 
 		// parses project ini for Android to get architecture(s) enabled
-		public static UnrealArchitectures GetAndroidProjectArchitectures(FileReference? ProjectFile, bool bGetAllSupported)
+		public static UnrealArchitectures GetAndroidProjectArchitectures(FileReference? ProjectFile)
 		{
 			List<string> ActiveArches = new();
 
@@ -298,7 +298,7 @@ namespace UnrealBuildTool
 				if (BuildPlatform.Platform == UnrealTargetPlatform.Android)
 				{
 					return InProjectTarget.UnrealProjectFilePath == null ? BuildPlatform.ArchitectureConfig.AllSupportedArchitectures
-						: GetAndroidProjectArchitectures(InProjectTarget.UnrealProjectFilePath, true);
+						: GetAndroidProjectArchitectures(InProjectTarget.UnrealProjectFilePath);
 				}
 				return BuildPlatform.ArchitectureConfig.AllSupportedArchitectures;
 			}
@@ -844,7 +844,7 @@ namespace UnrealBuildTool
 						VCSolutionFileContent.AppendLine("	GlobalSection(SolutionConfigurationPlatforms) = preSolution");
 
 						CollectSolutionConfigurations(SupportedConfigurations, SupportedPlatforms, AllProjectFiles, bMakeProjectPerTarget,
-							Logger, out PlatformsValidForProjects, out SolutionConfigCombinations);
+							out PlatformsValidForProjects, out SolutionConfigCombinations);
 
 						HashSet<string> AppendedSolutionConfigAndPlatformNames = new HashSet<string>(StringComparer.InvariantCultureIgnoreCase);
 						foreach (VCSolutionConfigCombination SolutionConfigCombination in SolutionConfigCombinations)
@@ -1040,7 +1040,7 @@ namespace UnrealBuildTool
 		}
 
 		public static void CollectSolutionConfigurations(List<UnrealTargetConfiguration> AllConfigurations,
-			List<UnrealTargetPlatform> AllPlatforms, List<ProjectFile> AllProjectFiles, bool bMakeProjectPerTarget, ILogger Logger,
+			List<UnrealTargetPlatform> AllPlatforms, List<ProjectFile> AllProjectFiles, bool bMakeProjectPerTarget,
 			out HashSet<UnrealTargetPlatform> OutValidPlatforms, out List<VCSolutionConfigCombination> OutSolutionConfigs)
 		{
 			OutValidPlatforms = new HashSet<UnrealTargetPlatform>();
@@ -1062,8 +1062,7 @@ namespace UnrealBuildTool
 									// Figure out the set of valid target configuration names
 									foreach (ProjectTarget ProjectTarget in CurProject.ProjectTargets.OfType<ProjectTarget>())
 									{
-										if (VCProjectFile.IsValidProjectPlatformAndConfiguration(ProjectTarget, CurPlatform,
-												CurConfiguration, Logger))
+										if (VCProjectFile.IsValidProjectPlatformAndConfiguration(ProjectTarget, CurPlatform, CurConfiguration))
 										{
 											OutValidPlatforms.Add(CurPlatform);
 
diff --git a/Engine/Source/Programs/UnrealBuildTool/ProjectFiles/VisualStudioCode/VSCodeProjectFileGenerator.cs b/Engine/Source/Programs/UnrealBuildTool/ProjectFiles/VisualStudioCode/VSCodeProjectFileGenerator.cs
--- a/Engine/Source/Programs/UnrealBuildTool/ProjectFiles/VisualStudioCode/VSCodeProjectFileGenerator.cs
+++ b/Engine/Source/Programs/UnrealBuildTool/ProjectFiles/VisualStudioCode/VSCodeProjectFileGenerator.cs
@@ -292,11 +292,11 @@ namespace UnrealBuildTool
 			}
 			Projects.Sort((A, B) => { return A.ProjectFilePath.GetFileName().CompareTo(B.ProjectFilePath.GetFileName()); });
 
-			ProjectData ProjectData = GatherProjectData(Projects, Logger);
+			ProjectData ProjectData = GatherProjectData(Projects);
 
 			WriteWorkspaceIgnoreFile(Projects);
 			WriteCppPropertiesFile(VSCodeDir, ProjectData);
-			WriteWorkspaceFile(ProjectData, Logger);
+			WriteWorkspaceFile(ProjectData);
 
 			if (bForeignProject && bIncludeEngineSource)
 			{
@@ -495,7 +495,7 @@ namespace UnrealBuildTool
 			public List<Project> AllProjects = new List<Project>();
 		}
 
-		private ProjectData GatherProjectData(List<ProjectFile> InProjects, ILogger Logger)
+		private ProjectData GatherProjectData(List<ProjectFile> InProjects)
 		{
 			ProjectData ProjectData = new ProjectData();
 
@@ -528,7 +528,7 @@ namespace UnrealBuildTool
 							{
 								foreach (UnrealTargetConfiguration Config in Configs)
 								{
-									if (MSBuildProjectFile.IsValidProjectPlatformAndConfiguration(Target, Platform, Config, Logger))
+									if (MSBuildProjectFile.IsValidProjectPlatformAndConfiguration(Target, Platform, Config))
 									{
 										NewTarget.BuildProducts.Add(new ProjectData.BuildProduct(GetExecutableFilename(Project, Target, Platform, Config))
 										{
@@ -759,7 +759,7 @@ namespace UnrealBuildTool
 			OutFile.EndObject();
 		}
 
-		private void WriteNativeTaskDeployAndroid(ProjectData.Project InProject, JsonFile OutFile, ProjectData.Target Target, ProjectData.BuildProduct BuildProduct)
+		private void WriteNativeTaskDeployAndroid(JsonFile OutFile, ProjectData.Target Target, ProjectData.BuildProduct BuildProduct)
 		{
 			if (BuildProduct.UProjectFile == null)
 			{
@@ -1005,7 +1005,7 @@ namespace UnrealBuildTool
 
 						if (BuildProduct.Platform == UnrealTargetPlatform.Android && BaseCommand.Equals("Build", StringComparison.Ordinal))
 						{
-							WriteNativeTaskDeployAndroid(InProject, OutFile, Target, BuildProduct);
+							WriteNativeTaskDeployAndroid(OutFile, Target, BuildProduct);
 						}
 					}
 				}
@@ -1211,7 +1211,7 @@ namespace UnrealBuildTool
 			return new FileReference(ExecutableFilename);
 		}
 
-		private void WriteNativeLaunchConfigAndroidOculus(ProjectData.Project InProject, JsonFile OutFile, ProjectData.Target Target, ProjectData.BuildProduct BuildProduct, ILogger Logger)
+		private void WriteNativeLaunchConfigAndroidOculus(JsonFile OutFile, ProjectData.Target Target, ProjectData.BuildProduct BuildProduct)
 		{
 			ConfigHierarchy Ini = ConfigCache.ReadHierarchy(ConfigHierarchyType.Engine, DirectoryReference.FromFile(BuildProduct.UProjectFile), BuildProduct.Platform);
 
@@ -1439,7 +1439,7 @@ namespace UnrealBuildTool
 			}
 		}
 
-		private void WriteNativeLaunchConfig(ProjectData.Project InProject, JsonFile OutFile, ILogger Logger)
+		private void WriteNativeLaunchConfig(ProjectData.Project InProject, JsonFile OutFile)
 		{
 			foreach (ProjectData.Target Target in InProject.Targets)
 			{
@@ -1451,7 +1451,7 @@ namespace UnrealBuildTool
 					}
 					else if (BuildProduct.Platform == UnrealTargetPlatform.Android)
 					{
-						WriteNativeLaunchConfigAndroidOculus(InProject, OutFile, Target, BuildProduct, Logger);
+						WriteNativeLaunchConfigAndroidOculus(OutFile, Target, BuildProduct);
 					}
 				}
 			}
@@ -1533,7 +1533,7 @@ namespace UnrealBuildTool
 			}
 		}
 
-		private void WriteLaunch(JsonFile OutFile, ProjectData ProjectData, ILogger Logger)
+		private void WriteLaunch(JsonFile OutFile, ProjectData ProjectData)
 		{
 			OutFile.AddField("version", "0.2.0");
 			if (bAddDebugCoreConfig)
@@ -1561,7 +1561,7 @@ namespace UnrealBuildTool
 				{
 					foreach (ProjectData.Project Project in ProjectData.NativeProjects)
 					{
-						WriteNativeLaunchConfig(Project, OutFile, Logger);
+						WriteNativeLaunchConfig(Project, OutFile);
 					}
 				}
 
@@ -1675,7 +1675,7 @@ namespace UnrealBuildTool
 			}
 		}
 
-		private void WriteWorkspaceFile(ProjectData ProjectData, ILogger Logger)
+		private void WriteWorkspaceFile(ProjectData ProjectData)
 		{
 			JsonFile WorkspaceFile = new JsonFile();
 
@@ -1779,7 +1779,7 @@ namespace UnrealBuildTool
 			WorkspaceFile.EndObject();
 
 			WorkspaceFile.BeginObject("launch");
-			WriteLaunch(WorkspaceFile, ProjectData, Logger);
+			WriteLaunch(WorkspaceFile, ProjectData);
 			WorkspaceFile.EndObject();
 
 			WorkspaceFile.EndRootObject();
diff --git a/Engine/Source/Programs/UnrealBuildTool/ProjectFiles/VisualStudioWorkspace/VSWorkspaceProjectFile.cs b/Engine/Source/Programs/UnrealBuildTool/ProjectFiles/VisualStudioWorkspace/VSWorkspaceProjectFile.cs
--- a/Engine/Source/Programs/UnrealBuildTool/ProjectFiles/VisualStudioWorkspace/VSWorkspaceProjectFile.cs
+++ b/Engine/Source/Programs/UnrealBuildTool/ProjectFiles/VisualStudioWorkspace/VSWorkspaceProjectFile.cs
@@ -44,20 +44,13 @@ namespace UnrealBuildTool
 		/// <param name="availableConfigurations"></param>
 		/// <param name="InPlatforms"></param>
 		/// <param name="InConfigurations"></param>
-		/// <param name="PlatformProjectGenerators"></param>
-		/// <param name="Minimize"></param>
-		/// <param name="Logger"></param>
 		public void CollectAvailableConfigurations(Dictionary<string, AvailableConfiguration> availableConfigurations,
 			List<UnrealTargetPlatform> InPlatforms,
-			List<UnrealTargetConfiguration> InConfigurations,
-			PlatformProjectGeneratorCollection PlatformProjectGenerators,
-			JsonWriterStyle Minimize,
-			ILogger Logger)
+			List<UnrealTargetConfiguration> InConfigurations)
 		{
 			ProcessProjectFiles(
 				InPlatforms,
 				InConfigurations,
-				PlatformProjectGenerators,
 				// always provide null in here as we want to collect all available data.
 				null,
 				(UnrealTargetConfiguration Configuration, UnrealTargetPlatform Platform, ProjectTarget ProjectTarget) =>
@@ -107,7 +100,6 @@ namespace UnrealBuildTool
 			ProcessProjectFiles(
 				InPlatforms,
 				InConfigurations,
-				PlatformProjectGenerators,
 				ProjectGenerationDescriptionData,
 				(UnrealTargetConfiguration Configuration, UnrealTargetPlatform Platform, ProjectTarget ProjectTarget) => Task.Run(async () =>
 				{
@@ -180,12 +172,10 @@ namespace UnrealBuildTool
 		/// </summary>
 		/// <param name="InPlatforms"></param>
 		/// <param name="InConfigurations"></param>
-		/// <param name="PlatformProjectGenerators"></param>
 		/// <param name="ProjectGenerationDescriptionData"></param>
 		/// <param name="Handler"></param>
 		private void ProcessProjectFiles(List<UnrealTargetPlatform> InPlatforms,
 			List<UnrealTargetConfiguration> InConfigurations,
-			PlatformProjectGeneratorCollection PlatformProjectGenerators,
 			Dictionary<string, AvailableConfiguration>? ProjectGenerationDescriptionData,
 			Func<UnrealTargetConfiguration, UnrealTargetPlatform, ProjectTarget, Task> Handler)
 		{
diff --git a/Engine/Source/Programs/UnrealBuildTool/ProjectFiles/VisualStudioWorkspace/VSWorkspaceProjectFileGenerator.cs b/Engine/Source/Programs/UnrealBuildTool/ProjectFiles/VisualStudioWorkspace/VSWorkspaceProjectFileGenerator.cs
--- a/Engine/Source/Programs/UnrealBuildTool/ProjectFiles/VisualStudioWorkspace/VSWorkspaceProjectFileGenerator.cs
+++ b/Engine/Source/Programs/UnrealBuildTool/ProjectFiles/VisualStudioWorkspace/VSWorkspaceProjectFileGenerator.cs
@@ -175,7 +175,7 @@ namespace UnrealBuildTool
 			return true;
 		}
 
-		private bool WriteAvailableConfigurations(PlatformProjectGeneratorCollection PlatformProjectGenerators, ILogger Logger)
+		private bool WriteAvailableConfigurations()
 		{
 			Dictionary<string, AvailableConfiguration> AvailableConfigurations = new Dictionary<string, AvailableConfiguration>();
 			HashSet<UnrealTargetPlatform> PlatformsToGenerate = new(SupportedPlatforms);
@@ -192,7 +192,7 @@ namespace UnrealBuildTool
 					return false;
 				}
 
-				CurrentProject.CollectAvailableConfigurations(AvailableConfigurations, FilteredPlatforms, new HashSet<UnrealTargetConfiguration>(SupportedConfigurations).ToList(), PlatformProjectGenerators, Minimize, Logger);
+				CurrentProject.CollectAvailableConfigurations(AvailableConfigurations, FilteredPlatforms, new HashSet<UnrealTargetConfiguration>(SupportedConfigurations).ToList());
 			}
 
 			FileReference OutputFile = FileReference.Combine(IntermediateProjectFilesPath, ProjectFilesFolder, Query);
@@ -280,7 +280,7 @@ namespace UnrealBuildTool
 
 			if (!String.IsNullOrEmpty(Query))
 			{
-				return WriteAvailableConfigurations(PlatformProjectGenerators, Logger);
+				return WriteAvailableConfigurations();
 			}
 
 			WriteProjectFiles(PlatformProjectGenerators, Logger);
diff --git a/Engine/Source/Programs/UnrealBuildTool/ProjectFiles/Xcode/XcodeProject.Utils.cs b/Engine/Source/Programs/UnrealBuildTool/ProjectFiles/Xcode/XcodeProject.Utils.cs
--- a/Engine/Source/Programs/UnrealBuildTool/ProjectFiles/Xcode/XcodeProject.Utils.cs
+++ b/Engine/Source/Programs/UnrealBuildTool/ProjectFiles/Xcode/XcodeProject.Utils.cs
@@ -51,7 +51,7 @@ namespace UnrealBuildTool.XcodeProjectXcconfig
 			};
 		}
 
-		public static bool ShouldIncludeProjectInWorkspace(ProjectFile Proj, ILogger Logger)
+		public static bool ShouldIncludeProjectInWorkspace(ProjectFile Proj)
 		{
 			// since IOS/TVOS don't have the UnrealEditor project as valid, force it so that we get the source code
 			// this is likely temporary until we can put source code into UnrealGame
@@ -65,7 +65,7 @@ namespace UnrealBuildTool.XcodeProjectXcconfig
 				{
 					foreach (UnrealTargetConfiguration Config in GetSupportedConfigurations())
 					{
-						if (MSBuildProjectFile.IsValidProjectPlatformAndConfiguration(ProjectTarget, Platform, Config, Logger))
+						if (MSBuildProjectFile.IsValidProjectPlatformAndConfiguration(ProjectTarget, Platform, Config))
 						{
 							return true;
 						}
diff --git a/Engine/Source/Programs/UnrealBuildTool/ProjectFiles/Xcode/XcodeProject.cs b/Engine/Source/Programs/UnrealBuildTool/ProjectFiles/Xcode/XcodeProject.cs
--- a/Engine/Source/Programs/UnrealBuildTool/ProjectFiles/Xcode/XcodeProject.cs
+++ b/Engine/Source/Programs/UnrealBuildTool/ProjectFiles/Xcode/XcodeProject.cs
@@ -106,7 +106,7 @@ namespace UnrealBuildTool.XcodeProjectXcconfig
 			ResponseFile = FileReference.Combine(UnrealData.XcodeProjectFileLocation.ParentDirectory!, "ResponseFiles", Module.Rules.Target.Platform.ToString(), $"{UnrealData.ProductName}{Index}.response");
 		}
 
-		public void GenerateResponseFile(UnrealTargetPlatform Platform)
+		public void GenerateResponseFile()
 		{
 			StringBuilder ResponseFileContents = new();
 			ResponseFileContents.Append("-isystem");
@@ -228,7 +228,7 @@ namespace UnrealBuildTool.XcodeProjectXcconfig
 		public Dictionary<MetadataPlatform, MetadataItem> ShippingEntitlementsFiles = new();
 		public Dictionary<MetadataPlatform, MetadataItem> ProjectPrivacyInfoFiles = new();
 
-		public Metadata(DirectoryReference ProductDirectory, DirectoryReference XcodeProject, ConfigHierarchy Ini, bool bSupportsMac, bool bSupportsIOSOrTVOS, ILogger Logger)
+		public Metadata(DirectoryReference ProductDirectory, DirectoryReference XcodeProject, ConfigHierarchy Ini, bool bSupportsMac, bool bSupportsIOSOrTVOS)
 		{
 			if (bSupportsMac)
 			{
@@ -532,10 +532,10 @@ namespace UnrealBuildTool.XcodeProjectXcconfig
 				BundleIdentifier = "$(UE_SIGNING_PREFIX).$(UE_PRODUCT_NAME_STRIPPED)";
 			}
 
-			InitializeMetadata(Logger);
+			InitializeMetadata();
 
 			// Figure out all the desired configurations on the unreal side
-			AllConfigs = GetSupportedBuildConfigs(XcodeProjectFileGenerator.XcodePlatforms, Configurations, Logger);
+			AllConfigs = GetSupportedBuildConfigs(XcodeProjectFileGenerator.XcodePlatforms, Configurations);
 			// if we can't find any configs, we will fail to create a project
 			if (AllConfigs.Count == 0)
 			{
@@ -590,14 +590,14 @@ namespace UnrealBuildTool.XcodeProjectXcconfig
 			return true;
 		}
 
-		private void InitializeMetadata(ILogger Logger)
+		private void InitializeMetadata()
 		{
 
 			// read setings from the configs, now that we have a project
 			ConfigHierarchy SharedPlatformIni = ConfigCache.ReadHierarchy(ConfigHierarchyType.Engine, ConfigDirectory, UnrealTargetPlatform.Mac);
 			SharedPlatformIni.TryGetValue("/Script/MacTargetPlatform.XcodeProjectSettings", "bUseAutomaticCodeSigning", out bUseAutomaticSigning);
 
-			Metadata = new Metadata(ProductDirectory, XcodeProjectFileLocation, SharedPlatformIni, bSupportsMac, bSupportsIOS || bSupportsTVOS || bSupportsVisionOS, Logger);
+			Metadata = new Metadata(ProductDirectory, XcodeProjectFileLocation, SharedPlatformIni, bSupportsMac, bSupportsIOS || bSupportsTVOS || bSupportsVisionOS);
 		}
 
 		public string? FindFile(List<string> Paths, UnrealTargetPlatform Platform, bool bMakeRelative)
@@ -685,7 +685,7 @@ namespace UnrealBuildTool.XcodeProjectXcconfig
 			FileBatch.UserIncludePaths.UnionWith(CompileEnvironment.UserIncludePaths);
 		}
 
-		private List<UnrealBuildConfig> GetSupportedBuildConfigs(List<UnrealTargetPlatform> Platforms, List<UnrealTargetConfiguration> Configurations, ILogger Logger)
+		private List<UnrealBuildConfig> GetSupportedBuildConfigs(List<UnrealTargetPlatform> Platforms, List<UnrealTargetConfiguration> Configurations)
 		{
 			List<UnrealBuildConfig> BuildConfigs = new List<UnrealBuildConfig>();
 
@@ -711,7 +711,7 @@ namespace UnrealBuildTool.XcodeProjectXcconfig
 								// Now go through all of the target types for this project
 								foreach (ProjectTarget ProjectTarget in ProjectFile.ProjectTargets.OfType<ProjectTarget>())
 								{
-									if (MSBuildProjectFile.IsValidProjectPlatformAndConfiguration(ProjectTarget, Platform, Configuration, Logger))
+									if (MSBuildProjectFile.IsValidProjectPlatformAndConfiguration(ProjectTarget, Platform, Configuration))
 									{
 										// Figure out if this is a monolithic build
 										bool bShouldCompileMonolithic = BuildPlatform.ShouldCompileMonolithicBinary(Platform);
@@ -1410,7 +1410,7 @@ namespace UnrealBuildTool.XcodeProjectXcconfig
 				References.Add(ResourcesBuildPhase);
 
 				ProcessFrameworks(ResourcesBuildPhase, Project, ProjectFile, Logger);
-				ProcessScripts(ResourcesBuildPhase, Project, ProjectFile);
+				ProcessScripts(Project, ProjectFile);
 				ProcessAssets(ResourcesBuildPhase);
 			}
 
@@ -1481,7 +1481,7 @@ namespace UnrealBuildTool.XcodeProjectXcconfig
 							References.Add(EmbedAppExt);
 							
 							// Also need a (Native)Run Target for this extension
-							XcodeAppExtRunTarget RunTarget = new XcodeAppExtRunTarget(Project, ExtConfig, DI.Name, ExtensionGuid, UnrealData.AllConfigs[0].ProjectTarget!.TargetRules!.Type, UnrealTargetPlatform.IOS, BuildTarget, ProjectFile, DI, Logger);
+							XcodeAppExtRunTarget RunTarget = new XcodeAppExtRunTarget(Project, ExtConfig, DI.Name, ExtensionGuid, UnrealTargetPlatform.IOS, BuildTarget, ProjectFile, DI, Logger);
 							Project.AppExtRunTargets.Add(RunTarget);
 							Project.References.Add(RunTarget);
 						}
@@ -1519,10 +1519,9 @@ namespace UnrealBuildTool.XcodeProjectXcconfig
 		///   - So here we gneerate a dSYM from the executable, copying it to where Xcode wants it (DWARF_DSYM_FOLDER_PATH/DWARF_DSYM_FILE_NAME), and
 		///       then we strip the executable in place
 		/// </summary>
-		/// <param name="ResourcesBuildPhase"></param>
 		/// <param name="Project"></param>
 		/// <param name="ProjectFile"></param>
-		protected void ProcessScripts(XcodeResourcesBuildPhase ResourcesBuildPhase, XcodeProject Project, XcodeProjectFile ProjectFile)
+		protected void ProcessScripts(XcodeProject Project, XcodeProjectFile ProjectFile)
 		{
 			List<string> CopyScript = new();
 
@@ -2370,7 +2369,7 @@ namespace UnrealBuildTool.XcodeProjectXcconfig
 		private readonly string AppGuid;
 		private readonly ExtensionConfig ExtConfig;
 		
-		public XcodeAppExtRunTarget(XcodeProject Project, ExtensionConfig? ExtConfig, string TargetName, string AppGuid, TargetType TargetType, UnrealTargetPlatform Platform, XcodeBuildTarget? BuildTarget, XcodeProjectFile ProjectFile, DirectoryInfo ExtensionPath, ILogger Logger)
+		public XcodeAppExtRunTarget(XcodeProject Project, ExtensionConfig? ExtConfig, string TargetName, string AppGuid, UnrealTargetPlatform Platform, XcodeBuildTarget? BuildTarget, XcodeProjectFile ProjectFile, DirectoryInfo ExtensionPath, ILogger Logger)
 			: base(XcodeTarget.Type.App_Ext, Project.UnrealData,
 				TargetName + (XcodeProjectFileGenerator.PerPlatformMode == XcodePerPlatformMode.OneWorkspacePerPlatform ? "" : $"_{Platform}"))
 		{
@@ -3169,7 +3168,7 @@ namespace UnrealBuildTool.XcodeProjectXcconfig
 			{
 				foreach(UnrealBatchedFiles Batch in Pair.Value)
 				{
-					Batch.GenerateResponseFile(Pair.Key);
+					Batch.GenerateResponseFile();
 				}
 			}
 		}
@@ -3284,7 +3283,7 @@ namespace UnrealBuildTool.XcodeProjectXcconfig
 					bSuccess = ProjectFileGenerator.WriteFileIfChanged(PBXFilePath.FullName, Content.ToString(), Logger, new UTF8Encoding()) && bSuccess;
 				}
 
-				bool bNeedScheme = !bAddStubEditor && XcodeUtils.ShouldIncludeProjectInWorkspace(this, Logger);
+				bool bNeedScheme = !bAddStubEditor && XcodeUtils.ShouldIncludeProjectInWorkspace(this);
 				if (bNeedScheme)
 				{
 					if (bSuccess)
diff --git a/Engine/Source/Programs/UnrealBuildTool/ProjectFiles/Xcode/XcodeProjectFileGenerator.cs b/Engine/Source/Programs/UnrealBuildTool/ProjectFiles/Xcode/XcodeProjectFileGenerator.cs
--- a/Engine/Source/Programs/UnrealBuildTool/ProjectFiles/Xcode/XcodeProjectFileGenerator.cs
+++ b/Engine/Source/Programs/UnrealBuildTool/ProjectFiles/Xcode/XcodeProjectFileGenerator.cs
@@ -342,7 +342,7 @@ namespace UnrealBuildTool
 							IEnumerable<XcodeProjectXcconfig.XcodeProjectFile> SupportedProjects =
 									CurFolder.ChildProjects.Where(P => P.GetType() == typeof(XcodeProjectXcconfig.XcodeProjectFile))
 										.Select(P => (XcodeProjectXcconfig.XcodeProjectFile)P)
-										.Where(P => XcodeProjectXcconfig.XcodeUtils.ShouldIncludeProjectInWorkspace(P, Logger))
+										.Where(P => XcodeProjectXcconfig.XcodeUtils.ShouldIncludeProjectInWorkspace(P))
 										// @todo - still need to handle legacy project getting split up?
 										.Where(P => P.RootProjects.Count == 0 || P.RootProjects.ContainsValue(Platform))
 										.OrderBy(P => P.ProjectFilePath.GetFileName());
diff --git a/Engine/Source/Programs/UnrealBuildTool/ProjectFiles/Xcode/XcodeProjectLegacy.cs b/Engine/Source/Programs/UnrealBuildTool/ProjectFiles/Xcode/XcodeProjectLegacy.cs
--- a/Engine/Source/Programs/UnrealBuildTool/ProjectFiles/Xcode/XcodeProjectLegacy.cs
+++ b/Engine/Source/Programs/UnrealBuildTool/ProjectFiles/Xcode/XcodeProjectLegacy.cs
@@ -323,7 +323,7 @@ namespace UnrealBuildTool.XcodeProjectLegacy
 		}
 
 		private void GenerateSectionsWithExtensions(StringBuilder PBXBuildFileSection, StringBuilder PBXFileReferenceSection, StringBuilder PBXCopyFilesBuildPhaseSection, StringBuilder PBXResourcesBuildPhaseSection,
-													List<XcodeExtensionInfo> AllExtensions, FileReference? UProjectPath, List<XcodeBuildConfig> BuildConfigs, ILogger Logger)
+													List<XcodeExtensionInfo> AllExtensions, FileReference? UProjectPath, List<XcodeBuildConfig> BuildConfigs)
 		{
 			if (UProjectPath != null)
 			{
@@ -427,7 +427,7 @@ namespace UnrealBuildTool.XcodeProjectLegacy
 							}
 
 							string? IOSRuntimeVersion, TVOSRuntimeVersion;
-							AppendPlatformConfiguration(ConfigSection, Configuration, ExtensionInfo.Name, UProjectPath, false, bSupportIOS, bSupportTVOS, Logger, out IOSRuntimeVersion, out TVOSRuntimeVersion);
+							AppendPlatformConfiguration(ConfigSection, Configuration, ExtensionInfo.Name, UProjectPath, false, bSupportIOS, bSupportTVOS, out IOSRuntimeVersion, out TVOSRuntimeVersion);
 
 							ConfigSection.Append("\t\t\t};" + ProjectFileGenerator.NewLine);
 							ConfigSection.Append("\t\t\tname = \"" + ConfigName + "\";" + ProjectFileGenerator.NewLine);
@@ -963,7 +963,7 @@ namespace UnrealBuildTool.XcodeProjectLegacy
 			return CachedMacProjectArcitectures[TargetName];
 		}
 
-		private void AppendPlatformConfiguration(StringBuilder Content, XcodeBuildConfig Config, string TargetName, FileReference? ProjectFile, bool bSupportMac, bool bSupportIOS, bool bSupportTVOS, ILogger Logger, out string? IOSRunTimeVersion, out string? TVOSRunTimeVersion, string BinariesSubDir = "/Payload")
+		private void AppendPlatformConfiguration(StringBuilder Content, XcodeBuildConfig Config, string TargetName, FileReference? ProjectFile, bool bSupportMac, bool bSupportIOS, bool bSupportTVOS, out string? IOSRunTimeVersion, out string? TVOSRunTimeVersion, string BinariesSubDir = "/Payload")
 		{
 			FileReference MacExecutablePath = Config.MacExecutablePath!;
 
@@ -1127,7 +1127,7 @@ namespace UnrealBuildTool.XcodeProjectLegacy
 			string MacExecutableFileName = Config.MacExecutablePath!.GetFileName();
 
 			string? IOSRunTimeVersion, TVOSRunTimeVersion;
-			AppendPlatformConfiguration(Content, Config, Config.BuildTarget, ProjectFile, true, !bMacOnly, !bMacOnly, Logger, out IOSRunTimeVersion, out TVOSRunTimeVersion);
+			AppendPlatformConfiguration(Content, Config, Config.BuildTarget, ProjectFile, true, !bMacOnly, !bMacOnly, out IOSRunTimeVersion, out TVOSRunTimeVersion);
 
 			bool bIsUnrealGame = Config.BuildTarget.Equals("UnrealGame", StringComparison.OrdinalIgnoreCase);
 			bool bIsUnrealClient = Config.BuildTarget.Equals("UnrealClient", StringComparison.OrdinalIgnoreCase);
@@ -1295,7 +1295,7 @@ namespace UnrealBuildTool.XcodeProjectLegacy
 			Content.Append("\t\t};" + ProjectFileGenerator.NewLine);
 		}
 
-		private void AppendLegacyTargetBuildConfiguration(StringBuilder Content, XcodeBuildConfig Config, string ConfigGuid, FileReference? ProjectFile)
+		private void AppendLegacyTargetBuildConfiguration(StringBuilder Content, XcodeBuildConfig Config, string ConfigGuid)
 		{
 			bool bMacOnly = true;
 			if (Config.ProjectTarget!.TargetRules != null)
@@ -1366,7 +1366,7 @@ namespace UnrealBuildTool.XcodeProjectLegacy
 
 			foreach (KeyValuePair<string, XcodeBuildConfig> Config in BuildTargetBuildConfigs)
 			{
-				AppendLegacyTargetBuildConfiguration(Content, Config.Value, Config.Key, GameProjectPath);
+				AppendLegacyTargetBuildConfiguration(Content, Config.Value, Config.Key);
 			}
 
 			foreach (KeyValuePair<string, XcodeBuildConfig> Config in IndexTargetBuildConfigs)
@@ -1418,7 +1418,7 @@ namespace UnrealBuildTool.XcodeProjectLegacy
 			Content.Append("/* End XCConfigurationList section */" + ProjectFileGenerator.NewLine);
 		}
 
-		private List<XcodeBuildConfig> GetSupportedBuildConfigs(List<UnrealTargetPlatform> Platforms, List<UnrealTargetConfiguration> Configurations, ILogger Logger)
+		private List<XcodeBuildConfig> GetSupportedBuildConfigs(List<UnrealTargetPlatform> Platforms, List<UnrealTargetConfiguration> Configurations)
 		{
 			List<XcodeBuildConfig> BuildConfigs = new List<XcodeBuildConfig>();
 
@@ -1444,7 +1444,7 @@ namespace UnrealBuildTool.XcodeProjectLegacy
 								// Now go through all of the target types for this project
 								foreach (ProjectTarget ProjectTarget in ProjectTargets.OfType<ProjectTarget>())
 								{
-									if (MSBuildProjectFile.IsValidProjectPlatformAndConfiguration(ProjectTarget, Platform, Configuration, Logger))
+									if (MSBuildProjectFile.IsValidProjectPlatformAndConfiguration(ProjectTarget, Platform, Configuration))
 									{
 										// Figure out if this is a monolithic build
 										bool bShouldCompileMonolithic = BuildPlatform.ShouldCompileMonolithicBinary(Platform);
@@ -1770,12 +1770,12 @@ namespace UnrealBuildTool.XcodeProjectLegacy
 			};
 		}
 
-		public bool ShouldIncludeProjectInWorkspace(ILogger Logger)
+		public bool ShouldIncludeProjectInWorkspace()
 		{
-			return CanBuildProjectLocally(Logger);
+			return CanBuildProjectLocally();
 		}
 
-		public bool CanBuildProjectLocally(ILogger Logger)
+		public bool CanBuildProjectLocally()
 		{
 			foreach (Project ProjectTarget in ProjectTargets)
 			{
@@ -1783,7 +1783,7 @@ namespace UnrealBuildTool.XcodeProjectLegacy
 				{
 					foreach (UnrealTargetConfiguration Config in GetSupportedConfigurations())
 					{
-						if (MSBuildProjectFile.IsValidProjectPlatformAndConfiguration(ProjectTarget, Platform, Config, Logger))
+						if (MSBuildProjectFile.IsValidProjectPlatformAndConfiguration(ProjectTarget, Platform, Config))
 						{
 							return true;
 						}
@@ -1820,7 +1820,7 @@ namespace UnrealBuildTool.XcodeProjectLegacy
 			string ShellScriptSectionGuid = XcodeProjectFileGenerator.MakeXcodeGuid();
 
 			// Figure out all the desired configurations
-			List<XcodeBuildConfig> BuildConfigs = GetSupportedBuildConfigs(InPlatforms, InConfigurations, Logger);
+			List<XcodeBuildConfig> BuildConfigs = GetSupportedBuildConfigs(InPlatforms, InConfigurations);
 			if (BuildConfigs.Count == 0)
 			{
 				return true;
@@ -1870,7 +1870,7 @@ namespace UnrealBuildTool.XcodeProjectLegacy
 			StringBuilder PBXResourcesBuildPhaseSection = new StringBuilder();
 			List<XcodeExtensionInfo> AllExtensions = new List<XcodeExtensionInfo>();
 			GenerateSectionsWithSourceFiles(PBXBuildFileSection, PBXFileReferenceSection, PBXSourcesBuildPhaseSection, TargetAppGuid, TargetName, bIsAppBundle);
-			GenerateSectionsWithExtensions(PBXBuildFileSection, PBXFileReferenceSection, PBXCopyExtensionsBuildPhaseSection, PBXResourcesBuildPhaseSection, AllExtensions, GameProjectPath, BuildConfigs, Logger);
+			GenerateSectionsWithExtensions(PBXBuildFileSection, PBXFileReferenceSection, PBXCopyExtensionsBuildPhaseSection, PBXResourcesBuildPhaseSection, AllExtensions, GameProjectPath, BuildConfigs);
 
 			StringBuilder ProjectFileContent = new StringBuilder();
 
@@ -1918,7 +1918,7 @@ namespace UnrealBuildTool.XcodeProjectLegacy
 				bSuccess = ProjectFileGenerator.WriteFileIfChanged(PBXProjFilePath.FullName, ProjectFileContent.ToString(), Logger, new UTF8Encoding());
 			}
 
-			bool bNeedScheme = CanBuildProjectLocally(Logger);
+			bool bNeedScheme = CanBuildProjectLocally();
 
 			if (bNeedScheme)
 			{
diff --git a/Engine/Source/Programs/UnrealBuildTool/System/ActionGraph.cs b/Engine/Source/Programs/UnrealBuildTool/System/ActionGraph.cs
--- a/Engine/Source/Programs/UnrealBuildTool/System/ActionGraph.cs
+++ b/Engine/Source/Programs/UnrealBuildTool/System/ActionGraph.cs
@@ -485,7 +485,7 @@ namespace UnrealBuildTool
 				if (executor != null)
 				{
 					// We still want to run the executor to give it a chance to clean up in case it has tasks running that can be cancelled
-					await InternalExecuteActions(executor, BuildConfiguration, ActionsToExecute, TargetDescriptors, Logger, actionArtifactCache);
+					await InternalExecuteActions(executor, ActionsToExecute, Logger, actionArtifactCache);
 				}
 			}
 			else
@@ -494,16 +494,16 @@ namespace UnrealBuildTool
 				if (executor == null)
 				{
 					using ActionExecutor SelectedExecutor = SelectExecutor(BuildConfiguration, ActionsToExecute.Count, TargetDescriptors, Logger);
-					await InternalExecuteActions(SelectedExecutor, BuildConfiguration, ActionsToExecute, TargetDescriptors, Logger, actionArtifactCache);
+					await InternalExecuteActions(SelectedExecutor, ActionsToExecute, Logger, actionArtifactCache);
 				}
 				else
 				{
-					await InternalExecuteActions(executor, BuildConfiguration, ActionsToExecute, TargetDescriptors, Logger, actionArtifactCache);
+					await InternalExecuteActions(executor, ActionsToExecute, Logger, actionArtifactCache);
 				}
 			}
 		}
 
-		static async Task InternalExecuteActions(ActionExecutor Executor, BuildConfiguration BuildConfiguration, List<LinkedAction> ActionsToExecute, List<TargetDescriptor> TargetDescriptors, ILogger Logger, IActionArtifactCache? actionArtifactCache = null)
+		static async Task InternalExecuteActions(ActionExecutor Executor, List<LinkedAction> ActionsToExecute, ILogger Logger, IActionArtifactCache? actionArtifactCache = null)
 		{
 			Logger.LogInformation("Using {ExecutorName} executor to run {ActionCount} action(s)", Executor.Name, ActionsToExecute.Count);
 
diff --git a/Engine/Source/Programs/UnrealBuildTool/System/ConfigCache.cs b/Engine/Source/Programs/UnrealBuildTool/System/ConfigCache.cs
--- a/Engine/Source/Programs/UnrealBuildTool/System/ConfigCache.cs
+++ b/Engine/Source/Programs/UnrealBuildTool/System/ConfigCache.cs
@@ -430,17 +430,17 @@ namespace UnrealBuildTool
 				}
 				foreach (FieldInfo FieldInfo in TargetObjectType.GetFields(BindingFlags.Instance | BindingFlags.GetField | BindingFlags.GetProperty | BindingFlags.Public | BindingFlags.NonPublic | BindingFlags.DeclaredOnly))
 				{
-					ProcessConfigTypeMember<FieldInfo>(TargetObjectType, FieldInfo, Members, (FieldInfo, Attribute) => new ConfigField(FieldInfo, Attribute));
+					ProcessConfigTypeMember<FieldInfo>(FieldInfo, Members, (FieldInfo, Attribute) => new ConfigField(FieldInfo, Attribute));
 				}
 				foreach (PropertyInfo PropertyInfo in TargetObjectType.GetProperties(BindingFlags.Instance | BindingFlags.GetField | BindingFlags.GetProperty | BindingFlags.Public | BindingFlags.NonPublic | BindingFlags.DeclaredOnly))
 				{
-					ProcessConfigTypeMember<PropertyInfo>(TargetObjectType, PropertyInfo, Members, (PropertyInfo, Attribute) => new ConfigProperty(PropertyInfo, Attribute));
+					ProcessConfigTypeMember<PropertyInfo>(PropertyInfo, Members, (PropertyInfo, Attribute) => new ConfigProperty(PropertyInfo, Attribute));
 				}
 				return Members;
 			})).Value;
 		}
 
-		static void ProcessConfigTypeMember<MEMBER>(Type TargetType, MEMBER MemberInfo, List<ConfigMember> Members, Func<MEMBER, ConfigFileAttribute, ConfigMember> CreateConfigMember)
+		static void ProcessConfigTypeMember<MEMBER>(MEMBER MemberInfo, List<ConfigMember> Members, Func<MEMBER, ConfigFileAttribute, ConfigMember> CreateConfigMember)
 			where MEMBER : System.Reflection.MemberInfo
 		{
 			IEnumerable<ConfigFileAttribute> Attributes = MemberInfo.GetCustomAttributes<ConfigFileAttribute>();
diff --git a/Engine/Source/Programs/UnrealBuildTool/System/FileHashCache.cs b/Engine/Source/Programs/UnrealBuildTool/System/FileHashCache.cs
--- a/Engine/Source/Programs/UnrealBuildTool/System/FileHashCache.cs
+++ b/Engine/Source/Programs/UnrealBuildTool/System/FileHashCache.cs
@@ -299,14 +299,14 @@ namespace UnrealBuildTool
 			}
 			if (allowAsync)
 			{
-				return Task.Run(() => InternalWriteFileIfChanged2(writtenBefore, partition, fileItem, newHash, data, logger, allowAsync, alwaysUpdateTimestamp));
+				return Task.Run(() => InternalWriteFileIfChanged2(writtenBefore, partition, fileItem, newHash, data, logger, alwaysUpdateTimestamp));
 			}
 
-			InternalWriteFileIfChanged2(writtenBefore, partition, fileItem, newHash, data, logger, allowAsync, alwaysUpdateTimestamp);
+			InternalWriteFileIfChanged2(writtenBefore, partition, fileItem, newHash, data, logger, alwaysUpdateTimestamp);
 			return null;
 		}
 
-		private void InternalWriteFileIfChanged2(bool writtenBefore, Partition? partition, FileItem fileItem, UInt128 newHash, object data, ILogger logger, bool allowAsync, bool alwaysUpdateTimestamp)
+		private void InternalWriteFileIfChanged2(bool writtenBefore, Partition? partition, FileItem fileItem, UInt128 newHash, object data, ILogger logger, bool alwaysUpdateTimestamp)
 		{
 			if (fileItem.Exists)
 			{
diff --git a/Engine/Source/Programs/UnrealBuildTool/System/NativeProjects.cs b/Engine/Source/Programs/UnrealBuildTool/System/NativeProjects.cs
--- a/Engine/Source/Programs/UnrealBuildTool/System/NativeProjects.cs
+++ b/Engine/Source/Programs/UnrealBuildTool/System/NativeProjects.cs
@@ -401,7 +401,7 @@ namespace UnrealBuildTool
 						if (!ConfiguredPluginNames.Contains(PluginReference.Name))
 						{
 							PluginReferenceDescriptor? MissingPlugin;
-							if (!ConfigureEnabledPluginForTarget(PluginReference, ProjectDescriptor, null, Platform, Configuration, TargetType, bLoadPluginsForTargetPlatforms, AllPlugins, EnabledPlugins, out MissingPlugin, Logger))
+							if (!ConfigureEnabledPluginForTarget(PluginReference, null, Platform, Configuration, TargetType, bLoadPluginsForTargetPlatforms, AllPlugins, EnabledPlugins, out MissingPlugin, Logger))
 							{
 								OutMissingPlugin = MissingPlugin;
 								return false;
@@ -418,7 +418,7 @@ namespace UnrealBuildTool
 				if (PluginPair.Value.IsEnabledByDefault(bAllowEnginePluginsEnabledByDefault) && !ConfiguredPluginNames.Contains(PluginPair.Key))
 				{
 					PluginReferenceDescriptor? MissingPlugin;
-					if (!ConfigureEnabledPluginForTarget(new PluginReferenceDescriptor(PluginPair.Key, null, true), ProjectDescriptor, null, Platform, Configuration, TargetType, bLoadPluginsForTargetPlatforms, AllPlugins, EnabledPlugins, out MissingPlugin, Logger))
+					if (!ConfigureEnabledPluginForTarget(new PluginReferenceDescriptor(PluginPair.Key, null, true), null, Platform, Configuration, TargetType, bLoadPluginsForTargetPlatforms, AllPlugins, EnabledPlugins, out MissingPlugin, Logger))
 					{
 						OutMissingPlugin = MissingPlugin;
 						return false;
@@ -452,7 +452,7 @@ namespace UnrealBuildTool
 		/// <summary>
 		/// NOTE: This function should mirror FPluginManager::ConfigureEnabledPluginForTarget
 		/// </summary>
-		static bool ConfigureEnabledPluginForTarget(PluginReferenceDescriptor FirstReference, ProjectDescriptor? ProjectDescriptor, string? TargetName, UnrealTargetPlatform Platform, UnrealTargetConfiguration Configuration, TargetType TargetType, bool bLoadPluginsForTargetPlatforms, Dictionary<string, PluginInfo> AllPlugins, Dictionary<string, PluginInfo> EnabledPlugins, out PluginReferenceDescriptor? OutMissingPlugin, ILogger Logger)
+		static bool ConfigureEnabledPluginForTarget(PluginReferenceDescriptor FirstReference, string? TargetName, UnrealTargetPlatform Platform, UnrealTargetConfiguration Configuration, TargetType TargetType, bool bLoadPluginsForTargetPlatforms, Dictionary<string, PluginInfo> AllPlugins, Dictionary<string, PluginInfo> EnabledPlugins, out PluginReferenceDescriptor? OutMissingPlugin, ILogger Logger)
 		{
 			if (!EnabledPlugins.ContainsKey(FirstReference.Name))
 			{
diff --git a/Engine/Source/Programs/UnrealBuildTool/System/SourceFileWorkingSet.cs b/Engine/Source/Programs/UnrealBuildTool/System/SourceFileWorkingSet.cs
--- a/Engine/Source/Programs/UnrealBuildTool/System/SourceFileWorkingSet.cs
+++ b/Engine/Source/Programs/UnrealBuildTool/System/SourceFileWorkingSet.cs
@@ -357,7 +357,7 @@ namespace UnrealBuildTool
 				{
 					return WorkingSet;
 				}
-				else if (TryCreatePerforceWorkingSet(RootDir, ProjectDirs, Logger, out WorkingSet))
+				else if (TryCreatePerforceWorkingSet(RootDir, out WorkingSet))
 				{
 					return WorkingSet;
 				}
@@ -396,7 +396,7 @@ namespace UnrealBuildTool
 			return OutWorkingSet != null;
 		}
 
-		static bool TryCreatePerforceWorkingSet(DirectoryReference RootDir, IEnumerable<DirectoryReference> ProjectDirs, ILogger Logger, [NotNullWhen(true)] out ISourceFileWorkingSet? OutWorkingSet)
+		static bool TryCreatePerforceWorkingSet(DirectoryReference RootDir, [NotNullWhen(true)] out ISourceFileWorkingSet? OutWorkingSet)
 		{
 			PerforceSourceFileWorkingSet? WorkingSet = null;
 			// If an installed engine, or the root directory contains any read-only files assume this is a perforce working set
diff --git a/Engine/Source/Programs/UnrealBuildTool/System/TopologicalSorter.cs b/Engine/Source/Programs/UnrealBuildTool/System/TopologicalSorter.cs
--- a/Engine/Source/Programs/UnrealBuildTool/System/TopologicalSorter.cs
+++ b/Engine/Source/Programs/UnrealBuildTool/System/TopologicalSorter.cs
@@ -139,7 +139,7 @@ namespace UnrealBuildTool
 			if (Node.Mark == Mark.InProgress)
 			{
 				// This node is being processed i.e. we've detected a cycle.
-				return HandleCycle(Node);
+				return HandleCycle();
 			}
 
 			Node.Mark = Mark.InProgress;
@@ -223,9 +223,8 @@ namespace UnrealBuildTool
 		/// <summary>
 		/// Member executed when a graph cycle is encountered.
 		/// </summary>
-		/// <param name="Node"></param>
 		/// <returns>True if operation should continue, false if a cycle is considered an error</returns>
-		private bool HandleCycle(GraphNode Node)
+		private bool HandleCycle()
 		{
 			switch (CycleHandling)
 			{
diff --git a/Engine/Source/Programs/UnrealBuildTool/System/UHTExecution.cs b/Engine/Source/Programs/UnrealBuildTool/System/UHTExecution.cs
--- a/Engine/Source/Programs/UnrealBuildTool/System/UHTExecution.cs
+++ b/Engine/Source/Programs/UnrealBuildTool/System/UHTExecution.cs
@@ -1232,7 +1232,7 @@ namespace UnrealBuildTool
 				}
 			}
 
-			internal void Save(ILogger logger)
+			internal void Save()
 			{
 				using FileStream stream = FileReference.Open(_location, FileMode.Create, FileAccess.Write);
 				using BinaryWriter writer = new(stream);
@@ -1469,7 +1469,7 @@ namespace UnrealBuildTool
 					{
 						using ITimelineEvent uhtExecute = Timeline.ScopeEvent("UHTTimestamps");
 						timestampsFile.Update(Makefile.UObjectModules);
-						timestampsFile.Save(Logger);
+						timestampsFile.Save();
 					}
 				}
 			}, TaskScheduler.Current);
diff --git a/Engine/Source/Programs/UnrealBuildTool/System/Unity.cs b/Engine/Source/Programs/UnrealBuildTool/System/Unity.cs
--- a/Engine/Source/Programs/UnrealBuildTool/System/Unity.cs
+++ b/Engine/Source/Programs/UnrealBuildTool/System/Unity.cs
@@ -205,7 +205,7 @@ namespace UnrealBuildTool
 			bool bForceIntoSingleUnityFile = Target.bStressTestUnity || (TotalBytesInSourceFiles < NumIncludedBytesPerUnitySource * 2 && Target.bUsePCHFiles);
 
 			// Even if every single file in the module appears in the working set adaptive unity should still be used even if it's slower.
-			GetAdaptiveFiles(Target, SourceFiles, HeaderFiles, CompileEnvironment, WorkingSet, BaseName, IntermediateDirectory, Graph, out NormalFiles, out AdaptiveFiles);
+			GetAdaptiveFiles(Target, SourceFiles, HeaderFiles, CompileEnvironment, WorkingSet, IntermediateDirectory, Graph, out NormalFiles, out AdaptiveFiles);
 
 			// Build the list of unity files.
 			List<FileCollection> AllUnityFiles;
@@ -341,7 +341,6 @@ namespace UnrealBuildTool
 			List<FileItem> HeaderFiles,
 			CppCompileEnvironment CompileEnvironment,
 			ISourceFileWorkingSet WorkingSet,
-			string BaseName,
 			DirectoryReference IntermediateDirectory,
 			IActionGraphBuilder Graph,
 			out List<FileItem> NormalFiles,
diff --git a/Engine/Source/Programs/UnrealBuildTool/ToolChain/ISPCToolChain.cs b/Engine/Source/Programs/UnrealBuildTool/ToolChain/ISPCToolChain.cs
--- a/Engine/Source/Programs/UnrealBuildTool/ToolChain/ISPCToolChain.cs
+++ b/Engine/Source/Programs/UnrealBuildTool/ToolChain/ISPCToolChain.cs
@@ -624,7 +624,7 @@ namespace UnrealBuildTool
 
 				Arguments.AddRange(GlobalArguments);
 
-				FileReference ResponseFileName = GetResponseFileName(CompileEnvironment, ISPCIncludeHeaderFile);
+				FileReference ResponseFileName = GetResponseFileName(ISPCIncludeHeaderFile);
 				FileItem ResponseFileItem = Graph.CreateIntermediateTextFile(ResponseFileName, Arguments.Select(x => Utils.ExpandVariables(x)));
 				CompileAction.CommandArguments = $"@\"{NormalizeCommandLinePath(ResponseFileName, RootPaths)}\"";
 				CompileAction.PrerequisiteItems.Add(ResponseFileItem);
@@ -862,7 +862,7 @@ namespace UnrealBuildTool
 
 				CompileAction.ProducedItems.UnionWith(CompiledISPCObjFiles);
 
-				FileReference ResponseFileName = GetResponseFileName(CompileEnvironment, CompiledISPCObjFileNoISA);
+				FileReference ResponseFileName = GetResponseFileName(CompiledISPCObjFileNoISA);
 				FileItem ResponseFileItem = Graph.CreateIntermediateTextFile(ResponseFileName, Arguments.Select(x => Utils.ExpandVariables(x)));
 
 				string AdditionalArguments = "";
@@ -920,7 +920,7 @@ namespace UnrealBuildTool
 							}
 
 							// Write the args to a response file
-							FileReference PostCompileResponseFileName = GetResponseFileName(CompileEnvironment, FinalCompiledISPCObjFile);
+							FileReference PostCompileResponseFileName = GetResponseFileName(FinalCompiledISPCObjFile);
 							FileItem PostCompileResponseFileItem = Graph.CreateIntermediateTextFile(PostCompileResponseFileName, PostCompileArgs.Select(x => Utils.ExpandVariables(x)));
 							PostCompileAction.CommandArguments = $"@\"{NormalizeCommandLinePath(PostCompileResponseFileName, RootPaths)}\"{AdditionalArguments}";
 							PostCompileAction.PrerequisiteItems.Add(PostCompileResponseFileItem);
diff --git a/Engine/Source/Programs/UnrealBuildTool/ToolChain/IWYUToolChain.cs b/Engine/Source/Programs/UnrealBuildTool/ToolChain/IWYUToolChain.cs
--- a/Engine/Source/Programs/UnrealBuildTool/ToolChain/IWYUToolChain.cs
+++ b/Engine/Source/Programs/UnrealBuildTool/ToolChain/IWYUToolChain.cs
@@ -104,7 +104,7 @@ namespace UnrealBuildTool
 		// Code below here is mostly copy-pasted from LinuxToolChain.
 		// We didn't inherit from LinuxToolChain on purpose in order to be able to add support for running iwyu against other clang platforms than linux
 
-		private static bool ShouldUseLibcxx(UnrealArch Architecture)
+		private static bool ShouldUseLibcxx()
 		{
 			// set UE_LINUX_USE_LIBCXX to either 0 or 1. If unset, defaults to 1.
 			string? UseLibcxxEnvVarOverride = Environment.GetEnvironmentVariable("UE_LINUX_USE_LIBCXX");
@@ -150,7 +150,7 @@ namespace UnrealBuildTool
 
 				if (TargetPlatform == UnrealTargetPlatform.Linux)
 				{
-					if (ShouldUseLibcxx(CompileEnvironment.Architecture))
+					if (ShouldUseLibcxx())
 					{
 						Arguments.Add("-nostdinc++");
 
diff --git a/Engine/Source/Programs/UnrealBuildTool/ToolChain/UEToolChain.cs b/Engine/Source/Programs/UnrealBuildTool/ToolChain/UEToolChain.cs
--- a/Engine/Source/Programs/UnrealBuildTool/ToolChain/UEToolChain.cs
+++ b/Engine/Source/Programs/UnrealBuildTool/ToolChain/UEToolChain.cs
@@ -275,12 +275,11 @@ namespace UnrealBuildTool
 		private static FileReference GetResponseFileName(FileReference Location) => Location.ChangeExtension($"{Location.GetExtension()}{ResponseExt}");
 
 		/// <summary>
-		/// Get the name of the response file for the current compile environment and output file
+		/// Get the name of the response file for the output file
 		/// </summary>
-		/// <param name="CompileEnvironment"></param>
 		/// <param name="OutputFile"></param>
 		/// <returns></returns>
-		public static FileReference GetResponseFileName(CppCompileEnvironment CompileEnvironment, FileItem OutputFile) => GetResponseFileName(OutputFile.Location);
+		public static FileReference GetResponseFileName(FileItem OutputFile) => GetResponseFileName(OutputFile.Location);
 
 		/// <summary>
 		/// Get the name of the response file for the current linker environment and output file
