Fix issue IDE0059 (Remove unnecessary value assignment) in UBT.

#Synced-CL 48282190

diff --git a/Engine/Source/Programs/UnrealBuildTool/Platform/IOS/UEDeployIOS.cs b/Engine/Source/Programs/UnrealBuildTool/Platform/IOS/UEDeployIOS.cs
--- a/Engine/Source/Programs/UnrealBuildTool/Platform/IOS/UEDeployIOS.cs
+++ b/Engine/Source/Programs/UnrealBuildTool/Platform/IOS/UEDeployIOS.cs
@@ -226,8 +226,6 @@ namespace UnrealBuildTool
 			// orientations
 			string InterfaceOrientation = "";
 			Ini.GetString("/Script/IOSRuntimeSettings.IOSRuntimeSettings", "PreferredLandscapeOrientation", out string PreferredLandscapeOrientation);
-			Ini.GetBool("/Script/IOSRuntimeSettings.IOSRuntimeSettings", "bSupportsPortraitOrientation", out bool bSupportsPortrait);
-			Ini.GetBool("/Script/IOSRuntimeSettings.IOSRuntimeSettings", "bSupportsUpsideDownOrientation", out bool bSupportsUpsideDown);
 			Ini.GetBool("/Script/IOSRuntimeSettings.IOSRuntimeSettings", "bSupportsLandscapeLeftOrientation", out bool bSupportsLandscapeLeft);
 			Ini.GetBool("/Script/IOSRuntimeSettings.IOSRuntimeSettings", "bSupportsLandscapeRightOrientation", out bool bSupportsLandscapeRight);
 
@@ -532,24 +530,17 @@ namespace UnrealBuildTool
 
 			// orientations
 			string InterfaceOrientation = "";
-			string PreferredLandscapeOrientation = "";
-			Ini.GetString("/Script/IOSRuntimeSettings.IOSRuntimeSettings", "PreferredLandscapeOrientation", out PreferredLandscapeOrientation);
+			Ini.GetString("/Script/IOSRuntimeSettings.IOSRuntimeSettings", "PreferredLandscapeOrientation", out string PreferredLandscapeOrientation);
 
 			string SupportedOrientations = "";
-			bool bSupported = true;
-			Ini.GetBool("/Script/IOSRuntimeSettings.IOSRuntimeSettings", "bSupportsPortraitOrientation", out bSupported);
-			SupportedOrientations += bSupported ? "\t\t<string>UIInterfaceOrientationPortrait</string>\n" : "";
-			bool bSupportsPortrait = bSupported;
+			Ini.GetBool("/Script/IOSRuntimeSettings.IOSRuntimeSettings", "bSupportsPortraitOrientation", out bool bSupportsPortrait);
+			SupportedOrientations += bSupportsPortrait ? "\t\t<string>UIInterfaceOrientationPortrait</string>\n" : "";
 
-			Ini.GetBool("/Script/IOSRuntimeSettings.IOSRuntimeSettings", "bSupportsUpsideDownOrientation", out bSupported);
-			SupportedOrientations += bSupported ? "\t\t<string>UIInterfaceOrientationPortraitUpsideDown</string>\n" : "";
-			bSupportsPortrait |= bSupported;
+			Ini.GetBool("/Script/IOSRuntimeSettings.IOSRuntimeSettings", "bSupportsUpsideDownOrientation", out bool bSupportsUpsideDown);
+			SupportedOrientations += bSupportsUpsideDown ? "\t\t<string>UIInterfaceOrientationPortraitUpsideDown</string>\n" : "";
 
-			bool bSupportsLandscapeLeft = false;
-			Ini.GetBool("/Script/IOSRuntimeSettings.IOSRuntimeSettings", "bSupportsLandscapeLeftOrientation", out bSupportsLandscapeLeft);
-			bool bSupportsLandscapeRight = false;
-			Ini.GetBool("/Script/IOSRuntimeSettings.IOSRuntimeSettings", "bSupportsLandscapeRightOrientation", out bSupportsLandscapeRight);
-			bool bSupportsLandscape = bSupportsLandscapeLeft || bSupportsLandscapeRight;
+			Ini.GetBool("/Script/IOSRuntimeSettings.IOSRuntimeSettings", "bSupportsLandscapeLeftOrientation", out bool bSupportsLandscapeLeft);
+			Ini.GetBool("/Script/IOSRuntimeSettings.IOSRuntimeSettings", "bSupportsLandscapeRightOrientation", out bool bSupportsLandscapeRight);
 
 			if (bSupportsLandscapeLeft && bSupportsLandscapeRight)
 			{
@@ -575,76 +566,60 @@ namespace UnrealBuildTool
 			}
 
 			// ITunes file sharing
-			bool bSupportsITunesFileSharing = false;
-			Ini.GetBool("/Script/IOSRuntimeSettings.IOSRuntimeSettings", "bSupportsITunesFileSharing", out bSupportsITunesFileSharing);
-			bool bSupportsFilesApp = false;
-			Ini.GetBool("/Script/IOSRuntimeSettings.IOSRuntimeSettings", "bSupportsFilesApp", out bSupportsFilesApp);
+			Ini.GetBool("/Script/IOSRuntimeSettings.IOSRuntimeSettings", "bSupportsITunesFileSharing", out bool bSupportsITunesFileSharing);
+			Ini.GetBool("/Script/IOSRuntimeSettings.IOSRuntimeSettings", "bSupportsFilesApp", out bool bSupportsFilesApp);
 
 			// bundle display name
-			string BundleDisplayName;
-			Ini.GetString("/Script/IOSRuntimeSettings.IOSRuntimeSettings", "BundleDisplayName", out BundleDisplayName);
+			Ini.GetString("/Script/IOSRuntimeSettings.IOSRuntimeSettings", "BundleDisplayName", out string BundleDisplayName);
 
 			// bundle identifier
-			string BundleIdentifier;
-			Ini.GetString("/Script/IOSRuntimeSettings.IOSRuntimeSettings", "BundleIdentifier", out BundleIdentifier);
+			Ini.GetString("/Script/IOSRuntimeSettings.IOSRuntimeSettings", "BundleIdentifier", out string BundleIdentifier);
 			if (!String.IsNullOrEmpty(BundleID))
 			{
 				BundleIdentifier = BundleID; // overriding bundle ID
 			}
 
 			// bundle name
-			string BundleName;
-			Ini.GetString("/Script/IOSRuntimeSettings.IOSRuntimeSettings", "BundleName", out BundleName);
+			Ini.GetString("/Script/IOSRuntimeSettings.IOSRuntimeSettings", "BundleName", out string BundleName);
 
 			// disable https requirement
-			bool bDisableHTTPS;
-			Ini.GetBool("/Script/IOSRuntimeSettings.IOSRuntimeSettings", "bDisableHTTPS", out bDisableHTTPS);
+			Ini.GetBool("/Script/IOSRuntimeSettings.IOSRuntimeSettings", "bDisableHTTPS", out bool bDisableHTTPS);
 			
 			// short version string
-			string BundleShortVersion;
-			Ini.GetString("/Script/IOSRuntimeSettings.IOSRuntimeSettings", "VersionInfo", out BundleShortVersion);
+			Ini.GetString("/Script/IOSRuntimeSettings.IOSRuntimeSettings", "VersionInfo", out string BundleShortVersion);
 
 			// required capabilities (arm64 always required)
 			string RequiredCaps = "\t\t<string>arm64</string>\n";
 
-			Ini.GetBool("/Script/IOSRuntimeSettings.IOSRuntimeSettings", "bSupportsMetal", out bSupported);
-			RequiredCaps += bSupported ? "\t\t<string>metal</string>\n" : "";
+			Ini.GetBool("/Script/IOSRuntimeSettings.IOSRuntimeSettings", "bSupportsMetal", out bool bSupportsMetal);
+			RequiredCaps += bSupportsMetal ? "\t\t<string>metal</string>\n" : "";
 
 			// minimum iOS version
-			string MinVersionSetting = "";
-			Ini.GetString("/Script/IOSRuntimeSettings.IOSRuntimeSettings", "MinimumiOSVersion", out MinVersionSetting);
+			Ini.GetString("/Script/IOSRuntimeSettings.IOSRuntimeSettings", "MinimumiOSVersion", out string MinVersionSetting);
 			string MinVersion = GetMinimumOSVersion(MinVersionSetting, Logger);
 
 			// Get Google Support details
-			bool bEnableGoogleSupport = true;
-			Ini.GetBool("/Script/IOSRuntimeSettings.IOSRuntimeSettings", "bEnableGoogleSupport", out bEnableGoogleSupport);
+			Ini.GetBool("/Script/IOSRuntimeSettings.IOSRuntimeSettings", "bEnableGoogleSupport", out bool bEnableGoogleSupport);
 
 			// Write the Google iOS URL Scheme if we need it.
-			string GoogleReversedClientId = "";
-			Ini.GetString("/Script/IOSRuntimeSettings.IOSRuntimeSettings", "GoogleReversedClientId", out GoogleReversedClientId);
+			Ini.GetString("/Script/IOSRuntimeSettings.IOSRuntimeSettings", "GoogleReversedClientId", out string GoogleReversedClientId);
 			bEnableGoogleSupport = bEnableGoogleSupport && !String.IsNullOrWhiteSpace(GoogleReversedClientId);
 
 			// Add remote-notifications as background mode
-			bool bRemoteNotificationsSupported = false;
-			Ini.GetBool("/Script/IOSRuntimeSettings.IOSRuntimeSettings", "bEnableRemoteNotificationsSupport", out bRemoteNotificationsSupported);
+			Ini.GetBool("/Script/IOSRuntimeSettings.IOSRuntimeSettings", "bEnableRemoteNotificationsSupport", out bool bRemoteNotificationsSupported);
 
 			// Add audio as background mode
-			bool bBackgroundAudioSupported = false;
-			Ini.GetBool("/Script/IOSRuntimeSettings.IOSRuntimeSettings", "bSupportsBackgroundAudio", out bBackgroundAudioSupported);
+			Ini.GetBool("/Script/IOSRuntimeSettings.IOSRuntimeSettings", "bSupportsBackgroundAudio", out bool bBackgroundAudioSupported);
 
 			// Add background fetch as background mode
-			bool bBackgroundFetch = false;
-			Ini.GetBool("/Script/IOSRuntimeSettings.IOSRuntimeSettings", "bEnableBackgroundFetch", out bBackgroundFetch);
+			Ini.GetBool("/Script/IOSRuntimeSettings.IOSRuntimeSettings", "bEnableBackgroundFetch", out bool bBackgroundFetch);
 
 			// Get any Location Services permission descriptions added
-			string LocationAlwaysUsageDescription = "";
-			string LocationWhenInUseDescription = "";
-			Ini.GetString("/Script/LocationServicesIOSEditor.LocationServicesIOSSettings", "LocationAlwaysUsageDescription", out LocationAlwaysUsageDescription);
-			Ini.GetString("/Script/LocationServicesIOSEditor.LocationServicesIOSSettings", "LocationWhenInUseDescription", out LocationWhenInUseDescription);
+			Ini.GetString("/Script/LocationServicesIOSEditor.LocationServicesIOSSettings", "LocationAlwaysUsageDescription", out string LocationAlwaysUsageDescription);
+			Ini.GetString("/Script/LocationServicesIOSEditor.LocationServicesIOSSettings", "LocationWhenInUseDescription", out string LocationWhenInUseDescription);
 
 			// extra plist data
-			string ExtraData = "";
-			Ini.GetString("/Script/IOSRuntimeSettings.IOSRuntimeSettings", "AdditionalPlistData", out ExtraData);
+			Ini.GetString("/Script/IOSRuntimeSettings.IOSRuntimeSettings", "AdditionalPlistData", out string ExtraData);
 
 			Ini.GetBool("/Script/IOSRuntimeSettings.IOSRuntimeSettings", "bCustomLaunchscreenStoryboard", out VersionUtilities.bCustomLaunchscreenStoryboard);
 
@@ -784,8 +759,7 @@ namespace UnrealBuildTool
 
 			// Support high refresh rates (iPhone only)
 			// https://developer.apple.com/documentation/quartzcore/optimizing_promotion_refresh_rates_for_iphone_13_pro_and_ipad_pro
-			bool bSupportHighRefreshRates = false;
-			Ini.GetBool("/Script/IOSRuntimeSettings.IOSRuntimeSettings", "bSupportHighRefreshRates", out bSupportHighRefreshRates);
+			Ini.GetBool("/Script/IOSRuntimeSettings.IOSRuntimeSettings", "bSupportHighRefreshRates", out bool bSupportHighRefreshRates);
 			if (bSupportHighRefreshRates)
 			{
 				Text.AppendLine("\t<key>CADisableMinimumFrameDurationOnPhone</key><true/>");
@@ -940,7 +914,6 @@ namespace UnrealBuildTool
 		public bool GeneratePList(FileReference ProjectFile, UnrealTargetConfiguration Config, string ProjectDirectory, bool bIsUnrealGame, string GameName, bool bIsClient, string ProjectName, string InEngineDir, string AppDirectory, TargetReceipt Receipt)
 		{
 			List<string> UPLScripts = CollectPluginDataPaths(Receipt.AdditionalProperties, Logger);
-			VersionNumber? SdkVersion = GetSdkVersion(Receipt);
 			bool bBuildAsFramework = GetCompileAsDll(Receipt);
 			return GeneratePList(ProjectFile, Config, ProjectDirectory, bIsUnrealGame, GameName, bIsClient, ProjectName, InEngineDir, AppDirectory, UPLScripts, "", bBuildAsFramework);
 		}
@@ -1032,7 +1005,6 @@ namespace UnrealBuildTool
 		public bool PrepForUATPackageOrDeploy(UnrealTargetConfiguration Config, FileReference ProjectFile, string InProjectName, string InProjectDirectory, FileReference Executable, string InEngineDir, bool bForDistribution, string CookFlavor, bool bIsDataDeploy, bool bCreateStubIPA, TargetReceipt Receipt)
 		{
 			List<string> UPLScripts = CollectPluginDataPaths(Receipt.AdditionalProperties, Logger);
-			VersionNumber? SdkVersion = GetSdkVersion(Receipt);
 			bool bBuildAsFramework = GetCompileAsDll(Receipt);
 			return PrepForUATPackageOrDeploy(Config, ProjectFile, InProjectName, InProjectDirectory, Executable, InEngineDir, bForDistribution, CookFlavor, bIsDataDeploy, bCreateStubIPA, UPLScripts, "", bBuildAsFramework);
 		}
diff --git a/Engine/Source/Programs/UnrealBuildTool/Platform/Linux/LinuxProjectGenerator.cs b/Engine/Source/Programs/UnrealBuildTool/Platform/Linux/LinuxProjectGenerator.cs
--- a/Engine/Source/Programs/UnrealBuildTool/Platform/Linux/LinuxProjectGenerator.cs
+++ b/Engine/Source/Programs/UnrealBuildTool/Platform/Linux/LinuxProjectGenerator.cs
@@ -7,7 +7,6 @@ using System.Linq;
 using System.Text.RegularExpressions;
 using EpicGames.Core;
 using Microsoft.Extensions.Logging;
-using UnrealBuildBase;
 
 namespace UnrealBuildTool
 {
@@ -45,7 +44,6 @@ namespace UnrealBuildTool
 		public override IList<string> GetSystemIncludePaths(UEBuildTarget InTarget)
 		{
 			List<string> Result = [];
-			string EngineDirectory = Unreal.EngineDirectory.ToString();
 			string? UseLibcxxEnvVarOverride = Environment.GetEnvironmentVariable("UE_LINUX_USE_LIBCXX");
 			bool ShouldUseLibcxx = String.IsNullOrEmpty(UseLibcxxEnvVarOverride) || UseLibcxxEnvVarOverride == "1";
 
diff --git a/Engine/Source/Programs/UnrealBuildTool/Platform/TVOS/UEDeployTVOS.cs b/Engine/Source/Programs/UnrealBuildTool/Platform/TVOS/UEDeployTVOS.cs
--- a/Engine/Source/Programs/UnrealBuildTool/Platform/TVOS/UEDeployTVOS.cs
+++ b/Engine/Source/Programs/UnrealBuildTool/Platform/TVOS/UEDeployTVOS.cs
@@ -46,45 +46,24 @@ namespace UnrealBuildTool
 			ConfigHierarchy Ini = ConfigCache.ReadHierarchy(ConfigHierarchyType.Engine, DirRef, UnrealTargetPlatform.IOS);
 
 			// bundle display name
-			string BundleDisplayName;
-			Ini.GetString("/Script/IOSRuntimeSettings.IOSRuntimeSettings", "BundleDisplayName", out BundleDisplayName);
+			Ini.GetString("/Script/IOSRuntimeSettings.IOSRuntimeSettings", "BundleDisplayName", out string BundleDisplayName);
 
 			// bundle identifier
-			string BundleIdentifier;
-			Ini.GetString("/Script/IOSRuntimeSettings.IOSRuntimeSettings", "BundleIdentifier", out BundleIdentifier);
+			Ini.GetString("/Script/IOSRuntimeSettings.IOSRuntimeSettings", "BundleIdentifier", out string BundleIdentifier);
 			if (!String.IsNullOrEmpty(BundleID))
 			{
 				BundleIdentifier = BundleID;
 			}
 
 			// bundle name
-			string BundleName;
-			Ini.GetString("/Script/IOSRuntimeSettings.IOSRuntimeSettings", "BundleName", out BundleName);
+			Ini.GetString("/Script/IOSRuntimeSettings.IOSRuntimeSettings", "BundleName", out string BundleName);
 
 			// short version string
-			string BundleShortVersion;
-			Ini.GetString("/Script/IOSRuntimeSettings.IOSRuntimeSettings", "VersionInfo", out BundleShortVersion);
+			Ini.GetString("/Script/IOSRuntimeSettings.IOSRuntimeSettings", "VersionInfo", out string BundleShortVersion);
 
 			// required capabilities
 			string RequiredCaps = "\t\t<string>arm64</string>\n";
 
-			// minimum iOS version
-			string MinVersionSetting = "";
-			Ini.GetString("/Script/IOSRuntimeSettings.IOSRuntimeSettings", "MinimumiOSVersion", out MinVersionSetting);
-			string MinVersion = GetMinimumOSVersion(MinVersionSetting, Logger);
-
-			// extra plist data
-			string ExtraData = "";
-			Ini.GetString("/Script/IOSRuntimeSettings.IOSRuntimeSettings", "AdditionalPlistData", out ExtraData);
-
-			// create the final display name, including converting all entities for XML use
-			string FinalDisplayName = BundleDisplayName.Replace("[PROJECT_NAME]", ProjectName).Replace("_", "");
-			FinalDisplayName = FinalDisplayName.Replace("&", "&amp;");
-			FinalDisplayName = FinalDisplayName.Replace("\"", "&quot;");
-			FinalDisplayName = FinalDisplayName.Replace("\'", "&apos;");
-			FinalDisplayName = FinalDisplayName.Replace("<", "&lt;");
-			FinalDisplayName = FinalDisplayName.Replace(">", "&gt;");
-
 			// generate the plist file
 			StringBuilder Text = new StringBuilder();
 			Text.AppendLine("<?xml version=\"1.0\" encoding=\"UTF-8\"?>");
diff --git a/Engine/Source/Programs/UnrealBuildTool/Platform/Windows/AppXManifestGeneratorBase.cs b/Engine/Source/Programs/UnrealBuildTool/Platform/Windows/AppXManifestGeneratorBase.cs
--- a/Engine/Source/Programs/UnrealBuildTool/Platform/Windows/AppXManifestGeneratorBase.cs
+++ b/Engine/Source/Programs/UnrealBuildTool/Platform/Windows/AppXManifestGeneratorBase.cs
@@ -89,8 +89,7 @@ namespace UnrealBuildTool
 				for (int VersionElement = 0; VersionElement < NumVersionElements; VersionElement++)
 				{
 					string QuadElement = SplitVersionString[VersionElement];
-					int QuadValue = 0;
-					if (QuadElement.Length == 0 || !Int32.TryParse(QuadElement, out QuadValue))
+					if (QuadElement.Length == 0 || !Int32.TryParse(QuadElement, out int QuadValue))
 					{
 						CompletedVersionString += "0";
 					}
@@ -272,7 +271,7 @@ namespace UnrealBuildTool
 						continue;
 					}
 
-					string CultureId = "";
+					string CultureId;
 					if (UEStageIdToAppXCultureId.ContainsKey(StageId))
 					{
 						CultureId = UEStageIdToAppXCultureId[StageId];
diff --git a/Engine/Source/Programs/UnrealBuildTool/Platform/Windows/WindowsExports.cs b/Engine/Source/Programs/UnrealBuildTool/Platform/Windows/WindowsExports.cs
--- a/Engine/Source/Programs/UnrealBuildTool/Platform/Windows/WindowsExports.cs
+++ b/Engine/Source/Programs/UnrealBuildTool/Platform/Windows/WindowsExports.cs
@@ -120,9 +120,6 @@ namespace UnrealBuildTool
 		[SupportedOSPlatform("windows")]
 		public static bool TryGetToolChainDir(WindowsCompiler Compiler, string? CompilerVersion, UnrealArch Architecture, ILogger Logger, [NotNullWhen(true)] out VersionNumber? OutToolChainVersion, [NotNullWhen(true)] out DirectoryReference? OutToolChainDir, out DirectoryReference? OutRedistDir)
 		{
-			OutToolChainVersion = null;
-			OutToolChainDir = null;
-			OutRedistDir = null;
 			return MicrosoftPlatformSDK.TryGetToolChainDir(Compiler, CompilerVersion, Architecture, Logger, out OutToolChainVersion, out OutToolChainDir, out OutRedistDir);
 		}
 	}
diff --git a/Engine/Source/Programs/UnrealBuildTool/ProjectFiles/CMake/CMakefileGenerator.cs b/Engine/Source/Programs/UnrealBuildTool/ProjectFiles/CMake/CMakefileGenerator.cs
--- a/Engine/Source/Programs/UnrealBuildTool/ProjectFiles/CMake/CMakefileGenerator.cs
+++ b/Engine/Source/Programs/UnrealBuildTool/ProjectFiles/CMake/CMakefileGenerator.cs
@@ -611,7 +611,7 @@ namespace UnrealBuildTool
 		private static string GetIncludeDirectory(string IncludeDir, string ProjectDir)
 		{
 			string FullProjectPath = Path.GetFullPath(PrimaryProjectPath.FullName);
-			string FullPath = "";
+			string FullPath;
 			// Check for paths outside of both the engine and the project
 			if (Path.IsPathRooted(IncludeDir) &&
 				!IncludeDir.StartsWith(FullProjectPath) &&
diff --git a/Engine/Source/Programs/UnrealBuildTool/ProjectFiles/CodeLite/CodeLiteProject.cs b/Engine/Source/Programs/UnrealBuildTool/ProjectFiles/CodeLite/CodeLiteProject.cs
--- a/Engine/Source/Programs/UnrealBuildTool/ProjectFiles/CodeLite/CodeLiteProject.cs
+++ b/Engine/Source/Programs/UnrealBuildTool/ProjectFiles/CodeLite/CodeLiteProject.cs
@@ -3,6 +3,7 @@
 using System;
 using System.Collections.Generic;
 using System.IO;
+using System.Linq;
 using System.Xml.Linq;
 using EpicGames.Core;
 using Microsoft.Extensions.Logging;
@@ -19,16 +20,6 @@ namespace UnrealBuildTool
 			OnlyGameProject = InOnlyGameProject;
 		}
 
-		// Check if the XElement is empty.
-		bool IsEmpty(IEnumerable<XElement> en)
-		{
-			foreach (XElement c in en)
-			{
-				return false;
-			}
-			return true;
-		}
-
 		public override bool WriteProjectFile(List<UnrealTargetPlatform> InPlatforms, List<UnrealTargetConfiguration> InConfigurations, PlatformProjectGeneratorCollection PlatformProjectGenerators, ILogger Logger)
 		{
 			bool bSuccess = false;
@@ -162,7 +153,7 @@ namespace UnrealBuildTool
 						// Let's look if there is a virtual folder withint the current XElement.
 						//
 						IEnumerable<XElement> tests = root.Elements("VirtualDirectory");
-						if (IsEmpty(tests))
+						if (!tests.Any())
 						{
 							//
 							// No, then we have to create.
diff --git a/Engine/Source/Programs/UnrealBuildTool/ProjectFiles/Eddie/EddieGenerator.cs b/Engine/Source/Programs/UnrealBuildTool/ProjectFiles/Eddie/EddieGenerator.cs
--- a/Engine/Source/Programs/UnrealBuildTool/ProjectFiles/Eddie/EddieGenerator.cs
+++ b/Engine/Source/Programs/UnrealBuildTool/ProjectFiles/Eddie/EddieGenerator.cs
@@ -54,6 +54,7 @@ namespace UnrealBuildTool
 			WorksetDataContent.Append("# @Eddie Workset@" + ProjectFileGenerator.NewLine);
 			WorksetDataContent.Append("AddWorkset \"" + PrimaryProjectName + ".wkst\" \"" + PrimaryProjectPath + "\"" + ProjectFileGenerator.NewLine);
 
+#pragma warning disable IDE0059 // Remove unnecessary value assignment - this is a false positive, the null assignment is required as the lambda is recursive.
 			System.Action<String /*Path*/, List<PrimaryProjectFolder> /* Folders */>? AddProjectsFunction = null;
 			AddProjectsFunction = (Path, FolderList) =>
 				{
@@ -76,6 +77,7 @@ namespace UnrealBuildTool
 						WorksetDataContent.Append("EndFileGroup \"" + NewPath + "\"" + ProjectFileGenerator.NewLine);
 					}
 				};
+#pragma warning restore IDE0059
 			AddProjectsFunction(PrimaryProjectName, RootFolder.SubFolders);
 
 			string ProjectName = PrimaryProjectName;
diff --git a/Engine/Source/Programs/UnrealBuildTool/ProjectFiles/Eddie/EddieProject.cs b/Engine/Source/Programs/UnrealBuildTool/ProjectFiles/Eddie/EddieProject.cs
--- a/Engine/Source/Programs/UnrealBuildTool/ProjectFiles/Eddie/EddieProject.cs
+++ b/Engine/Source/Programs/UnrealBuildTool/ProjectFiles/Eddie/EddieProject.cs
@@ -179,8 +179,6 @@ namespace UnrealBuildTool
 
 		public override bool WriteProjectFile(List<UnrealTargetPlatform> InPlatforms, List<UnrealTargetConfiguration> InConfigurations, PlatformProjectGeneratorCollection PlatformProjectGenerators, ILogger Logger)
 		{
-			bool bSuccess = false;
-
 			//string TargetName = ProjectFilePath.GetFileNameWithoutExtension();
 
 			//FileReference GameProjectPath = null;
@@ -201,9 +199,7 @@ namespace UnrealBuildTool
 			ParseSourceFilesIntoGroups();
 			EmitProject(ProjectFileContent, Folders, Logger);
 
-			bSuccess = ProjectFileGenerator.WriteFileIfChanged(ProjectFilePath.FullName, ProjectFileContent.ToString(), Logger, new UTF8Encoding());
-
-			return bSuccess;
+			return ProjectFileGenerator.WriteFileIfChanged(ProjectFilePath.FullName, ProjectFileContent.ToString(), Logger, new UTF8Encoding());
 		}
 	}
 }
diff --git a/Engine/Source/Programs/UnrealBuildTool/ProjectFiles/KDevelop/KDevelopGenerator.cs b/Engine/Source/Programs/UnrealBuildTool/ProjectFiles/KDevelop/KDevelopGenerator.cs
--- a/Engine/Source/Programs/UnrealBuildTool/ProjectFiles/KDevelop/KDevelopGenerator.cs
+++ b/Engine/Source/Programs/UnrealBuildTool/ProjectFiles/KDevelop/KDevelopGenerator.cs
@@ -311,11 +311,8 @@ namespace UnrealBuildTool
 			}
 
 			// Remove duplicate paths from include dir and system include dir list
-			List<string> Tmp = new List<string>();
-			List<string> Stmp = new List<string>();
-
-			Tmp = IncludeDirectories.Distinct().ToList();
-			Stmp = SystemIncludeDirectories.Distinct().ToList();
+			List<string> Tmp = IncludeDirectories.Distinct().ToList();
+			List<string> Stmp = SystemIncludeDirectories.Distinct().ToList();
 
 			IncludeDirectories = Tmp.ToList();
 			SystemIncludeDirectories = Stmp.ToList();
@@ -363,9 +360,6 @@ namespace UnrealBuildTool
 		/// <param name="Logger">Logger for output</param>
 		private void WriteDefineSection(ref StringBuilder FileContent, ILogger Logger)
 		{
-			string Key = "";
-			string Value = "";
-
 			List<string> DefineHolder = new List<string>();
 
 			foreach (ProjectFile CurProject in GeneratedProjectFiles)
@@ -379,7 +373,7 @@ namespace UnrealBuildTool
 
 				foreach (string CurDefine in KDevelopProject.IntelliSensePreprocessorDefinitions)
 				{
-					SplitDefinitionAndValue(CurDefine, out Key, out Value);
+					SplitDefinitionAndValue(CurDefine, out string Key, out string Value);
 					if (String.IsNullOrEmpty(Value))
 					{
 						DefineHolder.Add(String.Format("{0} \\\n", Key));
@@ -392,9 +386,7 @@ namespace UnrealBuildTool
 			}
 
 			// Remove duplicates if they are present.
-			List<string> Tmp = new List<string>();
-			Tmp = DefineHolder.Distinct().ToList();
-			DefineHolder = Tmp.ToList();
+			DefineHolder = DefineHolder.Distinct().ToList();
 
 			foreach (string Def in DefineHolder)
 			{
diff --git a/Engine/Source/Programs/UnrealBuildTool/ProjectFiles/PlatformProjectGeneratorCollection.cs b/Engine/Source/Programs/UnrealBuildTool/ProjectFiles/PlatformProjectGeneratorCollection.cs
--- a/Engine/Source/Programs/UnrealBuildTool/ProjectFiles/PlatformProjectGeneratorCollection.cs
+++ b/Engine/Source/Programs/UnrealBuildTool/ProjectFiles/PlatformProjectGeneratorCollection.cs
@@ -33,8 +33,7 @@ namespace UnrealBuildTool
 		public void RegisterPlatformProjectGenerator(UnrealTargetPlatform InPlatform, PlatformProjectGenerator InProjectGenerator, ILogger Logger)
 		{
 			// Make sure the build platform is legal
-			UEBuildPlatform? BuildPlatform;
-			if (UEBuildPlatform.TryGetBuildPlatform(InPlatform, out BuildPlatform))
+			if (UEBuildPlatform.TryGetBuildPlatform(InPlatform, out _))
 			{
 				if (ProjectGeneratorDictionary.ContainsKey(InPlatform) == true)
 				{
diff --git a/Engine/Source/Programs/UnrealBuildTool/ProjectFiles/ProjectFile.cs b/Engine/Source/Programs/UnrealBuildTool/ProjectFiles/ProjectFile.cs
--- a/Engine/Source/Programs/UnrealBuildTool/ProjectFiles/ProjectFile.cs
+++ b/Engine/Source/Programs/UnrealBuildTool/ProjectFiles/ProjectFile.cs
@@ -209,8 +209,7 @@ namespace UnrealBuildTool
 			}
 
 			// Don't add duplicates
-			SourceFile? ExistingFile = null;
-			if (SourceFileMap.TryGetValue(FilePath, out ExistingFile))
+			if (SourceFileMap.TryGetValue(FilePath, out SourceFile? ExistingFile))
 			{
 				if (ExistingFile.BaseFolder != BaseFolder)
 				{
@@ -293,20 +292,17 @@ namespace UnrealBuildTool
 					// Go ahead and check to see if the definition already exists, but the value is different
 					bool AlreadyExists = false;
 
-					string Def, Value;
-					SplitDefinitionAndValue(CurDef, out Def, out Value);
+					SplitDefinitionAndValue(CurDef, out string Def, out _);
 
 					// Ignore any API macros being import/export; we'll assume they're valid across the whole project
 					if (Def.EndsWith("_API", StringComparison.Ordinal))
 					{
 						CurDef = Def + "=";
-						Value = "";
 					}
 
 					for (int DefineIndex = 0; DefineIndex < IntelliSensePreprocessorDefinitions.Count; ++DefineIndex)
 					{
-						string ExistingDef, ExistingValue;
-						SplitDefinitionAndValue(IntelliSensePreprocessorDefinitions[DefineIndex], out ExistingDef, out ExistingValue);
+						SplitDefinitionAndValue(IntelliSensePreprocessorDefinitions[DefineIndex], out string ExistingDef, out _);
 						if (ExistingDef == Def)
 						{
 							// Already exists, but the value is changing.  We don't bother clobbering values for existing defines for this project.
diff --git a/Engine/Source/Programs/UnrealBuildTool/ProjectFiles/ProjectFileGenerator.cs b/Engine/Source/Programs/UnrealBuildTool/ProjectFiles/ProjectFileGenerator.cs
--- a/Engine/Source/Programs/UnrealBuildTool/ProjectFiles/ProjectFileGenerator.cs
+++ b/Engine/Source/Programs/UnrealBuildTool/ProjectFiles/ProjectFileGenerator.cs
@@ -2506,8 +2506,7 @@ namespace UnrealBuildTool
 				if (bInProjectPlatformsList && (IncludeAllPlatforms || IsRequiredPlatform))
 				{
 					// If there is a build platform present, add it to the SupportedPlatforms list
-					UEBuildPlatform? BuildPlatform;
-					if (UEBuildPlatform.TryGetBuildPlatform(Platform, out BuildPlatform))
+					if (UEBuildPlatform.TryGetBuildPlatform(Platform, out _))
 					{
 						if (InstalledPlatformInfo.IsValidPlatform(Platform, EProjectType.Code))
 						{
diff --git a/Engine/Source/Programs/UnrealBuildTool/ProjectFiles/QMake/QMakefileGenerator.cs b/Engine/Source/Programs/UnrealBuildTool/ProjectFiles/QMake/QMakefileGenerator.cs
--- a/Engine/Source/Programs/UnrealBuildTool/ProjectFiles/QMake/QMakefileGenerator.cs
+++ b/Engine/Source/Programs/UnrealBuildTool/ProjectFiles/QMake/QMakefileGenerator.cs
@@ -65,7 +65,7 @@ namespace UnrealBuildTool
 		private void AddIncludeDirectory(ref List<string> IncludeDirectories, string IncludeDir, string ProjectDir)
 		{
 			string FullProjectPath = ProjectFileGenerator.PrimaryProjectPath.FullName;
-			string FullPath = "";
+			string FullPath;
 			if (IncludeDir.StartsWith("/") && !IncludeDir.StartsWith(FullProjectPath))
 			{
 				// Full path to a fulder outside of project
@@ -98,10 +98,10 @@ namespace UnrealBuildTool
 			StringBuilder QMakeDefinesPriFileContent = new StringBuilder();
 
 			string GameProjectPath = "";
-			string GameProjectFile = "";
+			string GameProjectFile;
 			string GameProjectRootPath = "";
 
-			string BuildCommand = "";
+			string BuildCommand;
 
 			string QMakeGameProjectFile = "";
 			string QMakeFilesDirectory = ".qmake";
@@ -150,7 +150,7 @@ namespace UnrealBuildTool
 
 			// Iterate through all the defines for the projects that are generated by 
 			// UnrealBuildTool
-			// !RAKE: move to seperate function
+			// !RAKE: move to separate function
 			QMakeDefinesPriFileContent.Append("DEFINES += \\\n");
 			foreach (ProjectFile CurProject in GeneratedProjectFiles)
 			{
@@ -163,10 +163,7 @@ namespace UnrealBuildTool
 
 				foreach (string CurDefine in QMakeProject.IntelliSensePreprocessorDefinitions)
 				{
-					string define = "";
-					string value = "";
-
-					SplitDefinitionAndValue(CurDefine, out define, out value);
+					SplitDefinitionAndValue(CurDefine, out string define, out string value);
 
 					if (!DefinesOnly.Contains(define))
 					{
@@ -219,7 +216,7 @@ namespace UnrealBuildTool
 
 			// Iterate through all the include paths that
 			// UnrealBuildTool generates
-			// !RAKE: Move to seperate function
+			// !RAKE: Move to separate function
 			QMakeIncludesPriFileContent.Append("INCLUDEPATH += \\\n");
 			foreach (string CurPath in IncludeDirectories)
 			{
diff --git a/Engine/Source/Programs/UnrealBuildTool/ProjectFiles/Rider/RiderProjectFileGenerator.cs b/Engine/Source/Programs/UnrealBuildTool/ProjectFiles/Rider/RiderProjectFileGenerator.cs
--- a/Engine/Source/Programs/UnrealBuildTool/ProjectFiles/Rider/RiderProjectFileGenerator.cs
+++ b/Engine/Source/Programs/UnrealBuildTool/ProjectFiles/Rider/RiderProjectFileGenerator.cs
@@ -428,8 +428,7 @@ namespace UnrealBuildTool
 
 		private void SetupSupportedPlatformsAndConfigurations(ILogger Logger)
 		{
-			string SupportedPlatformNames;
-			SetupSupportedPlatformsAndConfigurations(true, Logger, out SupportedPlatformNames);
+			SetupSupportedPlatformsAndConfigurations(true, Logger, out _);
 		}
 
 		public override bool GenerateProjectFiles(PlatformProjectGeneratorCollection PlatformProjectGenerators, string[] arguments, ILogger Logger)
@@ -491,14 +490,11 @@ namespace UnrealBuildTool
 		private void GatherProjects(PlatformProjectGeneratorCollection PlatformProjectGenerators,
 			List<FileReference> AllGameProjects, ILogger Logger)
 		{
-			ProjectFile? EngineProject = null;
-			List<ProjectFile>? GameProjects = null;
 			List<ProjectFile>? ModProjects = new List<ProjectFile>();
-			Dictionary<FileReference, ProjectFile>? ProgramProjects = null;
 
 			// Setup buildable projects for all targets
 			AddProjectsForAllTargets(PlatformProjectGenerators, AllGameProjects, Logger,
-				out EngineProject, out GameProjects, out ProgramProjects);
+				out ProjectFile? EngineProject, out List<ProjectFile>? GameProjects, out _);
 
 			AddProjectsForMods(GameProjects, ModProjects);
 			AddAllGameProjects(GameProjects);
diff --git a/Engine/Source/Programs/UnrealBuildTool/ProjectFiles/VisualStudioCode/VSCodeProjectFileGenerator.cs b/Engine/Source/Programs/UnrealBuildTool/ProjectFiles/VisualStudioCode/VSCodeProjectFileGenerator.cs
--- a/Engine/Source/Programs/UnrealBuildTool/ProjectFiles/VisualStudioCode/VSCodeProjectFileGenerator.cs
+++ b/Engine/Source/Programs/UnrealBuildTool/ProjectFiles/VisualStudioCode/VSCodeProjectFileGenerator.cs
@@ -1257,11 +1257,9 @@ namespace UnrealBuildTool
 			}
 
 			// Get store version
-			int StoreVersion = 1;
-			int StoreVersionArm64 = 1;
-			int StoreVersionOffsetArm64 = 0;
-			Ini.GetInt32("/Script/AndroidRuntimeSettings.AndroidRuntimeSettings", "StoreVersion", out StoreVersion);
-			Ini.GetInt32("/Script/AndroidRuntimeSettings.AndroidRuntimeSettings", "StoreVersionOffsetArm64", out StoreVersionOffsetArm64);
+			int StoreVersionArm64;
+			Ini.GetInt32("/Script/AndroidRuntimeSettings.AndroidRuntimeSettings", "StoreVersion", out int StoreVersion);
+			Ini.GetInt32("/Script/AndroidRuntimeSettings.AndroidRuntimeSettings", "StoreVersionOffsetArm64", out int StoreVersionOffsetArm64);
 			StoreVersionArm64 = StoreVersion + StoreVersionOffsetArm64;
 
 			DirectoryReference SymbolPathArm64 = DirectoryReference.Combine(
