Clean up the final IDE warnings in UBT and enable code style warnings for Analysis builds.

#Synced-CL 48288999

diff --git a/Engine/Source/Programs/UnrealBuildTool/Platform/Apple/AppleProjectGenerator.cs b/Engine/Source/Programs/UnrealBuildTool/Platform/Apple/AppleProjectGenerator.cs
--- a/Engine/Source/Programs/UnrealBuildTool/Platform/Apple/AppleProjectGenerator.cs
+++ b/Engine/Source/Programs/UnrealBuildTool/Platform/Apple/AppleProjectGenerator.cs
@@ -3,7 +3,6 @@
 using System;
 using System.Collections.Generic;
 using System.Diagnostics;
-using EpicGames.Core;
 using Microsoft.Extensions.Logging;
 
 namespace UnrealBuildTool
diff --git a/Engine/Source/Programs/UnrealBuildTool/Platform/TVOS/UEDeployTVOS.cs b/Engine/Source/Programs/UnrealBuildTool/Platform/TVOS/UEDeployTVOS.cs
--- a/Engine/Source/Programs/UnrealBuildTool/Platform/TVOS/UEDeployTVOS.cs
+++ b/Engine/Source/Programs/UnrealBuildTool/Platform/TVOS/UEDeployTVOS.cs
@@ -23,7 +23,7 @@ namespace UnrealBuildTool
 			return "TVOS";
 		}
 
-		public static bool GenerateTVOSPList(string ProjectDirectory, bool bIsUnrealGame, string GameName, bool bIsClient, string ProjectName, string InEngineDir, string AppDirectory, UnrealPluginLanguage? UPL, string? BundleID, ILogger Logger)
+		public static bool GenerateTVOSPList(string ProjectDirectory, bool bIsUnrealGame, string GameName, bool bIsClient, string ProjectName, string InEngineDir, string AppDirectory, UnrealPluginLanguage? UPL, string? BundleID)
 		{
 			// @todo tvos: THIS!
 
@@ -183,7 +183,7 @@ namespace UnrealBuildTool
 			{
 				return base.GeneratePList(ProjectFile, Config, ProjectDirectory, bIsUnrealGame, GameName, bIsClient, ProjectName, InEngineDir, AppDirectory, UPLScripts, BundleID, bBuildAsFramework);
 			}
-			return GenerateTVOSPList(ProjectDirectory, bIsUnrealGame, GameName, bIsClient, ProjectName, InEngineDir, AppDirectory, null, BundleID, Logger);
+			return GenerateTVOSPList(ProjectDirectory, bIsUnrealGame, GameName, bIsClient, ProjectName, InEngineDir, AppDirectory, null, BundleID);
 		}
 	}
 }
diff --git a/Engine/Source/Programs/UnrealBuildTool/Platform/Windows/PVSToolChain.cs b/Engine/Source/Programs/UnrealBuildTool/Platform/Windows/PVSToolChain.cs
--- a/Engine/Source/Programs/UnrealBuildTool/Platform/Windows/PVSToolChain.cs
+++ b/Engine/Source/Programs/UnrealBuildTool/Platform/Windows/PVSToolChain.cs
@@ -5,7 +5,6 @@ using System.Collections.Generic;
 using System.Diagnostics;
 using System.IO;
 using System.Linq;
-using System.Reflection;
 using System.Text;
 using System.Text.Json;
 using System.Text.RegularExpressions;
diff --git a/Engine/Source/Programs/UnrealBuildTool/Platform/Windows/VCCompileAction.cs b/Engine/Source/Programs/UnrealBuildTool/Platform/Windows/VCCompileAction.cs
--- a/Engine/Source/Programs/UnrealBuildTool/Platform/Windows/VCCompileAction.cs
+++ b/Engine/Source/Programs/UnrealBuildTool/Platform/Windows/VCCompileAction.cs
@@ -3,7 +3,6 @@
 using System;
 using System.Collections.Generic;
 using EpicGames.Core;
-using Microsoft.Extensions.Logging;
 using UnrealBuildBase;
 
 namespace UnrealBuildTool
diff --git a/Engine/Source/Programs/UnrealBuildTool/ProjectFiles/Xcode/XcodeProject.Utils.cs b/Engine/Source/Programs/UnrealBuildTool/ProjectFiles/Xcode/XcodeProject.Utils.cs
--- a/Engine/Source/Programs/UnrealBuildTool/ProjectFiles/Xcode/XcodeProject.Utils.cs
+++ b/Engine/Source/Programs/UnrealBuildTool/ProjectFiles/Xcode/XcodeProject.Utils.cs
@@ -7,7 +7,6 @@ using System.Linq;
 using System.Text;
 using System.Text.RegularExpressions;
 using EpicGames.Core;
-using Microsoft.Extensions.Logging;
 
 namespace UnrealBuildTool.XcodeProjectXcconfig
 {
diff --git a/Engine/Source/Programs/UnrealBuildTool/ProjectFiles/Xcode/XcodeProject.cs b/Engine/Source/Programs/UnrealBuildTool/ProjectFiles/Xcode/XcodeProject.cs
--- a/Engine/Source/Programs/UnrealBuildTool/ProjectFiles/Xcode/XcodeProject.cs
+++ b/Engine/Source/Programs/UnrealBuildTool/ProjectFiles/Xcode/XcodeProject.cs
@@ -3257,7 +3257,7 @@ namespace UnrealBuildTool.XcodeProjectXcconfig
 				else
 				{
 					// write metadata now so we can add them to the FileCollection
-					ConditionalWriteMetadataFiles(UnrealTargetPlatform.IOS);
+					ConditionalWriteMetadataFiles();
 
 					StringBuilder Content = new StringBuilder();
 
@@ -3302,11 +3302,8 @@ namespace UnrealBuildTool.XcodeProjectXcconfig
 			return bSuccess;
 		}
 
-		private UnrealTargetPlatform CurrentPlistPlatform;
-		private void ConditionalWriteMetadataFiles(UnrealTargetPlatform Platform)
+		private void ConditionalWriteMetadataFiles()
 		{
-			CurrentPlistPlatform = Platform;
-
 			// we now use templates or premade, no writing out here
 			foreach (MetadataItem Data in UnrealData.Metadata!.PlistFiles.Values)
 			{
diff --git a/Engine/Source/Programs/UnrealBuildTool/ProjectFiles/Xcode/XcodeProjectLegacy.cs b/Engine/Source/Programs/UnrealBuildTool/ProjectFiles/Xcode/XcodeProjectLegacy.cs
--- a/Engine/Source/Programs/UnrealBuildTool/ProjectFiles/Xcode/XcodeProjectLegacy.cs
+++ b/Engine/Source/Programs/UnrealBuildTool/ProjectFiles/Xcode/XcodeProjectLegacy.cs
@@ -1259,7 +1259,7 @@ namespace UnrealBuildTool.XcodeProjectLegacy
 					if (bCreateTVOSInfoPlist)
 					{
 						Directory.CreateDirectory(Path.GetDirectoryName(TVOSInfoPlistPath)!);
-						UEDeployTVOS.GenerateTVOSPList(ProjectPath.FullName, bIsUnrealGame, GameName, bIsClient, Config.BuildTarget, EngineDir.FullName, ProjectPath + "/Binaries/TVOS/Payload", null, BundleIdentifier, Logger);
+						UEDeployTVOS.GenerateTVOSPList(ProjectPath.FullName, bIsUnrealGame, GameName, bIsClient, Config.BuildTarget, EngineDir.FullName, ProjectPath + "/Binaries/TVOS/Payload", null, BundleIdentifier);
 					}
 				}
 			}
diff --git a/Engine/Source/Programs/UnrealBuildTool/System/ActionExecutor.cs b/Engine/Source/Programs/UnrealBuildTool/System/ActionExecutor.cs
--- a/Engine/Source/Programs/UnrealBuildTool/System/ActionExecutor.cs
+++ b/Engine/Source/Programs/UnrealBuildTool/System/ActionExecutor.cs
@@ -2,7 +2,6 @@
 
 using System;
 using System.Collections.Generic;
-using System.Reflection;
 using System.Threading.Tasks;
 using EpicGames.Core;
 using EpicGames.LogMatchers;
diff --git a/Engine/Source/Programs/UnrealBuildTool/UnrealBuildTool.csproj b/Engine/Source/Programs/UnrealBuildTool/UnrealBuildTool.csproj
--- a/Engine/Source/Programs/UnrealBuildTool/UnrealBuildTool.csproj
+++ b/Engine/Source/Programs/UnrealBuildTool/UnrealBuildTool.csproj
@@ -28,6 +28,7 @@
     <DebugSymbols>true</DebugSymbols>
   </PropertyGroup>
   <PropertyGroup Condition=" '$(Configuration)|$(Platform)' == 'Analyze|AnyCPU' ">
+    <EnforceCodeStyleInBuild>true</EnforceCodeStyleInBuild>
     <RunAnalyzersDuringBuild>True</RunAnalyzersDuringBuild>
     <TreatWarningsAsErrors>false</TreatWarningsAsErrors>
   </PropertyGroup>
