UnrealBuildTool: Retry clang compiler crashes locally with UBA

#rnx
#Synced-CL 47605596

diff --git a/Engine/Source/Programs/UnrealBuildTool/Executors/UnrealBuildAccelerator/UBAExecutor.cs b/Engine/Source/Programs/UnrealBuildTool/Executors/UnrealBuildAccelerator/UBAExecutor.cs
--- a/Engine/Source/Programs/UnrealBuildTool/Executors/UnrealBuildAccelerator/UBAExecutor.cs
+++ b/Engine/Source/Programs/UnrealBuildTool/Executors/UnrealBuildAccelerator/UBAExecutor.cs
@@ -807,6 +807,10 @@ namespace UnrealBuildTool
 							{
 								return rerun("C1001");
 							}
+							else if (exitCode != 0 && logLines.Any(x => x.Contains("clang frontend command failed")))
+							{
+								return rerun("Clang Exception");
+							}
 							else if (exitCode != 0 && logLines.Any(x => x.Contains("cannot overwrite the original file", StringComparison.Ordinal)))
 							{
 								// We have seen this once on the farm and have no idea how it can happen.
