UnrealBuildTool: Print out reason a build is cancelled and add an opt-in option to cancel after num errors, defaulted off

#rnx
#Synced-CL 47605596

diff --git a/Engine/Source/Programs/UnrealBuildTool/Executors/ImmediateActionQueue.cs b/Engine/Source/Programs/UnrealBuildTool/Executors/ImmediateActionQueue.cs
--- a/Engine/Source/Programs/UnrealBuildTool/Executors/ImmediateActionQueue.cs
+++ b/Engine/Source/Programs/UnrealBuildTool/Executors/ImmediateActionQueue.cs
@@ -200,6 +200,11 @@ namespace UnrealBuildTool
 		/// </summary>
 		public bool StopCompilationAfterErrors = false;
 
+		/// <summary>
+		/// When enabled, will stop compiling targets after a number of compile errors occur.
+		/// </summary>
+		public int StopCompilationAfterNumErrors = Int32.MaxValue;
+
 		public int CompletedActions => _completedActions;
 		public int CacheHitActions => _cacheHitActions;
 		public int CacheMissActions => _cacheMissActions;
@@ -975,6 +980,7 @@ namespace UnrealBuildTool
 		/// </summary>
 		private void LogActions()
 		{
+			int numErrors = 0;
 			for (; ; )
 			{
 				int[]? actionsToLog = null;
@@ -986,7 +992,7 @@ namespace UnrealBuildTool
 					}
 					else
 					{
-						actionsToLog = _actionsToLog.ToArray();
+						actionsToLog = [.. _actionsToLog];
 						_actionsToLog.Clear();
 					}
 				}
@@ -1002,9 +1008,13 @@ namespace UnrealBuildTool
 					LogAction(Actions[index].Action, result, Actions.Length);
 
 					// Cancel all other pending tasks
-					if (StopCompilationAfterErrors && (result == null || result.ExitCode != 0))
+					if (result == null || result.ExitCode != 0)
 					{
-						CancellationTokenSource.Cancel();
+						numErrors++;
+						if (StopCompilationAfterErrors || numErrors >= StopCompilationAfterNumErrors)
+						{
+							CancellationTokenSource.Cancel();
+						}
 					}
 				}
 			}
diff --git a/Engine/Source/Programs/UnrealBuildTool/Executors/ParallelExecutor.cs b/Engine/Source/Programs/UnrealBuildTool/Executors/ParallelExecutor.cs
--- a/Engine/Source/Programs/UnrealBuildTool/Executors/ParallelExecutor.cs
+++ b/Engine/Source/Programs/UnrealBuildTool/Executors/ParallelExecutor.cs
@@ -65,6 +65,12 @@ namespace UnrealBuildTool
 		[XmlConfigFile]
 		protected static bool bStopCompilationAfterErrors = false;
 
+		/// <summary>
+		/// When set, will stop compiling targets after a set number of errors occurs.
+		/// </summary>
+		[XmlConfigFile]
+		protected static int StopCompilationAfterNumErrors = Int32.MaxValue;
+
 		/// <summary>
 		/// Whether to show compilation times along with worst offenders or not.
 		/// </summary>
@@ -183,6 +189,7 @@ namespace UnrealBuildTool
 				ShowPerActionCompilationTimes = bShowPerActionCompilationTimes,
 				CompactOutput = bCompactOutput,
 				StopCompilationAfterErrors = bStopCompilationAfterErrors,
+				StopCompilationAfterNumErrors = StopCompilationAfterNumErrors,
 			};
 		}
 
diff --git a/Engine/Source/Programs/UnrealBuildTool/Executors/UnrealBuildAccelerator/UBAExecutor.cs b/Engine/Source/Programs/UnrealBuildTool/Executors/UnrealBuildAccelerator/UBAExecutor.cs
--- a/Engine/Source/Programs/UnrealBuildTool/Executors/UnrealBuildAccelerator/UBAExecutor.cs
+++ b/Engine/Source/Programs/UnrealBuildTool/Executors/UnrealBuildAccelerator/UBAExecutor.cs
@@ -741,6 +741,7 @@ namespace UnrealBuildTool
 			long peakProcessMemory = 0;
 			long peakTotalMemory = 0;
 			bool shouldCancel = false;
+			string cancelReason = "user cancellation";
 
 			scheduler.SetProcessFinishedCallback((info) =>
 			{
@@ -920,8 +921,14 @@ namespace UnrealBuildTool
 						}
 					}
 
-					if (bStopCompilationAfterErrors)
+					if (bStopCompilationAfterErrors || StopCompilationAfterNumErrors <= 0)
 					{
+						cancelReason = "stop compilation after any errors";
+						shouldCancel = true;
+					}
+					else if (failedActions >= StopCompilationAfterNumErrors)
+					{
+						cancelReason = $"stop compilation after {StopCompilationAfterNumErrors} errors";
 						shouldCancel = true;
 					}
 				}
@@ -1033,6 +1040,10 @@ namespace UnrealBuildTool
 						DateTime.UtcNow,
 						peakProcessMemory, peakTotalMemory);
 				}
+				else
+				{
+					logger.LogInformation("Cancelled execution after {Actions}/{TotalActions} actions ({Reason})", succeededActions + failedActions + cacheHitActions, totalActions, cancelReason);
+				}
 				return res;
 			}
 			finally
