Mac Build: Allow keeping the generated executable name instead of forcing it to UnrealEditor when generating project files for non-monolithic editors (e.g., LiveLinkHub).

diff --git a/Engine/Source/Programs/UnrealBuildTool/Platform/Mac/UEBuildMac.cs b/Engine/Source/Programs/UnrealBuildTool/Platform/Mac/UEBuildMac.cs
--- a/Engine/Source/Programs/UnrealBuildTool/Platform/Mac/UEBuildMac.cs
+++ b/Engine/Source/Programs/UnrealBuildTool/Platform/Mac/UEBuildMac.cs
@@ -33,6 +33,11 @@ namespace UnrealBuildTool
 		[ConfigFile(ConfigHierarchyType.Engine, "/Script/MacTargetPlatform.MacTargetSettings", "bEnableRayTracing")]
 		public bool bEnableRayTracing = false;
 
+		/// <summary>
+		/// When generating a project file for an editor target, forces the exe name to match the target name instead of being overridden to 'UnrealEditor'. 
+		/// </summary>
+		public bool bForceTargetNameAsExeName { get; set; } = false;
+
 		/// <summary>
 		/// Constructor
 		/// </summary>
diff --git a/Engine/Source/Programs/UnrealBuildTool/ProjectFiles/Xcode/XcodeProject.cs b/Engine/Source/Programs/UnrealBuildTool/ProjectFiles/Xcode/XcodeProject.cs
--- a/Engine/Source/Programs/UnrealBuildTool/ProjectFiles/Xcode/XcodeProject.cs
+++ b/Engine/Source/Programs/UnrealBuildTool/ProjectFiles/Xcode/XcodeProject.cs
@@ -715,9 +715,12 @@ namespace UnrealBuildTool.XcodeProjectXcconfig
 									{
 										// Figure out if this is a monolithic build
 										bool bShouldCompileMonolithic = BuildPlatform.ShouldCompileMonolithicBinary(Platform);
+										bool bForceRetainTargetNameAsExeName = false;
 										try
 										{
-											bShouldCompileMonolithic |= (ProjectTarget.CreateRulesDelegate(Platform, Configuration).LinkType == TargetLinkType.Monolithic);
+											TargetRules ProjectTargetRules = ProjectTarget.CreateRulesDelegate(Platform, Configuration);
+											bShouldCompileMonolithic |= ProjectTargetRules.LinkType == TargetLinkType.Monolithic;
+											bForceRetainTargetNameAsExeName = ProjectTargetRules.MacPlatform.bForceTargetNameAsExeName;
 										}
 										catch (BuildException)
 										{
@@ -752,7 +755,9 @@ namespace UnrealBuildTool.XcodeProjectXcconfig
 											}
 
 											string ExeName = TargetName;
-											if (!bShouldCompileMonolithic && ProjectTarget.TargetRules.Type != TargetType.Program)
+											if (!bShouldCompileMonolithic
+												&& ProjectTarget.TargetRules.Type != TargetType.Program
+												&& !bForceRetainTargetNameAsExeName)
 											{
 												// Figure out what the compiled binary will be called so that we can point the IDE to the correct file
 												if (ProjectTarget.TargetRules.Type != TargetType.Game)
