Move UBT files related to Actions to their own directory.

#Synced-CL 48387743

diff --git a/Engine/Source/Programs/UnrealBuildTool/Actions/Action.cs b/Engine/Source/Programs/UnrealBuildTool/Actions/Action.cs
new file mode 100644
--- /dev/null
+++ b/Engine/Source/Programs/UnrealBuildTool/Actions/Action.cs
@@ -0,0 +1,882 @@
+// Copyright Epic Games, Inc. All Rights Reserved.
+
+using System;
+using System.Collections.Generic;
+using System.Diagnostics;
+using System.Linq;
+using System.Threading;
+using EpicGames.Core;
+using Microsoft.Extensions.Logging;
+using UnrealBuildBase;
+
+namespace UnrealBuildTool
+{
+	/// <summary>
+	/// Enumerates build action types.
+	/// </summary>
+	enum ActionType
+	{
+		BuildProject,
+
+		GatherModuleDependencies,
+
+		CompileModuleInterface,
+
+		Compile,
+
+		CreateAppBundle,
+
+		GenerateDebugInfo,
+
+		Link,
+
+		WriteMetadata,
+
+		PostBuildStep,
+
+		ParseTimingInfo,
+	}
+
+	/// <summary>
+	/// Defines the action's support for artifacts
+	/// </summary>
+	[Flags]
+	enum ArtifactMode : byte
+	{
+
+		/// <summary>
+		/// Cached artifacts aren't enabled
+		/// </summary>
+		None = 0,
+
+		/// <summary>
+		/// If set, the outputs should be cached but respect the flags below
+		/// </summary>
+		Enabled = 1 << 0,
+
+		/// <summary>
+		/// Absolute file paths must be used when recording inputs and outputs
+		/// </summary>
+		AbsolutePath = 1 << 1,
+
+		/// <summary>
+		/// For actions that can't be cached, by setting this flag, the inputs for the action in question
+		/// will be used as additional inputs for any action that references outputs of this action.  
+		/// For example, PCH files shouldn't be cached, but compiles that use the PCH should still be 
+		/// valid if consider the inputs used to create the PCH file.
+		/// </summary>
+		PropagateInputs = 1 << 2,
+	}
+
+	interface IExternalAction
+	{
+		/// <summary>
+		/// The type of this action (for debugging purposes).
+		/// </summary>
+		ActionType ActionType { get; }
+
+		/// <summary>
+		/// Artifact support for this step
+		/// </summary>
+		ArtifactMode ArtifactMode { get; }
+
+		/// <summary>
+		/// Every file this action depends on.  These files need to exist and be up to date in order for this action to even be considered
+		/// </summary>
+		IEnumerable<FileItem> PrerequisiteItems { get; }
+
+		/// <summary>
+		/// The files that this action produces after completing
+		/// </summary>
+		IEnumerable<FileItem> ProducedItems { get; }
+
+		/// <summary>
+		/// Items that should be deleted before running this action
+		/// </summary>
+		IEnumerable<FileItem> DeleteItems { get; }
+
+		/// <summary>
+		/// Root paths for this action (generally engine root project root, toolchain root, sdk root)
+		/// </summary>
+		CppRootPaths RootPaths { get; }
+
+		/// <summary>
+		/// For C++ source files, specifies a dependency list file used to check changes to header files
+		/// </summary>
+		FileItem? DependencyListFile { get; }
+
+		/// <summary>
+		/// Directory from which to execute the program to create produced items
+		/// </summary>
+		DirectoryReference WorkingDirectory { get; }
+
+		/// <summary>
+		/// The command to run to create produced items
+		/// </summary>
+		FileReference CommandPath { get; }
+
+		/// <summary>
+		/// Command-line parameters to pass to the program
+		/// </summary>
+		string CommandArguments { get; }
+
+		/// <summary>
+		/// Version of the command used for this action. This will be considered a dependency.
+		/// </summary>
+		string CommandVersion { get; }
+
+		/// <summary>
+		/// Optional friendly description of the type of command being performed, for example "Compile" or "Link".  Displayed by some executors.
+		/// </summary>
+		string CommandDescription { get; }
+
+		/// <summary>
+		/// Human-readable description of this action that may be displayed as status while invoking the action.  This is often the name of the file being compiled, or an executable file name being linked.  Displayed by some executors.
+		/// </summary>
+		string StatusDescription { get; }
+
+		/// <summary>
+		/// True if this action is allowed to be run on a remote machine when a distributed build system is being used, such as XGE
+		/// </summary>
+		bool bCanExecuteRemotely { get; }
+
+		/// <summary>
+		/// True if this action is allowed to be run on a remote machine with SNDBS. Files with #import directives must be compiled locally. Also requires bCanExecuteRemotely = true.
+		/// </summary>
+		bool bCanExecuteRemotelyWithSNDBS { get; }
+
+		/// <summary>
+		/// True if this action is allowed to be run on a remote machine with XGE. Files with #import directives must be compiled locally. Also requires bCanExecuteRemotely = true.
+		/// </summary>
+		bool bCanExecuteRemotelyWithXGE { get; }
+
+		/// <summary>
+		/// True if this action can be executed by Unreal Build Accelerator
+		/// </summary>
+		bool bCanExecuteInUBA { get; }
+
+		/// <summary>
+		/// True if this action can be executed on a cross-architecture Unreal Build Accelerator helper. If no cross-architecture toolchain is found this setting has no effect
+		/// </summary>
+		bool bCanExecuteInUBACrossArchitecture { get; }
+
+		/// <summary>
+		/// True if this action is using the Clang compiler.  Some build systems may be able to optimize for this case.
+		/// </summary>
+		bool bIsClangCompiler { get; }
+
+		/// <summary>
+		/// True if the action executor should delete output files on an error.
+		/// May not work for XGE, SN-DBS, and FASTBuild executors.
+		/// </summary>
+		bool bDeleteProducedItemsOnError { get; }
+
+		/// <summary>
+		/// True if the action executor should force warning output to fail the action when the exit code would otherwise be zero.
+		/// Used if the command to be run does not support a warnings as error flag.
+		/// May not work for XGE, SN-DBS, and FASTBuild executors.
+		/// </summary>
+		bool bForceWarningsAsError { get; }
+
+		/// <summary>
+		/// Whether we should log this action, whether executed locally or remotely.  This is useful for actions that take time
+		/// but invoke tools without any console output.
+		/// </summary>
+		bool bShouldOutputStatusDescription { get; }
+
+		/// <summary>
+		/// Whether we should output log of this action, whether executed locally or remotely.  This is useful for actions that spam
+		/// logs to console output with non critical data
+		/// </summary>
+		bool bShouldOutputLog { get; }
+
+		/// <summary>
+		/// True if any libraries produced by this action should be considered 'import libraries'
+		/// </summary>
+		bool bProducesImportLibrary { get; }
+
+		/// <summary>
+		/// Whether changes in the command line used to generate these produced items should invalidate the action
+		/// </summary>
+		bool bUseActionHistory { get; }
+
+		/// <summary>
+		/// True if this action should be scheduled early. Some actions can take a long time (pch) and need to start very early
+		/// If set to true it will propagate high priority to all its prerequisite tasks
+		/// </summary>
+		bool bIsHighPriority { get; }
+
+		/// <summary>
+		/// Used to determine how much weight(CPU and Memory work) this action is.
+		/// </summary>
+		double Weight { get; }
+
+		/// <summary>
+		/// Used to determine which cache bucket should be used for this action
+		/// Using different cache buckets is an optimization where we want one build to use as few buckets as possible while 
+		/// two different builds should not use the same bucket if there is no chance their cache entries will be the same
+		/// </summary>
+		public uint CacheBucket { get; }
+	}
+
+	/// <summary>
+	/// A build action.
+	/// </summary>
+	class Action : IExternalAction
+	{
+		///
+		/// Preparation and Assembly (serialized)
+		/// 
+
+		/// <summary>
+		/// The type of this action (for debugging purposes).
+		/// </summary>
+		public ActionType ActionType { get; set; }
+
+		/// <summary>
+		/// Artifact support for this step
+		/// </summary>
+		public ArtifactMode ArtifactMode { get; set; } = ArtifactMode.None;
+
+		/// <summary>
+		/// Every file this action depends on.  These files need to exist and be up to date in order for this action to even be considered
+		/// </summary>
+		public SortedSet<FileItem> PrerequisiteItems { get; set; } = new(FileItem.SortByPath);
+
+		/// <summary>
+		/// The files that this action produces after completing
+		/// </summary>
+		public SortedSet<FileItem> ProducedItems { get; set; } = new(FileItem.SortByPath);
+
+		/// <summary>
+		/// Items that should be deleted before running this action
+		/// </summary>
+		public SortedSet<FileItem> DeleteItems { get; set; } = new(FileItem.SortByPath);
+
+		/// <summary>
+		/// Root paths for this action (generally engine root project root, toolchain root, sdk root)
+		/// Order needs to be maintained, so a list is used
+		/// </summary>
+		public CppRootPaths RootPaths { get; set; } = new();
+
+		/// <summary>
+		/// For C++ source files, specifies a dependency list file used to check changes to header files
+		/// </summary>
+		public FileItem? DependencyListFile { get; set; }
+
+		/// <summary>
+		/// Directory from which to execute the program to create produced items
+		/// </summary>
+		public DirectoryReference WorkingDirectory { get; set; } = null!;
+
+		/// <summary>
+		/// The command to run to create produced items
+		/// </summary>
+		public FileReference CommandPath { get; set; } = null!;
+
+		/// <summary>
+		/// Command-line parameters to pass to the program
+		/// </summary>
+		public string CommandArguments { get; set; } = null!;
+
+		/// <summary>
+		/// Version of the command used for this action. This will be considered a dependency.
+		/// </summary>
+		public string CommandVersion { get; set; } = "0";
+
+		/// <summary>
+		/// Optional friendly description of the type of command being performed, for example "Compile" or "Link".  Displayed by some executors.
+		/// </summary>
+		public string CommandDescription { get; set; } = null!;
+
+		/// <summary>
+		/// Human-readable description of this action that may be displayed as status while invoking the action.  This is often the name of the file being compiled, or an executable file name being linked.  Displayed by some executors.
+		/// </summary>
+		public string StatusDescription { get; set; } = "...";
+
+		/// <summary>
+		/// True if this action is allowed to be run on a remote machine when a distributed build system is being used, such as XGE
+		/// </summary>
+		public bool bCanExecuteRemotely { get; set; } = false;
+
+		/// <summary>
+		/// True if this action is allowed to be run on a remote machine with SNDBS. Files with #import directives must be compiled locally. Also requires bCanExecuteRemotely = true.
+		/// </summary>
+		public bool bCanExecuteRemotelyWithSNDBS { get; set; } = true;
+
+		/// <summary>
+		/// True if this action is allowed to be run on a remote machine with XGE. Files with #import directives must be compiled locally. Also requires bCanExecuteRemotely = true.
+		/// </summary>
+		public bool bCanExecuteRemotelyWithXGE { get; set; } = true;
+
+		/// <summary>
+		/// True if this action can be executed in uba
+		/// </summary>
+		public bool bCanExecuteInUBA { get; set; } = true;
+
+		/// <summary>
+		/// True if this action can be executed on a cross-architecture Unreal Build Accelerator helper. If no cross-architecture toolchain is found this setting has no effect
+		/// </summary>
+		public bool bCanExecuteInUBACrossArchitecture { get; set; } = true;
+
+		/// <summary>
+		/// True if this action is using the Clang compiler.  Some build systems may be able to optimize for this case.
+		/// </summary>
+		public bool bIsClangCompiler { get; set; } = false;
+
+		/// <summary>
+		/// True if the action executor should delete output files on an error.
+		/// May not work for XGE, SN-DBS, and FASTBuild executors.
+		/// </summary>
+		public bool bDeleteProducedItemsOnError { get; set; } = false;
+
+		/// <summary>
+		/// True if the action executor should force warning output to fail the action when the exit code would otherwise be zero.
+		/// Used if the command to be run does not support a warnings as error flag.
+		/// May not work for XGE, SN-DBS, and FASTBuild executors.
+		/// </summary>
+		public bool bForceWarningsAsError { get; set; } = false;
+
+		/// <summary>
+		/// Whether we should log this action, whether executed locally or remotely.  This is useful for actions that take time
+		/// but invoke tools without any console output.
+		/// </summary>
+		public bool bShouldOutputStatusDescription { get; set; } = true;
+
+		/// <inheritdoc/>
+		public bool bShouldOutputLog { get; set; } = true;
+
+		/// <summary>
+		/// True if any libraries produced by this action should be considered 'import libraries'
+		/// </summary>
+		public bool bProducesImportLibrary { get; set; } = false;
+
+		/// <inheritdoc/>
+		public bool bUseActionHistory { get; set; } = true;
+
+		/// <inheritdoc/>
+		public bool bIsHighPriority { get; set; } = false;
+
+		/// <inheritdoc/>
+		public double Weight { get; set; } = 1.0;
+
+		/// <inheritdoc/>
+		public uint CacheBucket { get; set; }
+
+		IEnumerable<FileItem> IExternalAction.PrerequisiteItems => PrerequisiteItems;
+		IEnumerable<FileItem> IExternalAction.ProducedItems => ProducedItems;
+		IEnumerable<FileItem> IExternalAction.DeleteItems => DeleteItems;
+        public Action(ActionType InActionType)
+		{
+			ActionType = InActionType;
+
+			// link actions are going to run locally on SN-DBS so don't try to distribute them as that generates warnings for missing tool templates
+			if (ActionType == ActionType.Link)
+			{
+				bCanExecuteRemotelyWithSNDBS = false;
+			}
+		}
+
+		public Action(IExternalAction InOther)
+		{
+			ActionType = InOther.ActionType;
+			ArtifactMode = InOther.ArtifactMode;
+			PrerequisiteItems = new (InOther.PrerequisiteItems, FileItem.SortByPath);
+			ProducedItems = new(InOther.ProducedItems, FileItem.SortByPath);
+			DeleteItems = new(InOther.DeleteItems, FileItem.SortByPath);
+			RootPaths = new(InOther.RootPaths);
+			DependencyListFile = InOther.DependencyListFile;
+			WorkingDirectory = InOther.WorkingDirectory;
+			CommandPath = InOther.CommandPath;
+			CommandArguments = InOther.CommandArguments;
+			CommandVersion = InOther.CommandVersion;
+			CommandDescription = InOther.CommandDescription;
+			StatusDescription = InOther.StatusDescription;
+			bCanExecuteRemotely = InOther.bCanExecuteRemotely;
+			bCanExecuteRemotelyWithSNDBS = InOther.bCanExecuteRemotelyWithSNDBS;
+			bCanExecuteRemotelyWithXGE = InOther.bCanExecuteRemotelyWithXGE;
+			bCanExecuteInUBA = InOther.bCanExecuteInUBA;
+			bCanExecuteInUBACrossArchitecture = InOther.bCanExecuteInUBACrossArchitecture;
+			bIsClangCompiler = InOther.bIsClangCompiler;
+			bDeleteProducedItemsOnError = InOther.bDeleteProducedItemsOnError;
+			bForceWarningsAsError = InOther.bForceWarningsAsError;
+			bShouldOutputStatusDescription = InOther.bShouldOutputStatusDescription;
+			bShouldOutputLog = InOther.bShouldOutputLog;
+			bProducesImportLibrary = InOther.bProducesImportLibrary;
+			bUseActionHistory = InOther.bUseActionHistory;
+			bIsHighPriority = InOther.bIsHighPriority;
+			Weight = InOther.Weight;
+			CacheBucket = InOther.CacheBucket;
+		}
+
+		public Action(BinaryArchiveReader Reader)
+		{
+			ActionType = (ActionType)Reader.ReadByte();
+			ArtifactMode = (ArtifactMode)Reader.ReadByte();
+			WorkingDirectory = Reader.ReadDirectoryReferenceNotNull();
+			CommandPath = Reader.ReadFileReference();
+			CommandArguments = Reader.ReadString()!;
+			CommandVersion = Reader.ReadString()!;
+			CommandDescription = Reader.ReadString()!;
+			StatusDescription = Reader.ReadString()!;
+			bCanExecuteRemotely = Reader.ReadBool();
+			bCanExecuteRemotelyWithSNDBS = Reader.ReadBool();
+			bCanExecuteRemotelyWithXGE = Reader.ReadBool();
+			bCanExecuteInUBA = Reader.ReadBool();
+			bCanExecuteInUBACrossArchitecture = Reader.ReadBool();
+			bIsClangCompiler = Reader.ReadBool();
+			bDeleteProducedItemsOnError = Reader.ReadBool();
+			bForceWarningsAsError = Reader.ReadBool();
+			bShouldOutputStatusDescription = Reader.ReadBool();
+			bShouldOutputLog = Reader.ReadBool();
+			bProducesImportLibrary = Reader.ReadBool();
+			PrerequisiteItems = Reader.ReadSortedSet(() => Reader.ReadFileItem(), FileItem.SortByPath)!;
+			ProducedItems = Reader.ReadSortedSet(() => Reader.ReadFileItem(), FileItem.SortByPath)!;
+			DeleteItems = Reader.ReadSortedSet(() => Reader.ReadFileItem(), FileItem.SortByPath)!;
+			RootPaths = new(Reader);
+			DependencyListFile = Reader.ReadFileItem();
+			bUseActionHistory = Reader.ReadBool();
+			bIsHighPriority = Reader.ReadBool();
+			Weight = Reader.ReadDouble();
+			CacheBucket = Reader.ReadUnsignedInt();
+		}
+
+		/// <summary>
+		/// ISerializable: Called when serialized to report additional properties that should be saved
+		/// </summary>
+		public void Write(BinaryArchiveWriter Writer)
+		{
+			Writer.WriteByte((byte)ActionType);
+			Writer.WriteByte((byte)ArtifactMode);
+			Writer.WriteDirectoryReference(WorkingDirectory);
+			Writer.WriteFileReference(CommandPath);
+			Writer.WriteString(CommandArguments);
+			Writer.WriteString(CommandVersion);
+			Writer.WriteString(CommandDescription);
+			Writer.WriteString(StatusDescription);
+			Writer.WriteBool(bCanExecuteRemotely);
+			Writer.WriteBool(bCanExecuteRemotelyWithSNDBS);
+			Writer.WriteBool(bCanExecuteRemotelyWithXGE);
+			Writer.WriteBool(bCanExecuteInUBA);
+			Writer.WriteBool(bCanExecuteInUBACrossArchitecture);
+			Writer.WriteBool(bIsClangCompiler);
+			Writer.WriteBool(bDeleteProducedItemsOnError);
+			Writer.WriteBool(bForceWarningsAsError);
+			Writer.WriteBool(bShouldOutputStatusDescription);
+			Writer.WriteBool(bShouldOutputLog);
+			Writer.WriteBool(bProducesImportLibrary);
+			Writer.WriteSortedSet(PrerequisiteItems, Item => Writer.WriteFileItem(Item));
+			Writer.WriteSortedSet(ProducedItems, Item => Writer.WriteFileItem(Item));
+			Writer.WriteSortedSet(DeleteItems, Item => Writer.WriteFileItem(Item));
+			RootPaths.Write(Writer);
+			Writer.WriteFileItem(DependencyListFile);
+			Writer.WriteBool(bUseActionHistory);
+			Writer.WriteBool(bIsHighPriority);
+			Writer.WriteDouble(Weight);
+			Writer.WriteUnsignedInt(CacheBucket);
+		}
+
+		/// <summary>
+		/// Writes an action to a json file
+		/// </summary>
+		/// <param name="Object">The object to parse</param>
+		public static Action ImportJson(JsonObject Object)
+		{
+			Action Action = new(Object.GetEnumField<ActionType>("Type"));
+
+			if (Object.TryGetEnumField("ArtifactMode", out ArtifactMode ArtifactMode))
+			{
+				Action.ArtifactMode = ArtifactMode;
+			}
+
+			if (Object.TryGetStringField("WorkingDirectory", out string? WorkingDirectory))
+			{
+				Action.WorkingDirectory = new DirectoryReference(WorkingDirectory);
+			}
+
+			if (Object.TryGetStringField("CommandPath", out string? CommandPath))
+			{
+				Action.CommandPath = new FileReference(CommandPath);
+			}
+
+			if (Object.TryGetStringField("CommandArguments", out string? CommandArguments))
+			{
+				Action.CommandArguments = CommandArguments;
+			}
+
+			if (Object.TryGetStringField("CommandVersion", out string? CommandVersion))
+			{
+				Action.CommandVersion = CommandVersion;
+			}
+
+			if (Object.TryGetStringField("CommandDescription", out string? CommandDescription))
+			{
+				Action.CommandDescription = CommandDescription;
+			}
+
+			if (Object.TryGetStringField("StatusDescription", out string? StatusDescription))
+			{
+				Action.StatusDescription = StatusDescription;
+			}
+
+			if (Object.TryGetBoolField("bCanExecuteRemotely", out bool bCanExecuteRemotely))
+			{
+				Action.bCanExecuteRemotely = bCanExecuteRemotely;
+			}
+
+			if (Object.TryGetBoolField("bCanExecuteRemotelyWithSNDBS", out bool bCanExecuteRemotelyWithSNDBS))
+			{
+				Action.bCanExecuteRemotelyWithSNDBS = bCanExecuteRemotelyWithSNDBS;
+			}
+
+			if (Object.TryGetBoolField("bCanExecuteRemotelyWithXGE", out bool bCanExecuteRemotelyWithXGE))
+			{
+				Action.bCanExecuteRemotelyWithXGE = bCanExecuteRemotelyWithXGE;
+			}
+
+			if (Object.TryGetBoolField("bCanExecuteInUBA", out bool bCanExecuteInUBA))
+			{
+				Action.bCanExecuteInUBA = bCanExecuteInUBA;
+			}
+
+			if (Object.TryGetBoolField("bCanExecuteInUBACrossArchitecture", out bool bCanExecuteInUBACrossArchitecture))
+			{
+				Action.bCanExecuteInUBACrossArchitecture = bCanExecuteInUBACrossArchitecture;
+			}
+
+			if (Object.TryGetBoolField("bIsClangCompiler", out bool bIsClangCompiler))
+			{
+				Action.bIsClangCompiler = bIsClangCompiler;
+			}
+
+			if (Object.TryGetBoolField("bDeleteProducedItemsOnError", out bool bDeleteProducedItemsOnError))
+			{
+				Action.bDeleteProducedItemsOnError = bDeleteProducedItemsOnError;
+			}
+
+			if (Object.TryGetBoolField("bForceWarningsAsError", out bool bForceWarningsAsError))
+			{
+				Action.bForceWarningsAsError = bForceWarningsAsError;
+			}
+
+			if (Object.TryGetBoolField("bShouldOutputStatusDescription", out bool bShouldOutputStatusDescription))
+			{
+				Action.bShouldOutputStatusDescription = bShouldOutputStatusDescription;
+			}
+
+			if (Object.TryGetBoolField("bShouldOutputLog", out bool bShouldOutputLog)) 
+			{
+				Action.bShouldOutputLog= bShouldOutputLog;
+			}
+
+			if (Object.TryGetBoolField("bProducesImportLibrary", out bool bProducesImportLibrary))
+			{
+				Action.bProducesImportLibrary = bProducesImportLibrary;
+			}
+
+			if (Object.TryGetStringArrayField("PrerequisiteItems", out string[]? PrerequisiteItems))
+			{
+				Action.PrerequisiteItems.UnionWith(PrerequisiteItems.Select(x => FileItem.GetItemByPath(x)));
+			}
+
+			if (Object.TryGetStringArrayField("ProducedItems", out string[]? ProducedItems))
+			{
+				Action.ProducedItems.UnionWith(ProducedItems.Select(x => FileItem.GetItemByPath(x)));
+			}
+
+			if (Object.TryGetStringArrayField("DeleteItems", out string[]? DeleteItems))
+			{
+				Action.DeleteItems.UnionWith(DeleteItems.Select(x => FileItem.GetItemByPath(x)));
+			}
+
+			if (Object.TryGetObjectField("RootPaths", out JsonObject? RootPaths))
+			{
+				Action.RootPaths.Read(RootPaths);
+			}
+
+			if (Object.TryGetStringField("DependencyListFile", out string? DependencyListFile))
+			{
+				Action.DependencyListFile = FileItem.GetItemByPath(DependencyListFile);
+			}
+
+			if (Object.TryGetDoubleField("Weight", out double Weight))
+			{
+				Action.Weight = Weight;
+			}
+
+			if (Object.TryGetUnsignedIntegerField("CacheBucket", out uint CacheBucket))
+			{
+				Action.CacheBucket = CacheBucket;
+			}
+
+			return Action;
+		}
+
+		public override string ToString()
+		{
+			string ReturnString = "";
+			if (CommandPath != null)
+			{
+				ReturnString += CommandPath + " - ";
+			}
+			if (CommandArguments != null)
+			{
+				ReturnString += CommandArguments;
+			}
+			return ReturnString;
+		}
+	}
+
+	/// <summary>
+	/// Interface that is used when -SingleFile=xx is used on the cmd line to make specific files compile
+	/// Actions that has this interface can be used to create real actions that can compile the specific files
+	/// </summary>
+	interface ISpecificFileAction
+	{
+		/// <summary>
+		/// The directory this action can create actions for
+		/// </summary>
+		DirectoryReference RootDirectory { get; }
+
+		/// <summary>
+		/// Creates an action for a specific file. It can return null if for example file extension is not handled
+		/// </summary>
+		IExternalAction? CreateAction(FileItem SourceFile, ILogger Logger);
+	}
+
+	/// <summary>
+	/// Extension methods for action classes
+	/// </summary>
+	static class ActionExtensions
+	{
+		/// <summary>
+		/// Writes an action to a json file
+		/// </summary>
+		/// <param name="Action">The action to write</param>
+		/// <param name="LinkedActionToId">Map of action to unique id</param>
+		/// <param name="Writer">Writer to receive the output</param>
+		public static void ExportJson(this LinkedAction Action, Dictionary<LinkedAction, int> LinkedActionToId, JsonWriter Writer)
+		{
+			Writer.WriteValue("Id", LinkedActionToId[Action]);
+			Writer.WriteEnumValue("Type", Action.ActionType);
+			Writer.WriteEnumValue("ArtifactMode", Action.ArtifactMode);
+			Writer.WriteValue("WorkingDirectory", Action.WorkingDirectory.FullName);
+			Writer.WriteValue("CommandPath", Action.CommandPath.FullName);
+			Writer.WriteValue("CommandArguments", Action.CommandArguments);
+			Writer.WriteValue("CommandVersion", Action.CommandVersion);
+			Writer.WriteValue("CommandDescription", Action.CommandDescription);
+			Writer.WriteValue("StatusDescription", Action.StatusDescription);
+			Writer.WriteValue("bCanExecuteRemotely", Action.bCanExecuteRemotely);
+			Writer.WriteValue("bCanExecuteRemotelyWithSNDBS", Action.bCanExecuteRemotelyWithSNDBS);
+			Writer.WriteValue("bCanExecuteRemotelyWithXGE", Action.bCanExecuteRemotelyWithXGE);
+			Writer.WriteValue("bCanExecuteInUBA", Action.bCanExecuteInUBA);
+			Writer.WriteValue("bCanExecuteInUBACrossArchitecture", Action.bCanExecuteInUBACrossArchitecture);
+			Writer.WriteValue("bIsClangCompiler", Action.bIsClangCompiler);
+			Writer.WriteValue("bDeleteProducedItemsOnError", Action.bDeleteProducedItemsOnError);
+			Writer.WriteValue("bForceWarningsAsError", Action.bForceWarningsAsError);			
+			Writer.WriteValue("bShouldOutputStatusDescription", Action.bShouldOutputStatusDescription);
+			Writer.WriteValue("bShouldOutputLog", Action.bShouldOutputLog);
+			Writer.WriteValue("bProducesImportLibrary", Action.bProducesImportLibrary);
+			Writer.WriteValue("Weight", Action.Weight);
+			Writer.WriteValue("CacheBucket", Action.CacheBucket);
+
+			Writer.WriteArrayStart("PrerequisiteActions");
+			foreach (LinkedAction PrerequisiteAction in Action.PrerequisiteActions)
+			{
+				Writer.WriteValue(LinkedActionToId[PrerequisiteAction]);
+			}
+			Writer.WriteArrayEnd();
+
+			Writer.WriteArrayStart("ProducedItems");
+			foreach (FileItem ProducedItem in Action.ProducedItems)
+			{
+				Writer.WriteValue(ProducedItem.AbsolutePath);
+			}
+			Writer.WriteArrayEnd();
+
+			Writer.WriteArrayStart("DeleteItems");
+			foreach (FileItem DeleteItem in Action.DeleteItems)
+			{
+				Writer.WriteValue(DeleteItem.AbsolutePath);
+			}
+			Writer.WriteArrayEnd();
+
+			Writer.WriteObjectStart("RootPaths");
+			Action.RootPaths.Write(Writer);
+			Writer.WriteObjectEnd();
+
+			if (Action.DependencyListFile != null)
+			{
+				Writer.WriteValue("DependencyListFile", Action.DependencyListFile.AbsolutePath);
+			}
+		}
+
+		/// <summary>
+		/// Determine if this Action ignores output conflicts.
+		/// </summary>
+		/// <param name="Action"> The Action to check</param>
+		/// <returns>True if conflicts are ignored, else false.</returns>
+		public static bool IgnoreConflicts(this IExternalAction Action) =>
+			Action.ActionType == ActionType.WriteMetadata ||
+			Action.ActionType == ActionType.CreateAppBundle;
+	}
+
+	/// <summary>
+	/// Default serializer for <see cref="Action"/> instances
+	/// </summary>
+	class DefaultActionSerializer : ActionSerializerBase<Action>
+	{
+		/// <inheritdoc/>
+		public override Action Read(BinaryArchiveReader Reader) => new Action(Reader);
+
+		/// <inheritdoc/>
+		public override void Write(BinaryArchiveWriter Writer, Action Action) => Action.Write(Writer);
+	}
+
+	/// <summary>
+	/// Information about an action queued to be executed
+	/// </summary>
+	[DebuggerDisplay("{StatusDescription}")]
+	class LinkedAction : IExternalAction
+	{
+		/// <summary>
+		/// The inner action instance
+		/// </summary>
+		public IExternalAction Inner;
+
+		/// <summary>
+		/// A target that this action contributes to
+		/// </summary>
+		public TargetDescriptor? Target;
+
+		/// <summary>
+		/// Set of other actions that this action depends on. This set is built when the action graph is linked.
+		/// </summary>
+		public HashSet<LinkedAction> PrerequisiteActions = null!;
+
+		/// <summary>
+		/// Total number of actions depending on this one.
+		/// </summary>
+		public int NumTotalDependentActions = 0;
+
+		/// <summary>
+		/// Additional field used for sorting that can be used to control sorting of link actions from outside
+		/// </summary>
+		public int SortIndex = 0;
+
+		/// <summary>
+		/// True if this action should be scheduled early. Some actions can take a long time (pch) and need to start very early
+		/// If set to 1 it will propagate high priority to all its prerequisite tasks. (Note it is an int to be able to do interlocked operations)
+		/// </summary>
+		public int IsHighPriority;
+
+		/// <summary>
+		/// If set, will be output whenever the group differs to the last executed action. Set when executing multiple targets at once.
+		/// </summary>
+		public SortedSet<string> GroupNames = [];
+
+		#region Wrapper implementation of IAction
+
+		public ActionType ActionType => Inner.ActionType;
+		public ArtifactMode ArtifactMode => Inner.ArtifactMode;
+		public IEnumerable<FileItem> PrerequisiteItems => Inner.PrerequisiteItems;
+		public IEnumerable<FileItem> ProducedItems => Inner.ProducedItems;
+		public IEnumerable<FileItem> DeleteItems => Inner.DeleteItems;
+		public CppRootPaths RootPaths => Inner.RootPaths;
+		public FileItem? DependencyListFile => Inner.DependencyListFile;
+		public DirectoryReference WorkingDirectory => Inner.WorkingDirectory;
+		public FileReference CommandPath => Inner.CommandPath;
+		public string CommandArguments => Inner.CommandArguments;
+		public string CommandVersion => Inner.CommandVersion;
+		public string CommandDescription => Inner.CommandDescription;
+		public string StatusDescription => Inner.StatusDescription;
+		public bool bCanExecuteRemotely => Inner.bCanExecuteRemotely;
+		public bool bCanExecuteRemotelyWithSNDBS => Inner.bCanExecuteRemotelyWithSNDBS;
+		public bool bCanExecuteRemotelyWithXGE => Inner.bCanExecuteRemotelyWithXGE;
+		public bool bCanExecuteInUBA => Inner.bCanExecuteInUBA;
+		public bool bCanExecuteInUBACrossArchitecture => Inner.bCanExecuteInUBACrossArchitecture;
+		public bool bIsClangCompiler => Inner.bIsClangCompiler;
+		public bool bDeleteProducedItemsOnError => Inner.bDeleteProducedItemsOnError;
+		public bool bForceWarningsAsError => Inner.bForceWarningsAsError;
+		public bool bShouldOutputStatusDescription => Inner.bShouldOutputStatusDescription;
+		public bool bShouldOutputLog => Inner.bShouldOutputLog;
+		public bool bProducesImportLibrary => Inner.bProducesImportLibrary;
+		public bool bUseActionHistory => Inner.bUseActionHistory;
+		public bool bIsHighPriority => IsHighPriority != 0;
+		public double Weight => Inner.Weight;
+		public uint CacheBucket => Inner.CacheBucket;
+
+		#endregion
+
+		/// <summary>
+		/// Constructor
+		/// </summary>
+		/// <param name="Inner">The inner action instance</param>
+		/// <param name="Target"></param>
+		public LinkedAction(IExternalAction Inner, TargetDescriptor? Target)
+		{
+			this.Inner = Inner;
+			this.Target = Target;
+			IsHighPriority = Inner.bIsHighPriority ? 1 : 0;
+		}
+
+		/// <summary>
+		/// Increment the number of dependents, recursively
+		/// </summary>
+		/// <param name="VisitedActions">Set of visited actions</param>
+		/// <param name="bIsHighPriority">Propagate high priority to this and prerequisite actions</param>
+		public void IncrementDependentCount(HashSet<LinkedAction> VisitedActions, bool bIsHighPriority)
+		{
+			if (VisitedActions.Add(this))
+			{
+				if (bIsHighPriority)
+				{
+					Interlocked.Exchange(ref IsHighPriority, 1);
+				}
+
+				Interlocked.Increment(ref NumTotalDependentActions);
+				foreach (LinkedAction PrerequisiteAction in PrerequisiteActions)
+				{
+					PrerequisiteAction.IncrementDependentCount(VisitedActions, bIsHighPriority);
+				}
+			}
+		}
+
+		/// <summary>
+		/// Compares two actions based on total number of sort index, priority and dependent items, descending.
+		/// </summary>
+		/// <param name="A">Action to compare</param>
+		/// <param name="B">Action to compare</param>
+		public static int Compare(LinkedAction A, LinkedAction B)
+		{
+			if (A.SortIndex != B.SortIndex)
+			{
+				return Math.Sign(A.SortIndex - B.SortIndex);
+			}
+
+			if (A.IsHighPriority != B.IsHighPriority)
+			{
+				return B.IsHighPriority - A.IsHighPriority;
+			}
+
+			// Primary sort criteria is total number of dependent files, up to max depth.
+			if (B.NumTotalDependentActions != A.NumTotalDependentActions)
+			{
+				return Math.Sign(B.NumTotalDependentActions - A.NumTotalDependentActions);
+			}
+
+			// Secondary sort criteria is directory name of first produced item.
+			// This will sort actions from same module closer to each other
+			// Also sort in alphabetic order since when building with unity files it is likely that the last unity file is the fastest to compile
+			// and we want longer actions to run earlier
+			int Result = String.Compare(A.ProducedItems.FirstOrDefault()?.FullName, B.ProducedItems.FirstOrDefault()?.FullName, StringComparison.Ordinal);
+			if (Result != 0)
+			{
+				return Math.Sign(Result);
+			}
+
+			// Third sort criteria is number of pre-requisites.
+			return Math.Sign(B.PrerequisiteItems.Count() - A.PrerequisiteItems.Count());
+		}
+	}
+}
diff --git a/Engine/Source/Programs/UnrealBuildTool/Actions/ActionExecutor.cs b/Engine/Source/Programs/UnrealBuildTool/Actions/ActionExecutor.cs
new file mode 100644
--- /dev/null
+++ b/Engine/Source/Programs/UnrealBuildTool/Actions/ActionExecutor.cs
@@ -0,0 +1,106 @@
+// Copyright Epic Games, Inc. All Rights Reserved.
+
+using System;
+using System.Collections.Generic;
+using System.Threading.Tasks;
+using EpicGames.Core;
+using EpicGames.LogMatchers;
+using Microsoft.Extensions.Logging;
+using UnrealBuildTool.Artifacts;
+
+namespace UnrealBuildTool
+{
+	abstract class ActionExecutor : IDisposable
+	{
+		public abstract string Name
+		{
+			get;
+		}
+
+		protected static double MemoryPerActionBytesOverride
+		{
+			get;
+			private set;
+		} = 0.0;
+
+		readonly LogEventParser Parser;
+
+		public ActionExecutor(ILogger Logger)
+		{
+			Parser = new LogEventParser(Logger);
+			Parser.AddMatchers(LogMatcherCollections.UBTMatchers);
+		}
+
+		public void Dispose()
+		{
+			Dispose(true);
+			GC.SuppressFinalize(this);
+		}
+
+		protected virtual void Dispose(bool disposing)
+		{
+			if (disposing)
+			{
+				Parser.Dispose();
+			}
+		}
+
+		/// <summary>
+		/// Allow targets to override the expected amount of memory required for compiles, used to control the number
+		/// of parallel action processes.
+		/// </summary>
+		/// <param name="MemoryPerActionOverrideGB"></param>
+		public static void SetMemoryPerActionOverride(double MemoryPerActionOverrideGB)
+		{
+			MemoryPerActionBytesOverride = Math.Max(MemoryPerActionBytesOverride, MemoryPerActionOverrideGB * 1024 * 1024 * 1024);
+		}
+
+		/// <summary>
+		/// Execute the given actions
+		/// </summary>
+		/// <param name="ActionsToExecute">Actions to be executed</param>
+		/// <param name="Logger">Logger for output</param>
+		/// <param name="actionArtifactCache">Cache used to read/write action outputs.</param>
+		/// <returns>True if the build succeeded, false otherwise</returns>
+		public abstract Task<bool> ExecuteActionsAsync(IEnumerable<LinkedAction> ActionsToExecute, ILogger Logger, IActionArtifactCache? actionArtifactCache = null);
+
+		/// <summary>
+		/// When completed, get a populated telemetry event for this executor
+		/// </summary>
+		/// <returns></returns>
+		public abstract TelemetryExecutorEvent? GetTelemetryEvent();
+
+		/// <summary>
+		/// Post a telemetry event if enabled
+		/// </summary>
+		public void PostTelemetryEvent()
+		{
+			TelemetryExecutorEvent? telemetryEvent = GetTelemetryEvent();
+			if (telemetryEvent != null)
+			{
+				TelemetryService.Get().RecordEvent(telemetryEvent);
+			}
+		}
+
+		/// <summary>
+		/// Will verify that produced items exists on disk
+		/// </summary>
+		public virtual bool VerifyOutputs => true;
+
+		protected void WriteToolOutput(string Line)
+		{
+			lock (Parser)
+			{
+				Parser.WriteLine(Line);
+			}
+		}
+
+		protected void FlushToolOutput()
+		{
+			lock (Parser)
+			{
+				Parser.Flush();
+			}
+		}
+	}
+}
diff --git a/Engine/Source/Programs/UnrealBuildTool/Actions/ActionGraph.cs b/Engine/Source/Programs/UnrealBuildTool/Actions/ActionGraph.cs
new file mode 100644
--- /dev/null
+++ b/Engine/Source/Programs/UnrealBuildTool/Actions/ActionGraph.cs
@@ -0,0 +1,1134 @@
+// Copyright Epic Games, Inc. All Rights Reserved.
+
+using System;
+using System.Collections.Concurrent;
+using System.Collections.Generic;
+using System.Collections.Immutable;
+using System.Diagnostics;
+using System.Globalization;
+using System.IO;
+using System.Linq;
+using System.Text;
+using System.Text.Json;
+using System.Text.Json.Serialization;
+using System.Threading.Tasks;
+using EpicGames.Core;
+using Microsoft.Extensions.Logging;
+using OpenTracing.Util;
+using UnrealBuildBase;
+using UnrealBuildTool.Artifacts;
+
+namespace UnrealBuildTool
+{
+	internal static class ActionGraph
+	{
+		/// <summary>
+		/// Enum describing why an Action is in conflict with another Action
+		/// </summary>
+		[Flags]
+		private enum ActionConflictReasonFlags : byte
+		{
+			None = 0,
+			ActionType = 1 << 0,
+			PrerequisiteItems = 1 << 1,
+			DeleteItems = 1 << 2,
+			DependencyListFile = 1 << 3,
+			WorkingDirectory = 1 << 4,
+			CommandPath = 1 << 5,
+			CommandArguments = 1 << 6,
+		};
+
+		/// <summary>
+		/// Links the actions together and sets up their dependencies
+		/// </summary>
+		/// <param name="Actions">List of actions in the graph</param>
+		/// <param name="Logger">Logger for output</param>
+		/// <param name="Sort">Optional sorting of actions</param>
+		public static void Link(List<LinkedAction> Actions, ILogger Logger, bool Sort = true)
+		{
+			using ITimelineEvent _ = Timeline.ScopeEvent($"LinkActions");
+
+			// Build a map from item to its producing action
+			Dictionary<FileItem, LinkedAction> ItemToProducingAction = new (Actions.Count * 3);
+			foreach (LinkedAction Action in Actions)
+			{
+				foreach (FileItem ProducedItem in Action.ProducedItems)
+				{
+					ItemToProducingAction[ProducedItem] = Action;
+				}
+			}
+
+			// Check for cycles
+			DetectActionGraphCycles(Actions, ItemToProducingAction, Logger);
+
+			// Use this map to add all the prerequisite actions
+			foreach (LinkedAction Action in Actions)
+			{
+				Action.PrerequisiteActions = new HashSet<LinkedAction>();
+				foreach (FileItem PrerequisiteItem in Action.PrerequisiteItems)
+				{
+					if (ItemToProducingAction.TryGetValue(PrerequisiteItem, out LinkedAction? PrerequisiteAction))
+					{
+						Action.PrerequisiteActions.Add(PrerequisiteAction);
+					}
+				}
+			}
+
+			// Sort the action graph
+			if (Sort)
+			{
+				SortActionList(Actions);
+			}
+		}
+
+		/// <summary>
+		/// Checks a set of actions for conflicts (ie. different actions producing the same output items)
+		/// </summary>
+		/// <param name="Actions">The set of actions to check</param>
+		/// <param name="Logger">Logger for output</param>
+		public static void CheckForConflicts(IEnumerable<IExternalAction> Actions, ILogger Logger)
+		{
+			bool bResult = true;
+
+			Dictionary<FileItem, IExternalAction> ItemToProducingAction = new Dictionary<FileItem, IExternalAction>();
+			foreach (IExternalAction Action in Actions)
+			{
+				HashSet<IExternalAction> Conflicted = new HashSet<IExternalAction>();
+				foreach (FileItem ProducedItem in Action.ProducedItems)
+				{
+					if (ItemToProducingAction.TryGetValue(ProducedItem, out IExternalAction? ExistingAction))
+					{
+						if (!Conflicted.Contains(ExistingAction))
+						{
+							bResult &= CheckForConflicts(ExistingAction, Action, Logger);
+							Conflicted.Add(ExistingAction);
+						}
+					}
+					else
+					{
+						ItemToProducingAction.Add(ProducedItem, Action);
+					}
+				}
+			}
+
+			if (!bResult)
+			{
+				throw new CompilationResultException(CompilationResult.ActionGraphInvalid, "Action graph is invalid; unable to continue. See log for additional details.");
+			}
+		}
+
+		/// <summary>
+		/// Finds conflicts between two actions, and prints them to the log
+		/// </summary>
+		/// <param name="A">The first action</param>
+		/// <param name="B">The second action</param>
+		/// <param name="Logger">Logger for output</param>
+		/// <returns>True if no conflicts were found, false otherwise.</returns>
+		private static bool CheckForConflicts(IExternalAction A, IExternalAction B, ILogger Logger)
+		{
+			ActionConflictReasonFlags Reason = ActionConflictReasonFlags.None;
+			if (A.ActionType != B.ActionType)
+			{
+				Reason |= ActionConflictReasonFlags.ActionType;
+			}
+
+			if (!A.PrerequisiteItems.SequenceEqual(B.PrerequisiteItems))
+			{
+				Reason |= ActionConflictReasonFlags.PrerequisiteItems;
+			}
+
+			if (!A.DeleteItems.SequenceEqual(B.DeleteItems))
+			{
+				Reason |= ActionConflictReasonFlags.DeleteItems;
+			}
+
+			if (A.DependencyListFile != B.DependencyListFile)
+			{
+				Reason |= ActionConflictReasonFlags.DependencyListFile;
+			}
+
+			if (A.WorkingDirectory != B.WorkingDirectory)
+			{
+				Reason |= ActionConflictReasonFlags.WorkingDirectory;
+			}
+
+			if (A.CommandPath != B.CommandPath)
+			{
+				Reason |= ActionConflictReasonFlags.CommandPath;
+			}
+
+			if (A.CommandArguments != B.CommandArguments)
+			{
+				Reason |= ActionConflictReasonFlags.CommandArguments;
+			}
+
+			if (Reason != ActionConflictReasonFlags.None)
+			{
+				LogConflict(A, B, Reason, Logger);
+				return false;
+			}
+
+			return true;
+		}
+
+		private class LogActionActionTypeConverter : JsonConverter<ActionType>
+		{
+			public override ActionType Read(ref Utf8JsonReader Reader, Type TypeToConvert, JsonSerializerOptions Options)
+			{
+				throw new NotImplementedException();
+			}
+
+			public override void Write(Utf8JsonWriter Writer, ActionType Value, JsonSerializerOptions Options)
+			{
+				Writer.WriteStringValue(Value.ToString());
+			}
+		}
+
+		private class LogActionFileItemConverter : JsonConverter<FileItem>
+		{
+			public override FileItem Read(ref Utf8JsonReader Reader, Type TypeToConvert, JsonSerializerOptions Options)
+			{
+				throw new NotImplementedException();
+			}
+
+			public override void Write(Utf8JsonWriter Writer, FileItem Value, JsonSerializerOptions Options)
+			{
+				Writer.WriteStringValue(Value.FullName);
+			}
+		}
+
+		private class LogActionDirectoryReferenceConverter : JsonConverter<DirectoryReference>
+		{
+			public override DirectoryReference Read(ref Utf8JsonReader Reader, Type TypeToConvert, JsonSerializerOptions Options)
+			{
+				throw new NotImplementedException();
+			}
+
+			public override void Write(Utf8JsonWriter Writer, DirectoryReference Value, JsonSerializerOptions Options)
+			{
+				Writer.WriteStringValue(Value.FullName);
+			}
+		}
+
+		private class LogActionFileReferenceConverter : JsonConverter<FileReference>
+		{
+			public override FileReference Read(ref Utf8JsonReader Reader, Type TypeToConvert, JsonSerializerOptions Options)
+			{
+				throw new NotImplementedException();
+			}
+
+			public override void Write(Utf8JsonWriter Writer, FileReference Value, JsonSerializerOptions Options)
+			{
+				Writer.WriteStringValue(Value.FullName);
+			}
+		}
+
+		/// <summary>
+		/// Adds the description of a merge error to an output message
+		/// </summary>
+		/// <param name="A">The first action with the conflict</param>
+		/// <param name="B">The second action with the conflict</param>
+		/// <param name="Reason">Enum flags for which properties are in conflict</param>
+		/// <param name="Logger">Logger for output</param>
+		static void LogConflict(IExternalAction A, IExternalAction B, ActionConflictReasonFlags Reason, ILogger Logger)
+		{
+			// Convert some complex types in IExternalAction to strings when printing json
+			JsonSerializerOptions Options = new JsonSerializerOptions
+			{
+				WriteIndented = true,
+				DefaultIgnoreCondition = JsonIgnoreCondition.WhenWritingNull,
+				Converters =
+				{
+					new LogActionActionTypeConverter(),
+					new LogActionFileItemConverter(),
+					new LogActionDirectoryReferenceConverter(),
+					new LogActionFileReferenceConverter(),
+				},
+			};
+
+			string AJson = JsonSerializer.Serialize(A, Options);
+			string BJson = JsonSerializer.Serialize(B, Options);
+			string AJsonPath = Path.Combine(Path.GetTempPath(), "UnrealBuildTool", $"{IoHash.Compute(Encoding.Default.GetBytes(AJson)).GetHashCode():X}") + ".json";
+			string BJsonPath = Path.Combine(Path.GetTempPath(), "UnrealBuildTool", $"{IoHash.Compute(Encoding.Default.GetBytes(BJson)).GetHashCode():X}") + ".json";
+
+			Directory.CreateDirectory(Path.Combine(Path.GetTempPath(), "UnrealBuildTool"));
+			File.WriteAllText(AJsonPath, AJson);
+			File.WriteAllText(BJsonPath, BJson);
+
+			Logger.LogError("Unable to merge actions '{StatusA}' and '{StatusB}': {Reason} are different", A.StatusDescription, B.StatusDescription, Reason);
+			Logger.LogInformation("  First Action: {AJson}", AJson);
+			Logger.LogInformation("  Second Action: {BJson}", BJson);
+			Logger.LogInformation("  First Action json written to '{AJsonPath}'", AJsonPath);
+			Logger.LogInformation("  Second Action json written to '{BJsonPath}'", BJsonPath);
+		}
+
+		/// <summary>
+		/// Builds a list of actions that need to be executed to produce the specified output items.
+		/// </summary>
+		public static List<LinkedAction> GetActionsToExecute(List<LinkedAction> Actions,
+			CppDependencyCache CppDependencies, ActionHistory History, bool bIgnoreOutdatedImportLibraries, ILogger Logger)
+		{
+			using (GlobalTracer.Instance.BuildSpan("ActionGraph.GetActionsToExecute()").StartActive())
+			{
+				// For all targets, build a set of all actions that are outdated.
+				ConcurrentDictionary<LinkedAction, bool> OutdatedActionDictionary = new ConcurrentDictionary<LinkedAction, bool>();
+				GatherAllOutdatedActions(Actions, History, OutdatedActionDictionary, CppDependencies, bIgnoreOutdatedImportLibraries, Logger);
+
+				// Build a list of actions that are both needed for this target and outdated.
+				return Actions.Where(Action => OutdatedActionDictionary[Action]).ToList();
+			}
+		}
+
+		/// <summary>
+		/// Checks that there aren't any intermediate files longer than the max allowed path length
+		/// </summary>
+		/// <param name="BuildConfiguration">The build configuration</param>
+		/// <param name="Actions">List of actions in the graph</param>
+		/// <param name="Logger"></param>
+		public static void CheckPathLengths(BuildConfiguration BuildConfiguration, IEnumerable<IExternalAction> Actions, ILogger Logger)
+		{
+			if (BuildHostPlatform.Current.Platform == UnrealTargetPlatform.Win64)
+			{
+				const int MAX_PATH = 260;
+
+				bool ShouldFail(FileItem item)
+				{
+					return item.Location.FullName.Length >= MAX_PATH &&
+						!(item.Location.ContainsName("Intermediate", 0) && item.Location.ContainsName("H", 1)); // Ignore -IncludeHeader items
+				}
+
+				bool ShouldWarn(FileItem item)
+				{
+					if (item.Location.FullName.Length <= Unreal.RootDirectory.FullName.Length + BuildConfiguration.MaxNestedPathLength)
+					{
+						return false;
+					}
+
+					if (!item.Location.IsUnderDirectory(Unreal.RootDirectory))
+					{
+						return false;
+					}
+
+					if (!item.Location.FullName.Contains("/Restricted/NotForLicensees/", StringComparison.OrdinalIgnoreCase)) //Be more relaxed for internal only code
+					{
+						return false;
+					}
+
+					if (item.Location.FullName.Contains("/Intermediate/", StringComparison.OrdinalIgnoreCase))
+					{
+						if (item.Location.ContainsName("H", 1)) // Ignore -IncludeHeader items
+						{
+							return false;
+						}
+
+						if (item.Location.FullName.EndsWith(".i.PVS-Studio.log")) // Ignore PVS Studio intermediate items
+						{
+							return false;
+						}
+					}
+
+					return true;
+				}
+
+				HashSet<FileReference> FailPaths = new();
+				HashSet<FileReference> WarnPaths = new();
+				foreach (IExternalAction Action in Actions)
+				{
+					foreach (FileItem PrerequisiteItem in Action.PrerequisiteItems)
+					{
+						if (ShouldFail(PrerequisiteItem))
+						{
+							FailPaths.Add(PrerequisiteItem.Location);
+						}
+
+						if (ShouldWarn(PrerequisiteItem))
+						{
+							WarnPaths.Add(PrerequisiteItem.Location);
+						}
+					}
+
+					foreach (FileItem ProducedItem in Action.ProducedItems)
+					{
+						if (ShouldFail(ProducedItem))
+						{
+							FailPaths.Add(ProducedItem.Location);
+						}
+
+						if (ShouldWarn(ProducedItem))
+						{
+							WarnPaths.Add(ProducedItem.Location);
+						}
+					}
+				}
+
+				if (FailPaths.Count > 0)
+				{
+					StringBuilder Message = new StringBuilder();
+					Message.Append($"The following action paths are longer than {MAX_PATH} characters. Please move the engine to a directory with a shorter path.");
+					foreach (FileReference Path in FailPaths.OrderBy(x => x.FullName))
+					{
+						Message.Append($"\n[{Path.FullName.Length.ToString()} characters] {Path}");
+					}
+
+					throw new BuildException(Message.ToString());
+				}
+				if (WarnPaths.Count > 0)
+				{
+					StringBuilder Message = new StringBuilder();
+					Message.Append($"Detected paths more than {BuildConfiguration.MaxNestedPathLength.ToString()} characters below UE root directory. This may cause portability issues due to the {MAX_PATH.ToString()} character maximum path length on Windows:\n");
+					foreach (FileReference Path in WarnPaths.OrderBy(x => x.FullName))
+					{
+						string RelativePath = Path.MakeRelativeTo(Unreal.RootDirectory);
+						Message.Append($"\n[{RelativePath.Length.ToString()} characters] {RelativePath}");
+					}
+
+					Message.Append($"\n\nConsider setting {nameof(ModuleRules.ShortName)} = ... in module *.Build.cs files to use alternative names for intermediate paths.");
+					Logger.LogWarning("{Message}", Message.ToString());
+				}
+			}
+		}
+
+		public static ActionExecutor? PrecreateExecutor(BuildConfiguration BuildConfiguration, ILogger Logger, List<TargetDescriptor> TargetDescriptors)
+		{
+			if (BuildConfiguration.bAllowXGE && XGE.IsAvailable(Logger))
+			{
+				return null;
+			}
+			if (BuildConfiguration.bAllowSNDBS && SNDBS.IsAvailable(Logger))
+			{
+				return null;
+			}
+			if (BuildConfiguration.bAllowFASTBuild && FASTBuild.IsAvailable(Logger))
+			{
+				return null;
+			}
+			if (BuildConfiguration.bAllowUBAExecutor && UBAExecutor.IsAvailable())
+			{
+				return new UBAExecutor(BuildConfiguration.MaxParallelActions, BuildConfiguration.bAllCores, BuildConfiguration.bCompactOutput, Logger, TargetDescriptors, true);
+			}
+
+			return null;
+		}
+
+		private static ActionExecutor? GetRemoteExecutorByName(string Name, BuildConfiguration BuildConfiguration, int ActionCount, int MinActionsForRemote, List<TargetDescriptor> TargetDescriptors, ILogger Logger)
+		{
+			switch (Name)
+			{
+				case "XGE":
+					{
+						if (ActionCount >= MinActionsForRemote && BuildConfiguration.bAllowXGE && XGE.IsAvailable(Logger) && ActionCount >= XGE.MinActions)
+						{
+							return new XGE(Logger);
+						}
+						return null;
+					}
+				case "SNDBS":
+					{
+						if (ActionCount >= MinActionsForRemote && BuildConfiguration.bAllowSNDBS && SNDBS.IsAvailable(Logger))
+						{
+							return new SNDBS(TargetDescriptors, Logger);
+						}
+						return null;
+					}
+				case "FASTBuild":
+					{
+						if (ActionCount >= MinActionsForRemote && BuildConfiguration.bAllowFASTBuild && FASTBuild.IsAvailable(Logger))
+						{
+							return new FASTBuild(BuildConfiguration.MaxParallelActions, BuildConfiguration.bAllCores, BuildConfiguration.bCompactOutput, Logger);
+						}
+						return null;
+					}
+				case "UBA":
+					{
+						// Intentionally not checking MinActionsForRemote
+						if (BuildConfiguration.bAllowUBAExecutor && UBAExecutor.IsAvailable())
+						{
+							return new UBAExecutor(BuildConfiguration.MaxParallelActions, BuildConfiguration.bAllCores, BuildConfiguration.bCompactOutput, Logger, TargetDescriptors, true);
+						}
+						return null;
+					}
+				default:
+					{
+						Logger.LogWarning("Unknown remote executor {Name}", Name);
+						return null;
+					}
+			}
+		}
+
+		/// <summary>
+		/// Selects an ActionExecutor
+		/// </summary>
+		private static ActionExecutor SelectExecutor(BuildConfiguration BuildConfiguration, int ActionCount, List<TargetDescriptor> TargetDescriptors, ILogger Logger)
+		{
+			int MinActionsForRemote = ParallelExecutor.GetDefaultNumParallelProcesses(BuildConfiguration.MaxParallelActions, BuildConfiguration.bAllCores, Logger);
+			foreach (string Name in BuildConfiguration.RemoteExecutorPriority)
+			{
+				ActionExecutor? Executor = GetRemoteExecutorByName(Name, BuildConfiguration, ActionCount, MinActionsForRemote, TargetDescriptors, Logger);
+				if (Executor != null)
+				{
+					return Executor;
+				}
+			}
+
+			return new ParallelExecutor(BuildConfiguration.MaxParallelActions, BuildConfiguration.bAllCores, BuildConfiguration.bCompactOutput, Logger);
+		}
+
+		/// <summary>
+		/// Executes a list of actions.
+		/// </summary>
+		public static async Task ExecuteActionsAsync(BuildConfiguration BuildConfiguration, List<LinkedAction> ActionsToExecute, List<TargetDescriptor> TargetDescriptors, ILogger Logger, IActionArtifactCache? actionArtifactCache = null, ActionExecutor? executor = null)
+		{
+			if (ActionsToExecute.Count == 0)
+			{
+				Logger.LogInformation("Target is up to date");
+
+				if (executor != null)
+				{
+					// We still want to run the executor to give it a chance to clean up in case it has tasks running that can be cancelled
+					await InternalExecuteActions(executor, ActionsToExecute, Logger, actionArtifactCache);
+				}
+			}
+			else
+			{
+				// Figure out which executor to use
+				if (executor == null)
+				{
+					using ActionExecutor SelectedExecutor = SelectExecutor(BuildConfiguration, ActionsToExecute.Count, TargetDescriptors, Logger);
+					await InternalExecuteActions(SelectedExecutor, ActionsToExecute, Logger, actionArtifactCache);
+				}
+				else
+				{
+					await InternalExecuteActions(executor, ActionsToExecute, Logger, actionArtifactCache);
+				}
+			}
+		}
+
+		static async Task InternalExecuteActions(ActionExecutor Executor, List<LinkedAction> ActionsToExecute, ILogger Logger, IActionArtifactCache? actionArtifactCache = null)
+		{
+			Logger.LogInformation("Using {ExecutorName} executor to run {ActionCount} action(s)", Executor.Name, ActionsToExecute.Count);
+
+			// Execute the build
+			Stopwatch Timer = Stopwatch.StartNew();
+			bool Result = await Executor.ExecuteActionsAsync(ActionsToExecute, Logger, actionArtifactCache);
+			Executor.PostTelemetryEvent();
+
+			Logger.LogInformation("Total time in {ExecutorName} executor: {TotalSeconds:0.00} seconds", Executor.Name, Timer.Elapsed.TotalSeconds);
+
+			if (!Result)
+			{
+				throw new CompilationResultException(CompilationResult.OtherCompilationError);
+			}
+
+			// Reset the file info for all the produced items
+			foreach (LinkedAction BuildAction in ActionsToExecute)
+			{
+				foreach (FileItem ProducedItem in BuildAction.ProducedItems)
+				{
+					ProducedItem.ResetCachedInfo();
+				}
+			}
+
+			// Verify the link outputs were created (seems to happen with Win64 compiles)
+			if (Executor.VerifyOutputs)
+			{
+				foreach (LinkedAction BuildAction in ActionsToExecute)
+				{
+					if (BuildAction.ActionType == ActionType.Link)
+					{
+						foreach (FileItem Item in BuildAction.ProducedItems)
+						{
+							if (!Item.Exists)
+							{
+								throw new BuildException($"Failed to produce item: {Item.AbsolutePath}");
+							}
+						}
+					}
+				}
+			}
+		}
+		/// <summary>
+		/// Sorts the action list for improved parallelism with local execution.
+		/// </summary>
+		static void SortActionList(List<LinkedAction> Actions)
+		{
+			// Clear the current dependent count and SortIndex if there is any
+			foreach (LinkedAction Action in Actions)
+			{
+				Action.NumTotalDependentActions = 0;
+				Action.SortIndex = 0;
+			}
+
+			// Increment all the dependencies plus propagate high priority flag
+			Parallel.ForEach(Actions, (Action) =>
+			{
+				Action.IncrementDependentCount(new HashSet<LinkedAction>(), Action.bIsHighPriority);
+			});
+
+			// Sort actions by number of actions depending on them, descending. Secondary sort criteria is file size.
+			Actions.Sort(LinkedAction.Compare);
+
+			// Now when everything is sorted we want to move link actions and actions that can't run remotely to as early as possible
+			// Find the last prereq action and set the sort index to the same. since action has at least one more dependency it will end up after in the next sorting call
+			int I = 0;
+			foreach (LinkedAction Action in Actions)
+			{
+				Action.SortIndex = ++I;
+				if ((!Action.bCanExecuteRemotely || Action.ActionType == ActionType.Link) && Action.PrerequisiteActions.Count > 0)
+				{
+					int LastSortIndex = 0;
+					foreach (LinkedAction Prereq in Action.PrerequisiteActions)
+					{
+						LastSortIndex = Math.Max(LastSortIndex, Prereq.SortIndex);
+					}
+					Action.SortIndex = LastSortIndex;
+				}
+			}
+
+			// Sort again now when we have set sorting index and will put link actions earlier
+			Actions.Sort(LinkedAction.Compare);
+		}
+
+		/// <summary>
+		/// Checks for cycles in the action graph.
+		/// </summary>
+		static void DetectActionGraphCycles(List<LinkedAction> Actions, Dictionary<FileItem, LinkedAction> ItemToProducingAction, ILogger Logger)
+		{
+			using ITimelineEvent _ = Timeline.ScopeEvent($"DetectActionGraphCycles");
+
+			// Starting with actions that only depend on non-produced items, iteratively expand a set of actions that are only dependent on
+			// non-cyclical actions.
+			Dictionary<LinkedAction, bool> ActionIsNonCyclical = new Dictionary<LinkedAction, bool>();
+			Dictionary<LinkedAction, List<LinkedAction>> CyclicActions = new Dictionary<LinkedAction, List<LinkedAction>>();
+			while (true)
+			{
+				bool bFoundNewNonCyclicalAction = false;
+
+				foreach (LinkedAction Action in Actions)
+				{
+					if (!ActionIsNonCyclical.ContainsKey(Action))
+					{
+						// Determine if the action depends on only actions that are already known to be non-cyclical.
+						bool bActionOnlyDependsOnNonCyclicalActions = true;
+						foreach (FileItem PrerequisiteItem in Action.PrerequisiteItems)
+						{
+							if (ItemToProducingAction.TryGetValue(PrerequisiteItem, out LinkedAction? ProducingAction))
+							{
+								if (!ActionIsNonCyclical.ContainsKey(ProducingAction))
+								{
+									bActionOnlyDependsOnNonCyclicalActions = false;
+									if (!CyclicActions.ContainsKey(Action))
+									{
+										CyclicActions.Add(Action, new List<LinkedAction>());
+									}
+
+									List<LinkedAction> CyclicPrerequisite = CyclicActions[Action];
+									if (!CyclicPrerequisite.Contains(ProducingAction))
+									{
+										CyclicPrerequisite.Add(ProducingAction);
+									}
+								}
+							}
+						}
+
+						// If the action only depends on known non-cyclical actions, then add it to the set of known non-cyclical actions.
+						if (bActionOnlyDependsOnNonCyclicalActions)
+						{
+							ActionIsNonCyclical.Add(Action, true);
+							bFoundNewNonCyclicalAction = true;
+							if (CyclicActions.ContainsKey(Action))
+							{
+								CyclicActions.Remove(Action);
+							}
+						}
+					}
+				}
+
+				// If this iteration has visited all actions without finding a new non-cyclical action, then all non-cyclical actions have
+				// been found.
+				if (!bFoundNewNonCyclicalAction)
+				{
+					break;
+				}
+			}
+
+			// If there are any cyclical actions, throw an exception.
+			if (ActionIsNonCyclical.Count < Actions.Count)
+			{
+				// Find the index of each action
+				Dictionary<LinkedAction, int> ActionToIndex = new Dictionary<LinkedAction, int>();
+				for (int Idx = 0; Idx < Actions.Count; Idx++)
+				{
+					ActionToIndex[Actions[Idx]] = Idx;
+				}
+
+				// Describe the cyclical actions.
+				foreach (LinkedAction Action in Actions)
+				{
+					if (!ActionIsNonCyclical.ContainsKey(Action))
+					{
+						string CycleDescription = "";
+						CycleDescription += $"Action #{ActionToIndex[Action]}: {Action.CommandPath}\n";
+						CycleDescription += $"\twith arguments: {Action.CommandArguments}\n";
+						foreach (FileItem PrerequisiteItem in Action.PrerequisiteItems)
+						{
+							CycleDescription += $"\tdepends on: {PrerequisiteItem.AbsolutePath}\n";
+						}
+
+						foreach (FileItem ProducedItem in Action.ProducedItems)
+						{
+							CycleDescription += $"\tproduces:   {ProducedItem.AbsolutePath}\n";
+						}
+
+						CycleDescription += "\tDepends on cyclic actions:\n";
+						if (CyclicActions.ContainsKey(Action))
+						{
+							foreach (LinkedAction CyclicPrerequisiteAction in CyclicActions[Action])
+							{
+								if (CyclicActions.ContainsKey(CyclicPrerequisiteAction))
+								{
+									List<FileItem> CyclicProducedItems =
+										CyclicPrerequisiteAction.ProducedItems.ToList();
+									if (CyclicProducedItems.Count == 1)
+									{
+										CycleDescription += $"\t\t{ActionToIndex[CyclicPrerequisiteAction]} (produces: {CyclicProducedItems[0].AbsolutePath})\n";
+									}
+									else
+									{
+										CycleDescription += $"\t\t{ActionToIndex[CyclicPrerequisiteAction]}\n";
+										foreach (FileItem CyclicProducedItem in CyclicProducedItems)
+										{
+											CycleDescription += $"\t\t\tproduces:   {CyclicProducedItem.AbsolutePath}\n";
+										}
+									}
+								}
+							}
+
+							CycleDescription += "\n";
+						}
+						else
+						{
+							CycleDescription += "\t\tNone?? Coding error!\n";
+						}
+
+						Logger.LogInformation(CycleDescription);
+					}
+				}
+
+				throw new CompilationResultException(CompilationResult.ActionGraphInvalid, "Action graph contains cycle!");
+			}
+		}
+
+		/// <summary>
+		/// Determines the full set of actions that must be built to produce an item.
+		/// </summary>
+		/// <param name="Actions">All the actions in the graph</param>
+		/// <param name="OutputItems">Set of output items to be built</param>
+		/// <returns>Set of prerequisite actions</returns>
+		public static List<LinkedAction> GatherPrerequisiteActions(List<LinkedAction> Actions, HashSet<FileItem> OutputItems)
+		{
+			using ITimelineEvent _ = Timeline.ScopeEvent($"GatherPrerequisiteActions");
+
+			HashSet<LinkedAction> PrerequisiteActions = new HashSet<LinkedAction>();
+			foreach (LinkedAction Action in Actions)
+			{
+				if (Action.ProducedItems.Any(OutputItems.Contains))
+				{
+					GatherPrerequisiteActions(Action, PrerequisiteActions);
+				}
+			}
+
+			return PrerequisiteActions.ToList();
+		}
+
+		/// <summary>
+		/// Determines the full set of actions that must be built to produce an item.
+		/// </summary>
+		/// <param name="Action">The root action to scan</param>
+		/// <param name="PrerequisiteActions">Set of prerequisite actions</param>
+		private static void GatherPrerequisiteActions(LinkedAction Action, HashSet<LinkedAction> PrerequisiteActions)
+		{
+			if (PrerequisiteActions.Add(Action))
+			{
+				foreach (LinkedAction PrerequisiteAction in Action.PrerequisiteActions)
+				{
+					GatherPrerequisiteActions(PrerequisiteAction, PrerequisiteActions);
+				}
+			}
+		}
+
+		/// <summary>
+		/// Determines whether an action is outdated based on the modification times for its prerequisite
+		/// and produced items, without considering the full set of prerequisites.
+		/// Writes to OutdatedActionDictionary iff the action is found to be outdated.
+		/// Safe to run in parallel, but only with different RootActions.
+		/// </summary>
+		/// <param name="RootAction">- The action being considered.</param>
+		/// <param name="OutdatedActionDictionary">-</param>
+		/// <param name="ActionHistory"></param>
+		/// <param name="CppDependencies"></param>
+		/// <param name="bIgnoreOutdatedImportLibraries"></param>
+		/// <param name="Logger"></param>
+		/// <returns>true if outdated</returns>
+		private static void IsIndividualActionOutdated(LinkedAction RootAction,
+			ConcurrentDictionary<LinkedAction, bool> OutdatedActionDictionary,
+			ActionHistory ActionHistory, CppDependencyCache CppDependencies, bool bIgnoreOutdatedImportLibraries, ILogger Logger)
+		{
+			// Only compute the outdated-ness for actions that don't aren't cached in the outdated action dictionary.
+			bool bIsOutdated = false;
+			{
+				// OutdatedActionDictionary may have already been populated for RootAction as part of a previously processed target
+				if (OutdatedActionDictionary.ContainsKey(RootAction))
+				{
+					return;
+				}
+			}
+
+			// Determine the last time the action was run based on the write times of its produced files.
+			DateTimeOffset LastExecutionTimeUtc = DateTimeOffset.MaxValue;
+			string NewProducingAttributes = $"{RootAction.CommandPath.FullName} {RootAction.CommandArguments} (ver {RootAction.CommandVersion})";
+			foreach (FileItem ProducedItem in RootAction.ProducedItems)
+			{
+				ProducedItem.Directory.CacheFiles();
+
+				if (!ProducedItem.HasInfo)
+				{
+					//Console.WriteLine("HASINFO " + ProducedItem);
+				}
+
+				// Check if the command-line of the action previously used to produce the item is outdated.
+				if (ActionHistory.UpdateProducingAttributes(ProducedItem, NewProducingAttributes, Logger) && RootAction.bUseActionHistory)
+				{
+					if (ProducedItem.Exists)
+					{
+						Logger.LogTrace("{StatusDescription}: Produced item \"{ProducedItem}\" was produced by outdated attributes.", RootAction.StatusDescription, ProducedItem.Location);
+						Logger.LogTrace("  New attributes: {NewProducingAttributes}", NewProducingAttributes);
+					}
+
+					bIsOutdated = true;
+				}
+
+				// If the produced file doesn't exist or has zero size, consider it outdated.  The zero size check is to detect cases
+				// where aborting an earlier compile produced invalid zero-sized obj files, but that may cause actions where that's
+				// legitimate output to always be considered outdated.
+				if (ProducedItem.Exists && (RootAction.ActionType != ActionType.Compile || ProducedItem.Length > 0 ||
+											(!ProducedItem.Location.HasExtension(".obj") && !ProducedItem.Location.HasExtension(".o"))))
+				{
+					// Find the newer of LastWriteTime and CreationTime, as copied files (such as from a build cache) can have a newer creation time in some cases
+					DateTime ExecutionTimeUtc = ProducedItem.LastWriteTimeUtc > ProducedItem.CreationTimeUtc ? ProducedItem.LastWriteTimeUtc : ProducedItem.CreationTimeUtc;
+					// Use the oldest produced item's time as the last execution time.
+					if (ExecutionTimeUtc < LastExecutionTimeUtc)
+					{
+						LastExecutionTimeUtc = ExecutionTimeUtc;
+					}
+				}
+				else
+				{
+					// If any of the produced items doesn't exist, the action is outdated.
+					Logger.LogTrace("{StatusDescription}: Produced item \"{ProducedItem}\" doesn't exist.", RootAction.StatusDescription, ProducedItem.Location);
+					bIsOutdated = true;
+				}
+			}
+
+			// Check if any prerequisite item has a newer timestamp than the last execution time of this action
+			if (!bIsOutdated)
+			{
+				foreach (FileItem PrerequisiteItem in RootAction.PrerequisiteItems)
+				{
+					// Need to check for import libraries here too
+					if (bIgnoreOutdatedImportLibraries && IsImportLibraryDependency(RootAction, PrerequisiteItem))
+					{
+						continue;
+					}
+
+					PrerequisiteItem.Directory.CacheFiles();
+
+					if (PrerequisiteItem.Exists)
+					{
+						// allow a 1 second slop for network copies
+						TimeSpan TimeDifference = PrerequisiteItem.LastWriteTimeUtc - LastExecutionTimeUtc;
+						bool bPrerequisiteItemIsNewerThanLastExecution = TimeDifference.TotalSeconds > 1;
+						if (bPrerequisiteItemIsNewerThanLastExecution)
+						{
+							Logger.LogTrace("{StatusDescription}: Prerequisite {PrerequisiteItem} is newer than the last execution of the action: {Message}", RootAction.StatusDescription, PrerequisiteItem.Location,
+									$"{PrerequisiteItem.LastWriteTimeUtc.ToLocalTime().ToString(CultureInfo.CurrentCulture)} vs {LastExecutionTimeUtc.LocalDateTime.ToString(CultureInfo.CurrentCulture)}");
+							bIsOutdated = true;
+							break;
+						}
+					}
+					else
+					{
+						// If the PrerequisiteItem doesn't exist, it means it's being produced right now and we need to update the ProducedItem that
+						// depends on it
+						Logger.LogTrace("{StatusDescription}: Prerequisite {PrerequisiteItem} doesn't exist.", RootAction.StatusDescription, PrerequisiteItem.Location);
+						bIsOutdated = true;
+						break;
+					}
+				}
+			}
+
+			// Check the dependency list
+			if (!bIsOutdated && RootAction.DependencyListFile != null)
+			{
+				if (!CppDependencies.TryGetDependencies(RootAction.DependencyListFile, Logger, out List<FileItem>? DependencyFiles))
+				{
+					Logger.LogTrace("{StatusDescription}: Missing dependency list file \"{DependencyListFile}\"", RootAction.StatusDescription, RootAction.DependencyListFile);
+					bIsOutdated = true;
+				}
+				else
+				{
+					foreach (FileItem DependencyFile in DependencyFiles.Except(RootAction.PrerequisiteItems))
+					{
+						if (!DependencyFile.Exists)
+						{
+							Logger.LogTrace("{StatusDescription}: Dependency {DependencyFile} doesn't exist", RootAction.StatusDescription, DependencyFile.AbsolutePath);
+							bIsOutdated = true;
+							break;
+						}
+						else if (DependencyFile.LastWriteTimeUtc > LastExecutionTimeUtc)
+						{
+							Logger.LogTrace("{StatusDescription}: Dependency {DependencyFile} is newer than the last execution of the action: {Message}", RootAction.StatusDescription, DependencyFile.AbsolutePath,
+								$"{DependencyFile.LastWriteTimeUtc.ToLocalTime().ToString(CultureInfo.CurrentCulture)} vs {LastExecutionTimeUtc.LocalDateTime.ToString(CultureInfo.CurrentCulture)}");
+							bIsOutdated = true;
+							break;
+						}
+					}
+				}
+			}
+
+			// if the action is known to be out of date, record that fact 
+			// We don't yet know that the action is up-to-date - to determine that requires traversal of the graph of prerequisites.
+			if (bIsOutdated)
+			{
+				OutdatedActionDictionary.TryAdd(RootAction, bIsOutdated);
+			}
+		}
+
+		/// <summary>
+		/// Determines whether an action is outdated by examining the up-to-date state of all of its prerequisites, recursively.
+		/// Not thread safe. Typically very fast.
+		/// </summary>
+		/// <param name="RootAction">- The action being considered.</param>
+		/// <param name="OutdatedActionDictionary">-</param>
+		/// <param name="bIgnoreOutdatedImportLibraries"></param>
+		/// <param name="Logger">Logger instance</param>
+		/// <returns>true if outdated</returns>
+		private static bool IsActionOutdatedDueToPrerequisites(LinkedAction RootAction, ConcurrentDictionary<LinkedAction, bool> OutdatedActionDictionary, bool bIgnoreOutdatedImportLibraries, ILogger Logger)
+		{
+			// Only compute the outdated-ness for actions that aren't already cached in the outdated action dictionary.
+			if (OutdatedActionDictionary.TryGetValue(RootAction, out bool bIsOutdated))
+			{
+				return bIsOutdated;
+			}
+
+			// Check if any of the prerequisite actions are out of date
+			foreach (LinkedAction PrerequisiteAction in RootAction.PrerequisiteActions)
+			{
+				if (IsActionOutdatedDueToPrerequisites(PrerequisiteAction, OutdatedActionDictionary, bIgnoreOutdatedImportLibraries, Logger))
+				{
+					// Only check for outdated import libraries if we were configured to do so.  Often, a changed import library
+					// won't affect a dependency unless a public header file was also changed, in which case we would be forced
+					// to recompile anyway.  This just allows for faster iteration when working on a subsystem in a DLL, as we
+					// won't have to wait for dependent targets to be relinked after each change.
+					if (!bIgnoreOutdatedImportLibraries || !IsImportLibraryDependency(RootAction, PrerequisiteAction))
+					{
+						Logger.LogTrace("{StatusDescription}: Prerequisite {PrereqStatusDescription} is produced by outdated action.", RootAction.StatusDescription, PrerequisiteAction.StatusDescription);
+						bIsOutdated = true;
+						break;
+					}
+				}
+			}
+
+			// Cache the outdated-ness of this action.
+			OutdatedActionDictionary.TryAdd(RootAction, bIsOutdated);
+
+			return bIsOutdated;
+		}
+
+		/// <summary>
+		/// Determines if the dependency between two actions is only for an import library
+		/// </summary>
+		/// <param name="RootAction">The action to check</param>
+		/// <param name="PrerequisiteAction">The action that it depends on</param>
+		/// <returns>True if the only dependency between two actions is for an import library</returns>
+		static bool IsImportLibraryDependency(LinkedAction RootAction, LinkedAction PrerequisiteAction)
+		{
+			if (PrerequisiteAction.bProducesImportLibrary)
+			{
+				return PrerequisiteAction.ProducedItems.All(I => IsLibraryFile(I.Location) || !RootAction.PrerequisiteItems.Contains(I));
+			}
+			else
+			{
+				return false;
+			}
+		}
+
+		/// <summary>
+		/// Determines if the dependency on a between two actions is only for an import library
+		/// </summary>
+		/// <param name="RootAction">The action to check</param>
+		/// <param name="PrerequisiteItem">The dependency that is out of date</param>
+		/// <returns>True if the only dependency between two actions is for an import library</returns>
+		static bool IsImportLibraryDependency(LinkedAction RootAction, FileItem PrerequisiteItem)
+		{
+			if (IsLibraryFile(PrerequisiteItem.Location))
+			{
+				foreach (LinkedAction PrerequisiteAction in RootAction.PrerequisiteActions)
+				{
+					if (PrerequisiteAction.bProducesImportLibrary && PrerequisiteAction.ProducedItems.FirstOrDefault(PrerequisiteItem) != null)
+					{
+						return true;
+					}
+				}
+			}
+
+			return false;
+		}
+
+		/// <summary>
+		/// Test to see if the given file is a library file (or in the case of linux/mac, a dynamic library)
+		/// </summary>
+		/// <param name="Location">File to test</param>
+		/// <returns>True if the file is a library file</returns>
+		static bool IsLibraryFile(FileReference Location)
+		{
+			return Location.HasExtension(".lib") || Location.HasExtension(".so") || Location.HasExtension(".dylib");
+		}
+
+		/// <summary>
+		/// Builds a dictionary containing the actions from AllActions that are outdated by calling
+		/// IsActionOutdated.
+		/// </summary>
+		public static void GatherAllOutdatedActions(IReadOnlyList<LinkedAction> Actions, ActionHistory ActionHistory,
+			ConcurrentDictionary<LinkedAction, bool> OutdatedActions, CppDependencyCache CppDependencies,
+			bool bIgnoreOutdatedImportLibraries, ILogger Logger)
+		{
+			using ITimelineEvent __ = Timeline.ScopeEvent($"GatherOutdatedActions");
+
+			using (GlobalTracer.Instance.BuildSpan("Cache individual outdated actions").StartActive())
+			{
+				Parallel2.ForEach(Actions,
+					Action => IsIndividualActionOutdated(Action, OutdatedActions,
+						ActionHistory, CppDependencies, bIgnoreOutdatedImportLibraries, Logger));
+			}
+
+			using (GlobalTracer.Instance.BuildSpan("Cache outdated actions based on recursive prerequisites").StartActive())
+			{
+				foreach (LinkedAction Action in Actions)
+				{
+					IsActionOutdatedDueToPrerequisites(Action, OutdatedActions, bIgnoreOutdatedImportLibraries, Logger);
+				}
+			}
+		}
+
+		/// <summary>
+		/// Deletes all the items produced by actions in the provided outdated action dictionary.
+		/// </summary>
+		/// <param name="OutdatedActions">List of outdated actions</param>
+		/// <param name="Logger">Logger for output</param>
+		public static void DeleteOutdatedProducedItems(List<LinkedAction> OutdatedActions, ILogger Logger)
+		{
+			using ITimelineEvent __ = Timeline.ScopeEvent($"DeleteOutdatedProducedItems");
+
+			Parallel2.ForEach(OutdatedActions.SelectMany(oa => oa.DeleteItems), DeleteItem =>
+			{
+				if (DeleteItem.Exists)
+				{
+					Logger.LogDebug("Deleting outdated item: {DeleteItem}", DeleteItem.Location);
+					DeleteItem.Delete(Logger);
+				}
+			});
+		}
+
+		/// <summary>
+		/// Creates directories for all the items produced by actions in the provided outdated action
+		/// dictionary.
+		/// </summary>
+		public static void CreateDirectoriesForProducedItems(List<LinkedAction> OutdatedActions)
+		{
+			using ITimelineEvent __ = Timeline.ScopeEvent($"CreateDirectoriesForProducedItems");
+
+			HashSet<DirectoryReference> OutputDirectories = new HashSet<DirectoryReference>();
+			foreach (LinkedAction OutdatedAction in OutdatedActions)
+			{
+				foreach (FileItem ProducedItem in OutdatedAction.ProducedItems)
+				{
+					OutputDirectories.Add(ProducedItem.Location.Directory);
+				}
+			}
+
+			foreach (DirectoryReference OutputDirectory in OutputDirectories)
+			{
+				if (!DirectoryReference.Exists(OutputDirectory))
+				{
+					DirectoryReference.CreateDirectory(OutputDirectory);
+				}
+			}
+		}
+
+		/// <summary>
+		/// Imports an action graph from a JSON file
+		/// </summary>
+		/// <param name="InputFile">The file to read from</param>
+		/// <returns>List of actions</returns>
+		public static List<Action> ImportJson(FileReference InputFile)
+		{
+			JsonObject Object = JsonObject.Read(InputFile);
+
+			JsonObject EnvironmentObject = Object.GetObjectField("Environment");
+			foreach (string KeyName in EnvironmentObject.KeyNames)
+			{
+				Environment.SetEnvironmentVariable(KeyName, EnvironmentObject.GetStringField(KeyName));
+			}
+
+			List<Action> Actions = new List<Action>();
+			foreach (JsonObject ActionObject in Object.GetObjectArrayField("Actions"))
+			{
+				Actions.Add(Action.ImportJson(ActionObject));
+			}
+
+			return Actions;
+		}
+
+		/// <summary>
+		/// Exports an action graph to a JSON file
+		/// </summary>
+		/// <param name="Actions">The actions to write</param>
+		/// <param name="OutputFile">Output file to write the actions to</param>
+		public static void ExportJson(IReadOnlyList<LinkedAction> Actions, FileReference OutputFile)
+		{
+			DirectoryReference.CreateDirectory(OutputFile.Directory);
+			using JsonWriter Writer = new JsonWriter(OutputFile);
+			Writer.WriteObjectStart();
+
+			Writer.WriteObjectStart("Environment");
+			foreach (object? Object in Environment.GetEnvironmentVariables())
+			{
+				System.Collections.DictionaryEntry Pair = (System.Collections.DictionaryEntry)Object!;
+				if (!UnrealBuildTool.InitialEnvironment!.Contains(Pair.Key) ||
+					(string)(UnrealBuildTool.InitialEnvironment[Pair.Key]!) != (string)(Pair.Value!))
+				{
+					Writer.WriteValue((string)Pair.Key, (string)Pair.Value!);
+				}
+			}
+
+			Writer.WriteObjectEnd();
+
+			Dictionary<LinkedAction, int> ActionToId = new Dictionary<LinkedAction, int>();
+			foreach (LinkedAction Action in Actions)
+			{
+				ActionToId[Action] = ActionToId.Count;
+			}
+
+			Writer.WriteArrayStart("Actions");
+			foreach (LinkedAction Action in Actions)
+			{
+				Writer.WriteObjectStart();
+				Action.ExportJson(ActionToId, Writer);
+				Writer.WriteObjectEnd();
+			}
+
+			Writer.WriteArrayEnd();
+			Writer.WriteObjectEnd();
+		}
+	}
+}
diff --git a/Engine/Source/Programs/UnrealBuildTool/Actions/ActionGraphBuilder.cs b/Engine/Source/Programs/UnrealBuildTool/Actions/ActionGraphBuilder.cs
new file mode 100644
--- /dev/null
+++ b/Engine/Source/Programs/UnrealBuildTool/Actions/ActionGraphBuilder.cs
@@ -0,0 +1,372 @@
+// Copyright Epic Games, Inc. All Rights Reserved.
+
+using System;
+using System.Collections.Generic;
+using System.Reflection;
+using EpicGames.Core;
+using Microsoft.Extensions.Logging;
+using UnrealBuildBase;
+
+namespace UnrealBuildTool
+{
+	/// <summary>
+	/// Interface for toolchain operations that produce output
+	/// </summary>
+	interface IActionGraphBuilder
+	{
+		/// <summary>
+		/// Adds an action to this graph
+		/// </summary>
+		/// <param name="Action">Action to add</param>
+		void AddAction(IExternalAction Action);
+
+		/// <summary>
+		/// Creates a response file for use in the action graph
+		/// </summary>
+		/// <param name="Location">Location of the response file</param>
+		/// <param name="Contents">Contents of the file</param>
+		/// <param name="AllowAsync">Allows the backend to write the file in a separate task.</param>
+		/// <returns>New file item</returns>
+		void CreateIntermediateTextFile(FileItem Location, string Contents, bool AllowAsync = true);
+
+		/// <summary>
+		/// Creates a response file for use in the action graph, with a newline between each string in ContentLines
+		/// </summary>
+		/// <param name="Location">Location of the response file</param>
+		/// <param name="ContentLines">Contents of the file</param>
+		/// <param name="AllowAsync">Allows the backend to write the file in a separate task.</param>
+		/// <returns>New file item</returns>
+		void CreateIntermediateTextFile(FileItem Location, IEnumerable<string> ContentLines, bool AllowAsync = true);
+
+		/// <summary>
+		/// Adds a file which is in the non-unity working set
+		/// </summary>
+		/// <param name="File">The file to add to the working set</param>
+		void AddFileToWorkingSet(FileItem File);
+
+		/// <summary>
+		/// Adds a file which is a candidate for being in the non-unity working set
+		/// </summary>
+		/// <param name="File">The file to add to the working set</param>
+		void AddCandidateForWorkingSet(FileItem File);
+
+		/// <summary>
+		/// Adds a source directory. These folders are scanned recursively for C++ source files.
+		/// </summary>
+		/// <param name="SourceDir">Base source directory</param>
+		void AddSourceDir(DirectoryItem SourceDir);
+
+		/// <summary>
+		/// Adds the given source files as dependencies
+		/// </summary>
+		/// <param name="SourceDir">Source directory containing files to build</param>
+		/// <param name="SourceFiles">Contents of the directory</param>
+		void AddSourceFiles(DirectoryItem SourceDir, FileItem[] SourceFiles);
+
+		/// <summary>
+		/// Adds a list of known header files
+		/// </summary>
+		/// <param name="HeaderFiles">List of header files to track</param>
+		void AddHeaderFiles(FileItem[] HeaderFiles);
+
+		/// <summary>
+		/// Sets the output items which belong to a particular module
+		/// </summary>
+		/// <param name="ModuleName">Name of the module</param>
+		/// <param name="OutputItems">Array of output items for this module</param>
+		void SetOutputItemsForModule(string ModuleName, FileItem[] OutputItems);
+
+		/// <summary>
+		/// Adds a diagnostic message
+		/// </summary>
+		/// <param name="Message">Message to display</param>
+		void AddDiagnostic(string Message);
+	}
+
+	/// <summary>
+	/// Implementation of IActionGraphBuilder which discards all unnecessary operations
+	/// </summary>
+	sealed class NullActionGraphBuilder : IActionGraphBuilder
+	{
+		private readonly ILogger Logger;
+		private readonly FileHashCache? _fileHashCache;
+
+		/// <summary>
+		/// Constructor
+		/// </summary>
+		public NullActionGraphBuilder(ILogger InLogger, FileHashCache? fileHashCache = null)
+		{
+			Logger = InLogger;
+			_fileHashCache = fileHashCache;
+		}
+
+		/// <inheritdoc/>
+		public void AddAction(IExternalAction Action)
+		{
+		}
+
+		/// <inheritdoc/>
+		public void CreateIntermediateTextFile(FileItem FileItem, string Contents, bool AllowAsync = true)
+		{
+			if (_fileHashCache != null)
+			{
+				_fileHashCache.WriteFileIfChanged(FileItem, Contents, Logger, false);
+			}
+			else
+			{
+				Utils.WriteFileIfChanged(FileItem, Contents, Logger);
+			}
+		}
+
+		/// <inheritdoc/>
+		public void CreateIntermediateTextFile(FileItem FileItem, IEnumerable<string> ContentLines, bool AllowAsync = true)
+		{
+			if (_fileHashCache != null)
+			{
+				_fileHashCache.WriteFileIfChanged(FileItem, ContentLines, Logger, false);
+			}
+			else
+			{
+				Utils.WriteFileIfChanged(FileItem, ContentLines, Logger);
+			}
+		}
+
+		/// <inheritdoc/>
+		public void AddSourceDir(DirectoryItem SourceDir)
+		{
+		}
+
+		/// <inheritdoc/>
+		public void AddSourceFiles(DirectoryItem SourceDir, FileItem[] SourceFiles)
+		{
+		}
+
+		/// <inheritdoc/>
+		public void AddHeaderFiles(FileItem[] HeaderFiles)
+		{
+		}
+
+		/// <inheritdoc/>
+		public void AddFileToWorkingSet(FileItem File)
+		{
+		}
+
+		/// <inheritdoc/>
+		public void AddCandidateForWorkingSet(FileItem File)
+		{
+		}
+
+		/// <inheritdoc/>
+		public void AddDiagnostic(string Message)
+		{
+		}
+
+		/// <inheritdoc/>
+		public void SetOutputItemsForModule(string ModuleName, FileItem[] OutputItems)
+		{
+		}
+	}
+
+	/// <summary>
+	/// Implementation of IActionGraphBuilder which forwards calls to an underlying implementation, allowing derived classes to intercept certain calls
+	/// </summary>
+	class ForwardingActionGraphBuilder : IActionGraphBuilder
+	{
+		/// <summary>
+		/// The inner graph builder
+		/// </summary>
+		readonly IActionGraphBuilder Inner;
+
+		/// <summary>
+		/// Constructor
+		/// </summary>
+		/// <param name="Inner">Builder to pass all calls to</param>
+		public ForwardingActionGraphBuilder(IActionGraphBuilder Inner)
+		{
+			this.Inner = Inner;
+		}
+
+		/// <inheritdoc/>
+		public virtual void AddAction(IExternalAction Action)
+		{
+			Inner.AddAction(Action);
+		}
+
+		/// <inheritdoc/>
+		public virtual void CreateIntermediateTextFile(FileItem FileItem, string Contents, bool AllowAsync = true)
+		{
+			Inner.CreateIntermediateTextFile(FileItem, Contents, AllowAsync);
+		}
+
+		/// <inheritdoc/>
+		public virtual void CreateIntermediateTextFile(FileItem FileItem, IEnumerable<string> ContentLines, bool AllowAsync = true)
+		{
+			Inner.CreateIntermediateTextFile(FileItem, ContentLines, AllowAsync);
+		}
+
+		/// <inheritdoc/>
+		public virtual void AddSourceDir(DirectoryItem SourceDir)
+		{
+			Inner.AddSourceDir(SourceDir);
+		}
+
+		/// <inheritdoc/>
+		public virtual void AddSourceFiles(DirectoryItem SourceDir, FileItem[] SourceFiles)
+		{
+			Inner.AddSourceFiles(SourceDir, SourceFiles);
+		}
+
+		/// <inheritdoc/>
+		public virtual void AddHeaderFiles(FileItem[] HeaderFiles)
+		{
+			Inner.AddHeaderFiles(HeaderFiles);
+		}
+
+		/// <inheritdoc/>
+		public virtual void AddFileToWorkingSet(FileItem File)
+		{
+			Inner.AddFileToWorkingSet(File);
+		}
+
+		/// <inheritdoc/>
+		public virtual void AddCandidateForWorkingSet(FileItem File)
+		{
+			Inner.AddCandidateForWorkingSet(File);
+		}
+
+		/// <inheritdoc/>
+		public virtual void AddDiagnostic(string Message)
+		{
+			Inner.AddDiagnostic(Message);
+		}
+
+		/// <inheritdoc/>
+		public virtual void SetOutputItemsForModule(string ModuleName, FileItem[] OutputItems)
+		{
+			Inner.SetOutputItemsForModule(ModuleName, OutputItems);
+		}
+	}
+
+	/// <summary>
+	/// Extension methods for IActionGraphBuilder classes
+	/// </summary>
+	static class ActionGraphBuilderExtensions
+	{
+		/// <summary>
+		/// Creates a new action to be built as part of this target
+		/// </summary>
+		/// <param name="Graph">Graph to add the action to</param>
+		/// <param name="Type">Type of action to create</param>
+		/// <returns>New action</returns>
+		public static Action CreateAction(this IActionGraphBuilder Graph, ActionType Type)
+		{
+			Action Action = new Action(Type);
+			Graph.AddAction(Action);
+			return Action;
+		}
+
+		/// <summary>
+		/// Creates an action which copies a file from one location to another
+		/// </summary>
+		/// <param name="Graph">The action graph</param>
+		/// <param name="SourceFile">The source file location</param>
+		/// <param name="TargetFile">The target file location</param>
+		/// <returns>File item for the output file</returns>
+		public static Action CreateCopyAction(this IActionGraphBuilder Graph, FileItem SourceFile, FileItem TargetFile)
+		{
+			Action CopyAction = Graph.CreateAction(ActionType.BuildProject);
+			CopyAction.CommandDescription = "Copy";
+			CopyAction.CommandPath = BuildHostPlatform.Current.Shell;
+			if (BuildHostPlatform.Current.ShellType == ShellType.Cmd)
+			{
+				CopyAction.CommandArguments = String.Format("/C \"copy /Y \"{0}\" \"{1}\" 1>nul\"", SourceFile.AbsolutePath, TargetFile.AbsolutePath);
+			}
+			else
+			{
+				CopyAction.CommandArguments = String.Format("-c \"cp -f \\\"{0}\\\" \\\"{1}\\\"\"", SourceFile.AbsolutePath, TargetFile.AbsolutePath);
+			}
+			CopyAction.WorkingDirectory = Unreal.EngineSourceDirectory;
+			CopyAction.PrerequisiteItems.Add(SourceFile);
+			CopyAction.ProducedItems.Add(TargetFile);
+			CopyAction.DeleteItems.Add(TargetFile);
+			CopyAction.StatusDescription = TargetFile.Location.GetFileName();
+			CopyAction.bCanExecuteRemotely = false;
+			return CopyAction;
+		}
+
+		/// <summary>
+		/// Creates an action which copies a file from one location to another
+		/// </summary>
+		/// <param name="Graph">List of actions to be executed. Additional actions will be added to this list.</param>
+		/// <param name="SourceFile">The source file location</param>
+		/// <param name="TargetFile">The target file location</param>
+		/// <returns>File item for the output file</returns>
+		public static FileItem CreateCopyAction(this IActionGraphBuilder Graph, FileReference SourceFile, FileReference TargetFile)
+		{
+			FileItem SourceFileItem = FileItem.GetItemByFileReference(SourceFile);
+			FileItem TargetFileItem = FileItem.GetItemByFileReference(TargetFile);
+
+			Graph.CreateCopyAction(SourceFileItem, TargetFileItem);
+
+			return TargetFileItem;
+		}
+
+		/// <summary>
+		/// Creates an action which calls UBT recursively
+		/// </summary>
+		/// <param name="Graph">The action graph</param>
+		/// <param name="Type">Type of the action</param>
+		/// <param name="Arguments">Arguments for the action</param>
+		/// <returns>New action instance</returns>
+		public static Action CreateRecursiveAction<T>(this IActionGraphBuilder Graph, ActionType Type, string Arguments) where T : ToolMode
+		{
+			ToolModeAttribute? Attribute = typeof(T).GetCustomAttribute<ToolModeAttribute>();
+			if (Attribute == null)
+			{
+				throw new BuildException("Missing ToolModeAttribute on {0}", typeof(T).Name);
+			}
+
+			Action NewAction = Graph.CreateAction(Type);
+			NewAction.CommandPath = Unreal.DotnetPath;
+			NewAction.WorkingDirectory = Unreal.EngineSourceDirectory;
+			NewAction.CommandArguments = $"\"{Unreal.UnrealBuildToolDllPath}\" -Session=\"{UnrealBuildTool.SessionIdentifier}\" -Mode={Attribute.Name} {Arguments}";
+			NewAction.CommandDescription = Attribute.Name;
+			NewAction.bCanExecuteRemotely = false;
+			NewAction.bCanExecuteRemotelyWithSNDBS = false;
+			NewAction.bCanExecuteInUBA = false;
+			return NewAction;
+		}
+
+		/// <summary>
+		/// Creates a text file with the given contents.  If the contents of the text file aren't changed, it won't write the new contents to
+		/// the file to avoid causing an action to be considered outdated.
+		/// </summary>
+		/// <param name="Graph">The action graph</param>
+		/// <param name="AbsolutePath">Path to the intermediate file to create</param>
+		/// <param name="Contents">Contents of the new file</param>
+		/// <param name="AllowAsync">Allows the backend to write the file in a separate task.</param>
+		/// <returns>File item for the newly created file</returns>
+		public static FileItem CreateIntermediateTextFile(this IActionGraphBuilder Graph, FileReference AbsolutePath, string Contents, bool AllowAsync = true)
+		{
+			FileItem FileItem = FileItem.GetItemByFileReference(AbsolutePath);
+			Graph.CreateIntermediateTextFile(FileItem, Contents, AllowAsync);
+			return FileItem;
+		}
+
+		/// <summary>
+		/// Creates a text file with the given contents.  If the contents of the text file aren't changed, it won't write the new contents to
+		/// the file to avoid causing an action to be considered outdated.
+		/// </summary>
+		/// <param name="Graph">The action graph</param>
+		/// <param name="AbsolutePath">Path to the intermediate file to create</param>
+		/// <param name="ContentLines">Contents of the new file</param>
+		/// <param name="AllowAsync">Allows the backend to write the file in a separate task.</param>
+		/// <returns>File item for the newly created file</returns>
+		public static FileItem CreateIntermediateTextFile(this IActionGraphBuilder Graph, FileReference AbsolutePath, IEnumerable<string> ContentLines, bool AllowAsync = true)
+		{
+			FileItem FileItem = UnrealBuildBase.FileItem.GetItemByFileReference(AbsolutePath);
+			Graph.CreateIntermediateTextFile(FileItem, ContentLines, AllowAsync);
+			return FileItem;
+		}
+	}
+}
diff --git a/Engine/Source/Programs/UnrealBuildTool/Actions/ActionHistory.cs b/Engine/Source/Programs/UnrealBuildTool/Actions/ActionHistory.cs
new file mode 100644
--- /dev/null
+++ b/Engine/Source/Programs/UnrealBuildTool/Actions/ActionHistory.cs
@@ -0,0 +1,431 @@
+// Copyright Epic Games, Inc. All Rights Reserved.
+
+using System;
+using System.Collections.Concurrent;
+using System.Collections.Generic;
+using System.Diagnostics;
+using System.IO;
+using System.Linq;
+using System.Text;
+using EpicGames.Core;
+using Microsoft.Extensions.Logging;
+using UnrealBuildBase;
+
+namespace UnrealBuildTool
+{
+	/// <summary>
+	/// Caches include dependency information to speed up preprocessing on subsequent runs.
+	/// </summary>
+	[DebuggerDisplay("{Location}")]
+	class ActionHistoryLayer
+	{
+		/// <summary>
+		/// Version number to check
+		/// </summary>
+		const int CurrentVersion = 3;
+
+		/// <summary>
+		/// Path to store the cache data to.
+		/// </summary>
+		public FileReference Location
+		{
+			get;
+		}
+
+		/// <summary>
+		/// The Attributes used to produce files, keyed by the absolute file paths.
+		/// </summary>
+		ConcurrentDictionary<FileItem, IoHash> OutputItemToAttributeHash = new ConcurrentDictionary<FileItem, IoHash>();
+
+		/// <summary>
+		/// Whether the dependency cache is dirty and needs to be saved.
+		/// </summary>
+		bool bModified;
+
+		/// <summary>
+		/// Constructor
+		/// </summary>
+		/// <param name="Location">File to store this history in</param>
+		/// <param name="Logger">Logger for output</param>
+		public ActionHistoryLayer(FileReference Location, ILogger Logger)
+		{
+			this.Location = Location;
+
+			if (FileReference.Exists(Location))
+			{
+				Load(Logger);
+			}
+		}
+
+		/// <summary>
+		/// Attempts to load this action history from disk
+		/// </summary>
+		void Load(ILogger Logger)
+		{
+			try
+			{
+				using (BinaryArchiveReader Reader = new BinaryArchiveReader(Location))
+				{
+					int Version = Reader.ReadInt();
+					if (Version != CurrentVersion)
+					{
+						Logger.LogDebug("Unable to read action history from {Location}; version {Version} vs current {CurrentVersion}", Location, Version, CurrentVersion);
+						return;
+					}
+
+					OutputItemToAttributeHash = new ConcurrentDictionary<FileItem, IoHash>(Reader.ReadDictionary(() => Reader.ReadFileItem()!, () => Reader.ReadIoHash() ?? IoHash.Zero)!);
+				}
+			}
+			catch (Exception Ex)
+			{
+				Logger.LogWarning("Unable to read {Location}. See log for additional information.", Location);
+				Logger.LogDebug("{Ex}", ExceptionUtils.FormatExceptionDetails(Ex));
+			}
+		}
+
+		/// <summary>
+		/// Saves this action history to disk
+		/// </summary>
+		public void Save()
+		{
+			if (bModified)
+			{
+				DirectoryReference.CreateDirectory(Location.Directory);
+				using (BinaryArchiveWriter Writer = new BinaryArchiveWriter(Location))
+				{
+					Writer.WriteInt(CurrentVersion);
+					Writer.WriteDictionary(OutputItemToAttributeHash, Key => Writer.WriteFileItem(Key), Value => Writer.WriteIoHash(Value));
+				}
+				bModified = false;
+			}
+		}
+
+		/// <summary>
+		/// Computes the case-invariant IoHash for a string
+		/// </summary>
+		/// <param name="Text">The text to hash</param>
+		/// <returns>IoHash of the string</returns>
+		static IoHash ComputeHash(string Text) => IoHash.Compute(Encoding.Unicode.GetBytes(Text.ToUpperInvariant()));
+
+		/// <summary>
+		/// Gets the producing attributes for the given file
+		/// </summary>
+		/// <param name="File">The output file to look for</param>
+		/// <param name="Attributes">Receives the Attributes  used to produce this file</param>
+		/// <returns>True if Attributes have changed and is updated, false otherwise</returns>
+		public bool UpdateProducingCommandLine(FileItem File, string Attributes)
+		{
+			IoHash NewHash = ComputeHash(Attributes);
+
+			for (; ; )
+			{
+				if (OutputItemToAttributeHash.TryAdd(File, NewHash))
+				{
+					// If this is a new entry we're done
+					bModified = true;
+					return true;
+				}
+				else
+				{
+					if (OutputItemToAttributeHash.TryGetValue(File, out IoHash OldHash))
+					{
+						if (NewHash == OldHash)
+						{
+							// hashes are the same, no update needed
+							return false;
+						}
+						else
+						{
+							// Try to update with the new value
+							if (OutputItemToAttributeHash.TryUpdate(File, NewHash, OldHash))
+							{
+								bModified = true;
+								return true;
+							}
+						}
+					}
+				}
+			}
+		}
+
+		/// <summary>
+		/// Gets the location for the engine action history
+		/// </summary>
+		/// <param name="TargetName">Target name being built</param>
+		/// <param name="Platform">The platform being built</param>
+		/// <param name="TargetType">Type of the target being built</param>
+		/// <param name="Architectures">The target architecture(s)</param>
+		/// <returns>Path to the engine action history for this target</returns>
+		public static FileReference GetEngineLocation(string TargetName, UnrealTargetPlatform Platform, TargetType TargetType, UnrealArchitectures Architectures)
+		{
+			string AppName;
+			if (TargetType == TargetType.Program)
+			{
+				AppName = TargetName;
+			}
+			else
+			{
+				AppName = UEBuildTarget.GetAppNameForTargetType(TargetType);
+			}
+
+			return FileReference.Combine(Unreal.EngineDirectory, UEBuildTarget.GetPlatformIntermediateFolder(Platform, Architectures, false), AppName, "ActionHistory.bin");
+		}
+
+		/// <summary>
+		/// Gets the location of the project action history
+		/// </summary>
+		/// <param name="ProjectFile">Path to the project file</param>
+		/// <param name="Platform">Platform being built</param>
+		/// <param name="TargetName">Name of the target being built</param>
+		/// <param name="Architectures">The target architecture(s)</param>
+		/// <returns>Path to the project action history</returns>
+		public static FileReference GetProjectLocation(FileReference ProjectFile, string TargetName, UnrealTargetPlatform Platform, UnrealArchitectures Architectures)
+		{
+			return FileReference.Combine(ProjectFile.Directory, UEBuildTarget.GetPlatformIntermediateFolder(Platform, Architectures, false), TargetName, "ActionHistory.dat");
+		}
+
+		/// <summary>
+		/// Enumerates all the locations of action history files for the given target
+		/// </summary>
+		/// <param name="ProjectFile">Project file for the target being built</param>
+		/// <param name="TargetName">Name of the target</param>
+		/// <param name="Platform">Platform being built</param>
+		/// <param name="TargetType">The target type</param>
+		/// <param name="Architectures">The target architecture(s)</param>
+		/// <returns>Dependency cache hierarchy for the given project</returns>
+		public static IEnumerable<FileReference> GetFilesToClean(FileReference ProjectFile, string TargetName, UnrealTargetPlatform Platform, TargetType TargetType, UnrealArchitectures Architectures)
+		{
+			if (ProjectFile == null || !Unreal.IsEngineInstalled())
+			{
+				yield return GetEngineLocation(TargetName, Platform, TargetType, Architectures);
+			}
+			if (ProjectFile != null)
+			{
+				yield return GetProjectLocation(ProjectFile, TargetName, Platform, Architectures);
+			}
+		}
+	}
+
+	/// <summary>
+	/// Information about actions producing artifacts under a particular directory
+	/// </summary>
+	[DebuggerDisplay("{BaseDir}")]
+	class ActionHistoryPartition
+	{
+		/// <summary>
+		/// The base directory for this partition
+		/// </summary>
+		public DirectoryReference BaseDir { get; }
+
+		/// <summary>
+		/// Used to ensure exclusive access to the layers list
+		/// </summary>
+		readonly object LockObject = new object();
+
+		/// <summary>
+		/// Map of filename to layer
+		/// </summary>
+		IReadOnlyList<ActionHistoryLayer> Layers = new List<ActionHistoryLayer>();
+
+		/// <summary>
+		/// Construct a new partition
+		/// </summary>
+		/// <param name="BaseDir">The base directory for this partition</param>
+		public ActionHistoryPartition(DirectoryReference BaseDir)
+		{
+			this.BaseDir = BaseDir;
+		}
+
+		/// <summary>
+		/// Attempt to update the producing command line for the given file
+		/// </summary>
+		/// <param name="File">The file to update</param>
+		/// <param name="Attributes">The new attributes</param>
+		/// <param name="Logger">Logger for output</param>
+		/// <returns>True if the attributes were updated, false otherwise</returns>
+		public bool UpdateProducingAttributes(FileItem File, string Attributes, ILogger Logger)
+		{
+			FileReference LayerLocation = GetLayerLocationForFile(File.Location);
+
+			ActionHistoryLayer? Layer = Layers.FirstOrDefault(x => x.Location == LayerLocation);
+			if (Layer == null)
+			{
+				lock (LockObject)
+				{
+					Layer = Layers.FirstOrDefault(x => x.Location == LayerLocation);
+					if (Layer == null)
+					{
+						Layer = new ActionHistoryLayer(LayerLocation, Logger);
+
+						List<ActionHistoryLayer> NewLayers = new List<ActionHistoryLayer>(Layers);
+						NewLayers.Add(Layer);
+						Layers = NewLayers;
+					}
+				}
+			}
+			return Layer.UpdateProducingCommandLine(File, Attributes);
+		}
+
+		/// <summary>
+		/// Get the path to the action history layer to use for the given file
+		/// </summary>
+		/// <param name="Location">Path to the file to use</param>
+		/// <returns>Path to the file</returns>
+		public FileReference GetLayerLocationForFile(FileReference Location)
+		{
+			int Offset = BaseDir.FullName.Length;
+			for (; ; )
+			{
+				int NameOffset = Offset + 1;
+
+				// Get the next directory separator
+				Offset = Location.FullName.IndexOf(Path.DirectorySeparatorChar, NameOffset + 1);
+				if (Offset == -1)
+				{
+					break;
+				}
+
+				// Get the length of the name
+				int NameLength = Offset - NameOffset;
+
+				// Try to find Binaries/<Platform>/ in the path
+				if (MatchPathFragment(Location, NameOffset, NameLength, "Binaries"))
+				{
+					int PlatformOffset = Offset + 1;
+					int PlatformEndOffset = Location.FullName.IndexOf(Path.DirectorySeparatorChar, PlatformOffset);
+					if (PlatformEndOffset != -1)
+					{
+						string PlatformName = Location.FullName.Substring(PlatformOffset, PlatformEndOffset - PlatformOffset);
+						return FileReference.Combine(BaseDir, "Intermediate", "Build", PlatformName, "ActionHistory.bin");
+					}
+				}
+
+				// Try to find /Intermediate/Build/<Platform>/<Target>/<Configuration> in the path
+				if (MatchPathFragment(Location, NameOffset, NameLength, "Intermediate"))
+				{
+					int BuildOffset = Offset + 1;
+					int BuildEndOffset = Location.FullName.IndexOf(Path.DirectorySeparatorChar, BuildOffset);
+					if (BuildEndOffset != -1 && MatchPathFragment(Location, BuildOffset, BuildEndOffset - BuildOffset, "Build"))
+					{
+						// Skip the platform, target/app name, and configuration
+						int EndOffset = BuildEndOffset;
+						for (int Idx = 0; ; Idx++)
+						{
+							EndOffset = Location.FullName.IndexOf(Path.DirectorySeparatorChar, EndOffset + 1);
+							if (EndOffset == -1)
+							{
+								break;
+							}
+							if (Idx == 2)
+							{
+								return FileReference.Combine(BaseDir, Location.FullName.Substring(NameOffset, EndOffset - NameOffset), "ActionHistory.bin");
+							}
+						}
+					}
+				}
+			}
+			return FileReference.Combine(BaseDir, "Intermediate", "Build", "ActionHistory.bin");
+		}
+
+		/// <summary>
+		/// Attempts to match a substring of a path with the given fragment
+		/// </summary>
+		/// <param name="Location">Path to match against</param>
+		/// <param name="Offset">Offset of the substring to match</param>
+		/// <param name="Length">Length of the substring to match</param>
+		/// <param name="Fragment">The path fragment</param>
+		/// <returns>True if the substring matches</returns>
+		static bool MatchPathFragment(FileReference Location, int Offset, int Length, string Fragment)
+		{
+			return Length == Fragment.Length && String.Compare(Location.FullName, Offset, Fragment, 0, Fragment.Length, FileReference.Comparison) == 0;
+		}
+
+		/// <summary>
+		/// Saves the modified layers
+		/// </summary>
+		public void Save()
+		{
+			foreach (ActionHistoryLayer Layer in Layers)
+			{
+				Layer.Save();
+			}
+		}
+	}
+
+	/// <summary>
+	/// A collection of ActionHistory layers
+	/// </summary>
+	class ActionHistory
+	{
+		/// <summary>
+		/// The lock object for this history
+		/// </summary>
+		readonly object LockObject = new object();
+
+		/// <summary>
+		/// List of partitions
+		/// </summary>
+		readonly List<ActionHistoryPartition> Partitions = new List<ActionHistoryPartition>();
+
+		/// <summary>
+		/// Constructor
+		/// </summary>
+		public ActionHistory()
+		{
+			Partitions.Add(new ActionHistoryPartition(Unreal.EngineDirectory));
+		}
+
+		/// <summary>
+		/// Reads a cache from the given location, or creates it with the given settings
+		/// </summary>
+		/// <param name="BaseDir">Base directory for files that this cache should store data for</param>
+		/// <returns>Reference to a dependency cache with the given settings</returns>
+		public void Mount(DirectoryReference BaseDir)
+		{
+			lock (LockObject)
+			{
+				ActionHistoryPartition? Partition = Partitions.FirstOrDefault(x => x.BaseDir == BaseDir);
+				if (Partition == null)
+				{
+					Partition = new ActionHistoryPartition(BaseDir);
+					Partitions.Add(Partition);
+				}
+			}
+		}
+
+		/// <summary>
+		/// Gets the producing command line for the given file
+		/// </summary>
+		/// <param name="File">The output file to look for</param>
+		/// <param name="Attributes">Receives the Attributes used to produce this file</param>
+		/// <param name="Logger">Logger for output</param>
+		/// <returns>True if the output item exists</returns>
+		public bool UpdateProducingAttributes(FileItem File, string Attributes, ILogger Logger)
+		{
+			foreach (ActionHistoryPartition Partition in Partitions)
+			{
+				if (File.Location.IsUnderDirectory(Partition.BaseDir))
+				{
+					return Partition.UpdateProducingAttributes(File, Attributes, Logger);
+				}
+			}
+
+			Logger.LogWarning("File {FileLocation} is not under any action history root directory", File.Location);
+			return false;
+		}
+
+		/// <summary>
+		/// Saves all layers of this action history
+		/// </summary>
+		public void Save()
+		{
+			lock (LockObject)
+			{
+				foreach (ActionHistoryPartition Partition in Partitions)
+				{
+					Partition.Save();
+				}
+			}
+		}
+	}
+}
diff --git a/Engine/Source/Programs/UnrealBuildTool/Actions/ActionSerializer.cs b/Engine/Source/Programs/UnrealBuildTool/Actions/ActionSerializer.cs
new file mode 100644
--- /dev/null
+++ b/Engine/Source/Programs/UnrealBuildTool/Actions/ActionSerializer.cs
@@ -0,0 +1,154 @@
+// Copyright Epic Games, Inc. All Rights Reserved.
+
+using System;
+using System.Collections.Generic;
+using System.Reflection;
+using EpicGames.Core;
+
+namespace UnrealBuildTool
+{
+	/// <summary>
+	/// Interface for a class that can serialize an action type
+	/// </summary>
+	interface IActionSerializer
+	{
+		/// <summary>
+		/// The action type
+		/// </summary>
+		Type Type { get; }
+
+		/// <summary>
+		/// Read the action from an archive
+		/// </summary>
+		/// <param name="Reader">Reader for the action</param>
+		/// <returns>New action</returns>
+		IExternalAction Read(BinaryArchiveReader Reader);
+
+		/// <summary>
+		/// Writes an action to an archive
+		/// </summary>
+		/// <param name="Writer">Writer for the archive</param>
+		/// <param name="Action">The action to write</param>
+		void Write(BinaryArchiveWriter Writer, IExternalAction Action);
+	}
+
+	/// <summary>
+	/// Generic base class for an action serializer
+	/// </summary>
+	/// <typeparam name="TAction"></typeparam>
+	abstract class ActionSerializerBase<TAction> : IActionSerializer where TAction : IExternalAction
+	{
+		/// <inheritdoc/>
+		public Type Type => typeof(TAction);
+
+		/// <inheritdoc/>
+		IExternalAction IActionSerializer.Read(BinaryArchiveReader Reader) => Read(Reader);
+
+		/// <inheritdoc/>
+		void IActionSerializer.Write(BinaryArchiveWriter Writer, IExternalAction Action) => Write(Writer, (TAction)Action);
+
+		/// <summary>
+		/// Read the action from an archive
+		/// </summary>
+		/// <param name="Reader">Reader for the action</param>
+		/// <returns>New action</returns>
+		public abstract TAction Read(BinaryArchiveReader Reader);
+
+		/// <summary>
+		/// Writes an action to an archive
+		/// </summary>
+		/// <param name="Writer">Writer for the archive</param>
+		/// <param name="Action">The action to write</param>
+		public abstract void Write(BinaryArchiveWriter Writer, TAction Action);
+	}
+
+	/// <summary>
+	/// Helper methods for registering serializers and serializing actions
+	/// </summary>
+	static class ActionSerialization
+	{
+		/// <summary>
+		/// Map from type name to deserializing constructor
+		/// </summary>
+		static readonly IReadOnlyDictionary<Type, IActionSerializer> TypeToSerializer;
+
+		/// <summary>
+		/// Map from serializer name to instance
+		/// </summary>
+		static readonly IReadOnlyDictionary<string, IActionSerializer> NameToSerializer;
+
+#pragma warning disable CA1810 // Initialize reference type static fields inline
+		/// <summary>
+		/// Creates a map of type name to constructor
+		/// </summary>
+		/// <returns></returns>
+		static ActionSerialization()
+		{
+			Dictionary<Type, IActionSerializer> TypeToSerializerDict = new Dictionary<Type, IActionSerializer>();
+			Dictionary<string, IActionSerializer> NameToSerializerDict = new Dictionary<string, IActionSerializer>(StringComparer.Ordinal);
+
+			Type[] Types = Assembly.GetExecutingAssembly().GetTypes();
+			foreach (Type Type in Types)
+			{
+				if (Type.IsClass && !Type.IsAbstract && typeof(IActionSerializer).IsAssignableFrom(Type))
+				{
+					IActionSerializer Serializer = (IActionSerializer)Activator.CreateInstance(Type)!;
+					TypeToSerializerDict[Serializer.Type] = Serializer;
+					NameToSerializerDict[Type.Name] = Serializer;
+				}
+			}
+
+			TypeToSerializer = TypeToSerializerDict;
+			NameToSerializer = NameToSerializerDict;
+		}
+#pragma warning restore CA1810
+
+		/// <summary>
+		/// Read an action from the given archive
+		/// </summary>
+		/// <param name="Reader">Reader to deserialize from</param>
+		/// <returns>New action</returns>
+		public static IExternalAction ReadAction(this BinaryArchiveReader Reader)
+		{
+			IActionSerializer Serializer = Reader.ReadObjectReference(() => ReadSerializer(Reader))!;
+			return Serializer.Read(Reader);
+		}
+
+		/// <summary>
+		/// Reads a type name and find its registered constructor from an archive
+		/// </summary>
+		/// <param name="Reader">Archive to read from</param>
+		/// <returns>New constructor info</returns>
+		static IActionSerializer ReadSerializer(BinaryArchiveReader Reader)
+		{
+			string Name = Reader.ReadString()!;
+
+			IActionSerializer? Serializer;
+			if (!NameToSerializer.TryGetValue(Name, out Serializer))
+			{
+				throw new BuildException("Unable to find action type '{0}'", Name);
+			}
+
+			return Serializer;
+		}
+
+		/// <summary>
+		/// Writes an action to the given archive
+		/// </summary>
+		/// <param name="Writer">Writer to serialize the action to</param>
+		/// <param name="Action">Action to serialize</param>
+		public static void WriteAction(this BinaryArchiveWriter Writer, IExternalAction Action)
+		{
+			Type Type = Action.GetType();
+
+			IActionSerializer? Serializer;
+			if (!TypeToSerializer.TryGetValue(Type, out Serializer))
+			{
+				throw new BuildException("Unable to find serializer for action type '{0}'", Type.Name);
+			}
+
+			Writer.WriteObjectReference(Serializer, () => Writer.WriteString(Serializer.GetType().Name));
+			Serializer.Write(Writer, Action);
+		}
+	}
+}
diff --git a/Engine/Source/Programs/UnrealBuildTool/System/Action.cs b/Engine/Source/Programs/UnrealBuildTool/System/Action.cs
deleted file mode 100644
diff --git a/Engine/Source/Programs/UnrealBuildTool/System/ActionExecutor.cs b/Engine/Source/Programs/UnrealBuildTool/System/ActionExecutor.cs
deleted file mode 100644
diff --git a/Engine/Source/Programs/UnrealBuildTool/System/ActionGraph.cs b/Engine/Source/Programs/UnrealBuildTool/System/ActionGraph.cs
deleted file mode 100644
diff --git a/Engine/Source/Programs/UnrealBuildTool/System/ActionGraphBuilder.cs b/Engine/Source/Programs/UnrealBuildTool/System/ActionGraphBuilder.cs
deleted file mode 100644
diff --git a/Engine/Source/Programs/UnrealBuildTool/System/ActionHistory.cs b/Engine/Source/Programs/UnrealBuildTool/System/ActionHistory.cs
deleted file mode 100644
diff --git a/Engine/Source/Programs/UnrealBuildTool/System/ActionSerializer.cs b/Engine/Source/Programs/UnrealBuildTool/System/ActionSerializer.cs
deleted file mode 100644
