The AppleVersionToLLVMVersions in Apple_SDK.json maps Xcode versions to LLVM versions, however, AppleTookChain.QueryClangVersion was mapping Clang front end version to an LLVM back end version. Map Xcode version to achieve desired behaviour.

#Synced-CL 48425151

diff --git a/Engine/Source/Programs/UnrealBuildTool/Platform/Apple/AppleToolChain.cs b/Engine/Source/Programs/UnrealBuildTool/Platform/Apple/AppleToolChain.cs
--- a/Engine/Source/Programs/UnrealBuildTool/Platform/Apple/AppleToolChain.cs
+++ b/Engine/Source/Programs/UnrealBuildTool/Platform/Apple/AppleToolChain.cs
@@ -40,10 +40,10 @@ namespace UnrealBuildTool
 
 			protected override Version QueryClangVersion()
 			{
-				// get the clang version from the clang -v
-				Version AppleVersion = base.QueryClangVersion();
+				// base Apple clang version against known toolchain version included with Xcode.
+				Version AppleVersion = new Version(ApplePlatformSDK.InstalledXcodeVersion.Value ?? "0.0");
 
-					// now look up this version in mappings
+				// now look up this version in mappings
 				for (int MappingIndex = AppleVersionToLLVMVersion.Length - 1; MappingIndex >= 0; MappingIndex--)
 				{
 					if (AppleVersion >= AppleVersionToLLVMVersion[MappingIndex].Item1)
