UnrealBuildTool: Do not add .NET 10 projects to the Visual Studio 2022 solution as they are unsupported

#rnx

#Synced-CL 47605596

diff --git a/Engine/Source/Programs/UnrealBuildTool/ProjectFiles/VisualStudio/VCProjectFileGenerator.cs b/Engine/Source/Programs/UnrealBuildTool/ProjectFiles/VisualStudio/VCProjectFileGenerator.cs
--- a/Engine/Source/Programs/UnrealBuildTool/ProjectFiles/VisualStudio/VCProjectFileGenerator.cs
+++ b/Engine/Source/Programs/UnrealBuildTool/ProjectFiles/VisualStudio/VCProjectFileGenerator.cs
@@ -410,6 +410,22 @@ namespace UnrealBuildTool
 					}
 				}
 			}
+
+			if (bIncludeDotNetPrograms && Settings.ProjectFileFormat == VCProjectFileFormat.VisualStudio2022)
+			{
+#if NET10_0_OR_GREATER
+				Logger.LogInformation("Visual Studio 2022 does not support .NET 10.0 C# projects, these projects will not be added to the generated solution.");
+				Logger.LogInformation("Please generate the Visual Studio 2026 solution if .NET 10.0 C# project support is required.");
+				bIncludeDotNetPrograms = false;
+#else
+				if (System.Environment.Version.Major >= 10)
+				{
+					Logger.LogInformation("Visual Studio 2022 does not support .NET 10.0 C# projects, these projects will not be added to the generated solution as the current dotnet runtime version is '{Version}'. Compiling with an older dotnet SDK from Visual Studio can cause errors.", System.Runtime.InteropServices.RuntimeInformation.FrameworkDescription);
+					Logger.LogInformation("Please generate the Visual Studio 2026 solution if .NET 10.0 C# project support is required.");
+					bIncludeDotNetPrograms = false;
+				}
+#endif
+			}
 		}
 
 		/// <summary>
