[UBT]
* Deleted code that deleted modules generated UHT folder when uht generated file count went from >0 to 0 while still existing. This is a very uncommon scenario since it will not cover modules that are renamed/deleted and going from >0 to 0 generated files must be very rare and not worth the effort of cleaning up

#Synced-CL 48464249

diff --git a/Engine/Source/Programs/UnrealBuildTool/System/UHTExecution.cs b/Engine/Source/Programs/UnrealBuildTool/System/UHTExecution.cs
--- a/Engine/Source/Programs/UnrealBuildTool/System/UHTExecution.cs
+++ b/Engine/Source/Programs/UnrealBuildTool/System/UHTExecution.cs
@@ -597,22 +597,6 @@ namespace UnrealBuildTool
 					];
 					UObjectModuleHeaders.Add(new UHTModuleHeaderInfo(ModuleDirectoryItem, ReflectedHeaderFiles, Module.Rules.bUsePrecompiled));
 				}
-				else
-				{
-					// Remove any stale generated code directory
-					if (Module.GeneratedCodeDirectoryUHT != null && !Module.Rules.bUsePrecompiled)
-					{
-						if (DirectoryReference.Exists(Module.GeneratedCodeDirectoryUHT))
-						{
-							Directory.Delete(Module.GeneratedCodeDirectoryUHT.FullName, true);
-							// Also delete parent directory if now empty
-							if (!Directory.EnumerateFileSystemEntries(Module.GeneratedCodeDirectory!.FullName).Any())
-							{
-								Directory.Delete(Module.GeneratedCodeDirectory!.FullName, true);
-							}
-						}
-					}
-				}
 			}
 
 			// Set Module.bHasUObjects for any IncludePathModules not already processed.
