Disable over-aggressive null static analysis warning on Android

diff --git a/Engine/Source/Programs/UnrealBuildTool/Platform/Android/UEBuildAndroid.cs b/Engine/Source/Programs/UnrealBuildTool/Platform/Android/UEBuildAndroid.cs
--- a/Engine/Source/Programs/UnrealBuildTool/Platform/Android/UEBuildAndroid.cs
+++ b/Engine/Source/Programs/UnrealBuildTool/Platform/Android/UEBuildAndroid.cs
@@ -429,6 +429,10 @@ namespace UnrealBuildTool
 		/// <param name="Target">The target being build</param>
 		public override void ModifyModuleRulesForActivePlatform(string ModuleName, ModuleRules Rules, ReadOnlyTargetRules Target)
 		{
+			if (Target.StaticAnalyzer == StaticAnalyzer.Default)
+			{
+				Rules.StaticAnalyzerDisabledCheckers.UnionWith(["nullability.NullPassedToNonnull"]);
+			}
 		}
 
 		public static bool IsAsisEnabled(ReadOnlyTargetRules Target)
