Remove unused artifact cache functionality from UBT. Caching is carried out by UBA, so this is not likely to be required in future.

#Synced-CL 48546779

diff --git a/Engine/Source/Programs/Shared/UnrealBuildTool.Tests/ArtifactsTest.cs b/Engine/Source/Programs/Shared/UnrealBuildTool.Tests/ArtifactsTest.cs
deleted file mode 100644
diff --git a/Engine/Source/Programs/UnrealBuildTool/Actions/ActionExecutor.cs b/Engine/Source/Programs/UnrealBuildTool/Actions/ActionExecutor.cs
--- a/Engine/Source/Programs/UnrealBuildTool/Actions/ActionExecutor.cs
+++ b/Engine/Source/Programs/UnrealBuildTool/Actions/ActionExecutor.cs
@@ -6,7 +6,6 @@ using System.Threading.Tasks;
 using EpicGames.Core;
 using EpicGames.LogMatchers;
 using Microsoft.Extensions.Logging;
-using UnrealBuildTool.Artifacts;
 
 namespace UnrealBuildTool
 {
@@ -60,9 +59,8 @@ namespace UnrealBuildTool
 		/// </summary>
 		/// <param name="ActionsToExecute">Actions to be executed</param>
 		/// <param name="Logger">Logger for output</param>
-		/// <param name="actionArtifactCache">Cache used to read/write action outputs.</param>
 		/// <returns>True if the build succeeded, false otherwise</returns>
-		public abstract Task<bool> ExecuteActionsAsync(IEnumerable<LinkedAction> ActionsToExecute, ILogger Logger, IActionArtifactCache? actionArtifactCache = null);
+		public abstract Task<bool> ExecuteActionsAsync(IEnumerable<LinkedAction> ActionsToExecute, ILogger Logger);
 
 		/// <summary>
 		/// When completed, get a populated telemetry event for this executor
diff --git a/Engine/Source/Programs/UnrealBuildTool/Actions/ActionGraph.cs b/Engine/Source/Programs/UnrealBuildTool/Actions/ActionGraph.cs
--- a/Engine/Source/Programs/UnrealBuildTool/Actions/ActionGraph.cs
+++ b/Engine/Source/Programs/UnrealBuildTool/Actions/ActionGraph.cs
@@ -16,7 +16,6 @@ using EpicGames.Core;
 using Microsoft.Extensions.Logging;
 using OpenTracing.Util;
 using UnrealBuildBase;
-using UnrealBuildTool.Artifacts;
 
 namespace UnrealBuildTool
 {
@@ -476,7 +475,7 @@ namespace UnrealBuildTool
 		/// <summary>
 		/// Executes a list of actions.
 		/// </summary>
-		public static async Task ExecuteActionsAsync(BuildConfiguration BuildConfiguration, List<LinkedAction> ActionsToExecute, List<TargetDescriptor> TargetDescriptors, ILogger Logger, IActionArtifactCache? actionArtifactCache = null, ActionExecutor? executor = null)
+		public static async Task ExecuteActionsAsync(BuildConfiguration BuildConfiguration, List<LinkedAction> ActionsToExecute, List<TargetDescriptor> TargetDescriptors, ILogger Logger, ActionExecutor? executor = null)
 		{
 			if (ActionsToExecute.Count == 0)
 			{
@@ -485,7 +484,7 @@ namespace UnrealBuildTool
 				if (executor != null)
 				{
 					// We still want to run the executor to give it a chance to clean up in case it has tasks running that can be cancelled
-					await InternalExecuteActions(executor, ActionsToExecute, Logger, actionArtifactCache);
+					await InternalExecuteActions(executor, ActionsToExecute, Logger);
 				}
 			}
 			else
@@ -494,22 +493,22 @@ namespace UnrealBuildTool
 				if (executor == null)
 				{
 					using ActionExecutor SelectedExecutor = SelectExecutor(BuildConfiguration, ActionsToExecute.Count, TargetDescriptors, Logger);
-					await InternalExecuteActions(SelectedExecutor, ActionsToExecute, Logger, actionArtifactCache);
+					await InternalExecuteActions(SelectedExecutor, ActionsToExecute, Logger);
 				}
 				else
 				{
-					await InternalExecuteActions(executor, ActionsToExecute, Logger, actionArtifactCache);
+					await InternalExecuteActions(executor, ActionsToExecute, Logger);
 				}
 			}
 		}
 
-		static async Task InternalExecuteActions(ActionExecutor Executor, List<LinkedAction> ActionsToExecute, ILogger Logger, IActionArtifactCache? actionArtifactCache = null)
+		static async Task InternalExecuteActions(ActionExecutor Executor, List<LinkedAction> ActionsToExecute, ILogger Logger)
 		{
 			Logger.LogInformation("Using {ExecutorName} executor to run {ActionCount} action(s)", Executor.Name, ActionsToExecute.Count);
 
 			// Execute the build
 			Stopwatch Timer = Stopwatch.StartNew();
-			bool Result = await Executor.ExecuteActionsAsync(ActionsToExecute, Logger, actionArtifactCache);
+			bool Result = await Executor.ExecuteActionsAsync(ActionsToExecute, Logger);
 			Executor.PostTelemetryEvent();
 
 			Logger.LogInformation("Total time in {ExecutorName} executor: {TotalSeconds:0.00} seconds", Executor.Name, Timer.Elapsed.TotalSeconds);
diff --git a/Engine/Source/Programs/UnrealBuildTool/Artifacts/ActionArtifactCache.cs b/Engine/Source/Programs/UnrealBuildTool/Artifacts/ActionArtifactCache.cs
deleted file mode 100644
diff --git a/Engine/Source/Programs/UnrealBuildTool/Artifacts/Artifacts.cs b/Engine/Source/Programs/UnrealBuildTool/Artifacts/Artifacts.cs
deleted file mode 100644
diff --git a/Engine/Source/Programs/UnrealBuildTool/Artifacts/HordeStorageArtifactCache.cs b/Engine/Source/Programs/UnrealBuildTool/Artifacts/HordeStorageArtifactCache.cs
deleted file mode 100644
diff --git a/Engine/Source/Programs/UnrealBuildTool/Executors/Experimental/FASTBuild.cs b/Engine/Source/Programs/UnrealBuildTool/Executors/Experimental/FASTBuild.cs
--- a/Engine/Source/Programs/UnrealBuildTool/Executors/Experimental/FASTBuild.cs
+++ b/Engine/Source/Programs/UnrealBuildTool/Executors/Experimental/FASTBuild.cs
@@ -23,7 +23,6 @@ using System.Threading.Tasks;
 using EpicGames.Core;
 using Microsoft.Extensions.Logging;
 using UnrealBuildBase;
-using UnrealBuildTool.Artifacts;
 
 namespace UnrealBuildTool
 {
@@ -474,7 +473,7 @@ namespace UnrealBuildTool
 
 		/// <inheritdoc/>
 		[SupportedOSPlatform("windows")]
-		public override async Task<bool> ExecuteActionsAsync(IEnumerable<LinkedAction> Actions, ILogger Logger, IActionArtifactCache? actionArtifactCache)
+		public override async Task<bool> ExecuteActionsAsync(IEnumerable<LinkedAction> Actions, ILogger Logger)
 		{
 			if (!Actions.Any())
 			{
@@ -496,7 +495,7 @@ namespace UnrealBuildTool
 				return result;
 			}
 
-			return await LocalExecutor.ExecuteActionsAsync(Actions, Logger, actionArtifactCache);
+			return await LocalExecutor.ExecuteActionsAsync(Actions, Logger);
 		}
 
 		private void AddText(string StringToWrite)
diff --git a/Engine/Source/Programs/UnrealBuildTool/Executors/ImmediateActionQueue.cs b/Engine/Source/Programs/UnrealBuildTool/Executors/ImmediateActionQueue.cs
--- a/Engine/Source/Programs/UnrealBuildTool/Executors/ImmediateActionQueue.cs
+++ b/Engine/Source/Programs/UnrealBuildTool/Executors/ImmediateActionQueue.cs
@@ -9,27 +9,9 @@ using EpicGames.Core;
 using Microsoft.CodeAnalysis;
 using Microsoft.Extensions.Logging;
 using UnrealBuildBase;
-using UnrealBuildTool.Artifacts;
 
 namespace UnrealBuildTool
 {
-	/// <summary>
-	/// Defines the phase of execution
-	/// </summary>
-	enum ActionPhase : byte
-	{
-
-		/// <summary>
-		/// Check for already existing artifacts for the inputs
-		/// </summary>
-		ArtifactCheck,
-
-		/// <summary>
-		/// Compile the action
-		/// </summary>
-		Compile,
-	}
-
 	/// <summary>
 	/// Defines the type of the runner.
 	/// </summary>
@@ -51,12 +33,11 @@ namespace UnrealBuildTool
 	/// Actions are assigned a runner when they need executing
 	/// </summary>
 	/// <param name="Type">Type of the runner</param>
-	/// <param name="ActionPhase">The action phase that this step supports</param>
 	/// <param name="RunAction">Function to be run queue running a task</param>
 	/// <param name="UseActionWeights">If true, use the action weight as a secondary limit</param>
 	/// <param name="MaxActions">Maximum number of action actions</param>
 	/// <param name="MaxActionWeight">Maximum weight of actions</param>
-	record ImmediateActionQueueRunner(ImmediateActionQueueRunnerType Type, ActionPhase ActionPhase, Func<LinkedAction, Func<Task>?> RunAction, bool UseActionWeights = false, int MaxActions = Int32.MaxValue, double MaxActionWeight = Int32.MaxValue)
+	record ImmediateActionQueueRunner(ImmediateActionQueueRunnerType Type, Func<LinkedAction, Func<Task>?> RunAction, bool UseActionWeights = false, int MaxActions = Int32.MaxValue, double MaxActionWeight = Int32.MaxValue)
 	{
 		/// <summary>
 		/// Current number of active actions
@@ -149,11 +130,6 @@ namespace UnrealBuildTool
 			/// </summary>
 			public LinkedAction Action;
 
-			/// <summary>
-			/// Phase of the action
-			/// </summary>
-			public ActionPhase Phase;
-
 			/// <summary>
 			/// Current status of the execution
 			/// </summary>
@@ -206,8 +182,6 @@ namespace UnrealBuildTool
 		public int StopCompilationAfterNumErrors = Int32.MaxValue;
 
 		public int CompletedActions => _completedActions;
-		public int CacheHitActions => _cacheHitActions;
-		public int CacheMissActions => _cacheMissActions;
 
 		/// <summary>
 		/// Return true if the queue is done
@@ -239,16 +213,6 @@ namespace UnrealBuildTool
 		/// </summary>
 		private int _completedActions = 0;
 
-		/// <summary>
-		/// Tracks the number of completed actions from cache hits.
-		/// </summary>
-		private int _cacheHitActions = 0;
-
-		/// <summary>
-		/// Tracks the number of unsuccessful actions from cache misses.
-		/// </summary>
-		private int _cacheMissActions = 0;
-
 		/// <summary>
 		/// Flags used to track how StartManyActions should run
 		/// </summary>
@@ -284,24 +248,18 @@ namespace UnrealBuildTool
 		/// </summary>
 		private readonly TaskCompletionSource _doneTaskSource = new();
 
-		/// <summary>
-		/// If set, artifact cache used to retrieve previously compiled results and save new results
-		/// </summary>
-		private readonly IActionArtifactCache? _actionArtifactCache;
-
 		private bool _disposedValue;
 
 		/// <summary>
 		/// Construct a new instance of the action queue
 		/// </summary>
 		/// <param name="actions">Collection of actions</param>
-		/// <param name="actionArtifactCache">If true, the artifact cache is being used</param>
 		/// <param name="maxActionArtifactCacheTasks">Max number of concurrent artifact cache tasks</param>
 		/// <param name="progressWriterText">Text to be displayed with the progress writer</param>
 		/// <param name="writeToolOutput">Action to invoke when writing tool output</param>
 		/// <param name="flushToolOutput">Action to invoke when flushing tool output</param>
 		/// <param name="logger">Logging interface</param>
-		public ImmediateActionQueue(IEnumerable<LinkedAction> actions, IActionArtifactCache? actionArtifactCache, int maxActionArtifactCacheTasks, string progressWriterText, Action<string> writeToolOutput, System.Action flushToolOutput, ILogger logger)
+		public ImmediateActionQueue(IEnumerable<LinkedAction> actions, int maxActionArtifactCacheTasks, string progressWriterText, Action<string> writeToolOutput, System.Action flushToolOutput, ILogger logger)
 		:	base(actions.Count(), progressWriterText, writeToolOutput, flushToolOutput, logger)
 		{
 			CommandLine.ParseArguments(Environment.GetCommandLineArgs(), this, logger);
@@ -310,10 +268,6 @@ namespace UnrealBuildTool
 			int count = actions.Count();
 			Actions = new ActionState[count];
 
-			_actionArtifactCache = actionArtifactCache;
-
-			bool readArtifacts = _actionArtifactCache != null && _actionArtifactCache.EnableReads;
-			ActionPhase initialPhase = readArtifacts ? ActionPhase.ArtifactCheck : ActionPhase.Compile;
 			int index = 0;
 			foreach (LinkedAction action in actions)
 			{
@@ -322,37 +276,10 @@ namespace UnrealBuildTool
 				{
 					Action = action,
 					Status = ActionStatus.Queued,
-					Phase = action.ArtifactMode.HasFlag(ArtifactMode.Enabled) ? initialPhase : ActionPhase.Compile,
 					PrereqActionsSortIndex = action.PrerequisiteActions.Select(a => a.SortIndex).ToArray()
 				};
 			}
 
-			if (readArtifacts)
-			{
-				Func<Task> runAction(LinkedAction action)
-				{
-					return new Func<Task>(async () =>
-					{
-						ActionArtifactResult result = await _actionArtifactCache!.CompleteActionFromCacheAsync(action, CancellationToken);
-						if (result.Success)
-						{
-							ExecuteResults results = new(result.LogLines, result.ExitCode, TimeSpan.Zero, TimeSpan.Zero, 0, action, "[Cache]");
-							Interlocked.Increment(ref _cacheHitActions);
-							OnActionCompleted(action, results.ExitCode == 0, results);
-							OnArtifactsRead?.Invoke(action);
-						}
-						else
-						{
-							Interlocked.Increment(ref _cacheMissActions);
-							OnArtifactsMiss?.Invoke(action);
-							RequeueAction(action);
-						}
-					});
-				}
-
-				_runners.Add(new(ImmediateActionQueueRunnerType.Automatic, ActionPhase.ArtifactCheck, runAction, false, maxActionArtifactCacheTasks, 0));
-			}
-
 			// Cancel the queue when Ctrl-C is pressed
 			Console.CancelKeyPress += CancelKeyPress;
 		}
@@ -399,7 +326,7 @@ namespace UnrealBuildTool
 		/// <returns>Created runner</returns>
 		public ImmediateActionQueueRunner CreateAutomaticRunner(Func<LinkedAction, Func<Task>?> runAction, bool useActionWeights, int maxActions, double maxActionWeight)
 		{
-			ImmediateActionQueueRunner runner = new(ImmediateActionQueueRunnerType.Automatic, ActionPhase.Compile, runAction, useActionWeights, maxActions, maxActionWeight);
+			ImmediateActionQueueRunner runner = new(ImmediateActionQueueRunnerType.Automatic, runAction, useActionWeights, maxActions, maxActionWeight);
 			_runners.Add(runner);
 			return runner;
 		}
@@ -414,7 +341,7 @@ namespace UnrealBuildTool
 		/// <returns>Created runner</returns>
 		public ImmediateActionQueueRunner CreateManualRunner(Func<LinkedAction, Func<Task>?> runAction, bool useActionWeights = false, int maxActions = Int32.MaxValue, double maxActionWeight = Double.MaxValue)
 		{
-			ImmediateActionQueueRunner runner = new(ImmediateActionQueueRunnerType.Manual, ActionPhase.Compile, runAction, useActionWeights, maxActions, maxActionWeight);
+			ImmediateActionQueueRunner runner = new(ImmediateActionQueueRunnerType.Manual, runAction, useActionWeights, maxActions, maxActionWeight);
 			_runners.Add(runner);
 			return runner;
 		}
@@ -489,15 +416,11 @@ namespace UnrealBuildTool
 		/// <param name="totalActions">Out parameter, the total number of actions</param>
 		/// <param name="succeededActions">Out parameter, the number of successful actions</param>
 		/// <param name="failedActions">Out parameter, the number of failed actions</param>
-		/// <param name="cacheHitActions">Out parameter, the number of cache hit actions</param>
-		/// <param name="cacheMissActions">Out parameter, the number of cache miss actions</param>
-		public void GetActionResultCounts(out int totalActions, out int succeededActions, out int failedActions, out int cacheHitActions, out int cacheMissActions)
+		public void GetActionResultCounts(out int totalActions, out int succeededActions, out int failedActions)
 		{
 			totalActions = Actions.Length;
 			succeededActions = Results.Where(x => x?.ExitCode == 0).Count();
 			failedActions = Results.Where(x => x != null && x.ExitCode != 0).Count();
-			cacheHitActions = _cacheHitActions;
-			cacheMissActions = _cacheMissActions;
 		}
 
 		/// <summary>
@@ -511,7 +434,6 @@ namespace UnrealBuildTool
 			{
 				ActionState actionState = Actions[actionIndex];
 				if (actionState.Status == ActionStatus.Queued &&
-					actionState.Phase == ActionPhase.Compile &&
 					GetActionReadyState(actionState) == ActionReadyState.Ready)
 				{
 					yield return actionState.Action;
@@ -693,7 +615,6 @@ namespace UnrealBuildTool
 
 					case ActionReadyState.Error:
 						Actions[actionIndex].Status = ActionStatus.Error;
-						Actions[actionIndex].Phase = ActionPhase.Compile;
 						completedActions++;
 						break;
 
@@ -702,7 +623,7 @@ namespace UnrealBuildTool
 						Func<Task>? runAction = null;
 						if (runner != null)
 						{
-							if (runner.IsUnderLimits && runner.ActionPhase == Actions[actionIndex].Phase)
+							if (runner.IsUnderLimits)
 							{
 								runAction = runner.RunAction(Actions[actionIndex].Action);
 								if (runAction != null)
@@ -716,7 +637,7 @@ namespace UnrealBuildTool
 							foreach (ImmediateActionQueueRunner tryRunner in _runners)
 							{
 								if (tryRunner.Type == ImmediateActionQueueRunnerType.Automatic &&
-									tryRunner.IsUnderLimits && tryRunner.ActionPhase == Actions[actionIndex].Phase)
+									tryRunner.IsUnderLimits)
 								{
 									runAction = tryRunner.RunAction(Actions[actionIndex].Action);
 									if (runAction != null)
@@ -857,18 +778,6 @@ namespace UnrealBuildTool
 				// Use to track that action states have changed
 				_lastActionChange++;
 
-				// If we are doing an artifact check, then move to compile phase
-				bool wasArtifactCheck = false;
-				if (Actions[actionIndex].Phase == ActionPhase.ArtifactCheck)
-				{
-					wasArtifactCheck = true;
-					Actions[actionIndex].Phase = ActionPhase.Compile;
-					if (status != ActionStatus.Finished && (results?.ExitCode ?? 0) == 0)
-					{
-						status = ActionStatus.Queued;
-					}
-				}
-
 				Actions[actionIndex].Status = status;
 				Actions[actionIndex].Runner = null;
 				if (results != null)
@@ -896,12 +805,6 @@ namespace UnrealBuildTool
 						break;
 
 					case ActionStatus.Finished:
-						// Notify the artifact handler of the action completing.  We don't wait on the resulting task.  The
-						// cache is required to remember any pending saves and a final call to Flush will wait for everything to complete.
-						if (!wasArtifactCheck)
-						{
-							_actionArtifactCache?.ActionCompleteAsync(action, CancellationToken);
-						}
 						completedActions = 1;
 						break;
 
@@ -1036,14 +939,6 @@ namespace UnrealBuildTool
 		{
 			foreach (int prereqIndex in action.PrereqActionsSortIndex)
 			{
-
-				// To avoid doing artifact checks on actions that might need compiling,
-				// we first make sure the action is in the compile phase
-				if (Actions[prereqIndex].Phase != ActionPhase.Compile)
-				{
-					return ActionReadyState.NotReady;
-				}
-
 				// Respect the compile status of the action
 				switch (Actions[prereqIndex].Status)
 				{
diff --git a/Engine/Source/Programs/UnrealBuildTool/Executors/ParallelExecutor.cs b/Engine/Source/Programs/UnrealBuildTool/Executors/ParallelExecutor.cs
--- a/Engine/Source/Programs/UnrealBuildTool/Executors/ParallelExecutor.cs
+++ b/Engine/Source/Programs/UnrealBuildTool/Executors/ParallelExecutor.cs
@@ -10,7 +10,6 @@ using System.Threading.Tasks;
 using EpicGames.Core;
 using Microsoft.Extensions.Logging;
 using UnrealBuildBase;
-using UnrealBuildTool.Artifacts;
 
 namespace UnrealBuildTool
 {
@@ -174,13 +173,12 @@ namespace UnrealBuildTool
 		/// Create an action queue
 		/// </summary>
 		/// <param name="actionsToExecute">Actions to be executed</param>
-		/// <param name="actionArtifactCache">Artifact cache</param>
 		/// <param name="maxActionArtifactCacheTasks">Max artifact tasks that can execute in parallel</param>
 		/// <param name="logger">Logging interface</param>
 		/// <returns>Action queue</returns>
-		public ImmediateActionQueue CreateActionQueue(IEnumerable<LinkedAction> actionsToExecute, IActionArtifactCache? actionArtifactCache, int maxActionArtifactCacheTasks, ILogger logger)
+		public ImmediateActionQueue CreateActionQueue(IEnumerable<LinkedAction> actionsToExecute, int maxActionArtifactCacheTasks, ILogger logger)
 		{
-			return new(actionsToExecute, actionArtifactCache, maxActionArtifactCacheTasks, "Compiling C++ source code...", x => WriteToolOutput(x), () => FlushToolOutput(), logger)
+			return new(actionsToExecute, maxActionArtifactCacheTasks, "Compiling C++ source code...", x => WriteToolOutput(x), () => FlushToolOutput(), logger)
 			{
 				ShowCompilationTimes = bShowCompilationTimes,
 				ShowCPUUtilization = bShowCPUUtilization,
@@ -194,7 +192,7 @@ namespace UnrealBuildTool
 		}
 
 		/// <inheritdoc/>
-		public override async Task<bool> ExecuteActionsAsync(IEnumerable<LinkedAction> ActionsToExecute, ILogger Logger, IActionArtifactCache? actionArtifactCache)
+		public override async Task<bool> ExecuteActionsAsync(IEnumerable<LinkedAction> ActionsToExecute, ILogger Logger)
 		{
 			if (!ActionsToExecute.Any())
 			{
@@ -208,19 +206,19 @@ namespace UnrealBuildTool
 			bool useAutomaticQueue = true;
 			if (useAutomaticQueue)
 			{
-				using ImmediateActionQueue queue = CreateActionQueue(ActionsToExecute, actionArtifactCache, NumParallelProcesses, Logger);
+				using ImmediateActionQueue queue = CreateActionQueue(ActionsToExecute, NumParallelProcesses, Logger);
 				int actionLimit = Math.Min(NumParallelProcesses, queue.TotalActions);
 				queue.CreateAutomaticRunner(action => RunAction(queue, action), bUseActionWeights, actionLimit, NumParallelProcesses);
 				queue.Start();
 				queue.StartManyActions();
 				result = await queue.RunTillDone();
 
-				queue.GetActionResultCounts(out int totalActions, out int succeededActions, out int failedActions, out int cacheHitActions, out int cacheMissActions);
-				telemetryEvent = new TelemetryExecutorEvent(Name, startTimeUTC, result, totalActions, succeededActions, failedActions, cacheHitActions, cacheMissActions, DateTime.UtcNow, queue.ProcessGroup.PeakProcessMemoryUsed, queue.ProcessGroup.PeakJobMemoryUsed);
+				queue.GetActionResultCounts(out int totalActions, out int succeededActions, out int failedActions);
+				telemetryEvent = new TelemetryExecutorEvent(Name, startTimeUTC, result, totalActions, succeededActions, failedActions, cacheHitActions: 0, cacheMissActions: 0, DateTime.UtcNow, queue.ProcessGroup.PeakProcessMemoryUsed, queue.ProcessGroup.PeakJobMemoryUsed);
 			}
 			else
 			{
-				using ImmediateActionQueue queue = CreateActionQueue(ActionsToExecute, actionArtifactCache, NumParallelProcesses, Logger);
+				using ImmediateActionQueue queue = CreateActionQueue(ActionsToExecute, NumParallelProcesses, Logger);
 				int actionLimit = Math.Min(NumParallelProcesses, queue.TotalActions);
 				ImmediateActionQueueRunner runner = queue.CreateManualRunner(action => RunAction(queue, action), bUseActionWeights, actionLimit, actionLimit);
 				queue.Start();
@@ -228,8 +226,8 @@ namespace UnrealBuildTool
 				queue.StartManyActions();
 				result = await queue.RunTillDone();
 
-				queue.GetActionResultCounts(out int totalActions, out int succeededActions, out int failedActions, out int cacheHitActions, out int cacheMissActions);
-				telemetryEvent = new TelemetryExecutorEvent(Name, startTimeUTC, result, totalActions, succeededActions, failedActions, cacheHitActions, cacheMissActions, DateTime.UtcNow, queue.ProcessGroup.PeakProcessMemoryUsed, queue.ProcessGroup.PeakJobMemoryUsed);
+				queue.GetActionResultCounts(out int totalActions, out int succeededActions, out int failedActions);
+				telemetryEvent = new TelemetryExecutorEvent(Name, startTimeUTC, result, totalActions, succeededActions, failedActions, cacheHitActions: 0, cacheMissActions: 0, DateTime.UtcNow, queue.ProcessGroup.PeakProcessMemoryUsed, queue.ProcessGroup.PeakJobMemoryUsed);
 			}
 
 			return result;
diff --git a/Engine/Source/Programs/UnrealBuildTool/Executors/SNDBS.cs b/Engine/Source/Programs/UnrealBuildTool/Executors/SNDBS.cs
--- a/Engine/Source/Programs/UnrealBuildTool/Executors/SNDBS.cs
+++ b/Engine/Source/Programs/UnrealBuildTool/Executors/SNDBS.cs
@@ -17,7 +17,6 @@ using System.Threading.Tasks;
 using EpicGames.Core;
 using Microsoft.Extensions.Logging;
 using UnrealBuildBase;
-using UnrealBuildTool.Artifacts;
 
 namespace UnrealBuildTool
 {
@@ -230,7 +229,7 @@ namespace UnrealBuildTool
 		public override TelemetryExecutorEvent? GetTelemetryEvent() => telemetryEvent;
 
 		/// <inheritdoc/>
-		public override Task<bool> ExecuteActionsAsync(IEnumerable<LinkedAction> ActionsToExecute, ILogger Logger, IActionArtifactCache? actionArtifactCache)
+		public override Task<bool> ExecuteActionsAsync(IEnumerable<LinkedAction> ActionsToExecute, ILogger Logger)
 		{
 			return Task.FromResult(ExecuteActions(ActionsToExecute, Logger));
 		}
diff --git a/Engine/Source/Programs/UnrealBuildTool/Executors/UnrealBuildAccelerator/UBAExecutor.cs b/Engine/Source/Programs/UnrealBuildTool/Executors/UnrealBuildAccelerator/UBAExecutor.cs
--- a/Engine/Source/Programs/UnrealBuildTool/Executors/UnrealBuildAccelerator/UBAExecutor.cs
+++ b/Engine/Source/Programs/UnrealBuildTool/Executors/UnrealBuildAccelerator/UBAExecutor.cs
@@ -13,7 +13,6 @@ using EpicGames.Core;
 using EpicGames.UBA;
 using Microsoft.Extensions.Logging;
 using UnrealBuildBase;
-using UnrealBuildTool.Artifacts;
 using ILogger = Microsoft.Extensions.Logging.ILogger;
 
 namespace UnrealBuildTool
@@ -485,7 +484,7 @@ namespace UnrealBuildTool
 			}
 		}
 
-		public override async Task<bool> ExecuteActionsAsync(IEnumerable<LinkedAction> inputActions, Microsoft.Extensions.Logging.ILogger logger, IActionArtifactCache? actionArtifactCache)
+		public override async Task<bool> ExecuteActionsAsync(IEnumerable<LinkedAction> inputActions, Microsoft.Extensions.Logging.ILogger logger)
 		{
 			// Disable cache server when using hot reload path (there would likely never be any cache hits anyway)
 			// Note we can't do this earlier because HotReloadMode is not set until actions are created
diff --git a/Engine/Source/Programs/UnrealBuildTool/Executors/XGE.cs b/Engine/Source/Programs/UnrealBuildTool/Executors/XGE.cs
--- a/Engine/Source/Programs/UnrealBuildTool/Executors/XGE.cs
+++ b/Engine/Source/Programs/UnrealBuildTool/Executors/XGE.cs
@@ -17,7 +17,6 @@ using EpicGames.Core;
 using Microsoft.Extensions.Logging;
 using Microsoft.Win32;
 using UnrealBuildBase;
-using UnrealBuildTool.Artifacts;
 
 namespace UnrealBuildTool
 {
@@ -399,7 +398,7 @@ namespace UnrealBuildTool
 		}
 
 		/// <inheritdoc/>
-		public override Task<bool> ExecuteActionsAsync(IEnumerable<LinkedAction> ActionsToExecute, ILogger Logger, IActionArtifactCache? actionArtifactCache)
+		public override Task<bool> ExecuteActionsAsync(IEnumerable<LinkedAction> ActionsToExecute, ILogger Logger)
 		{
 			return Task.FromResult(ExecuteActions(ActionsToExecute, Logger));
 		}
diff --git a/Engine/Source/Programs/UnrealBuildTool/Modes/BuildMode.cs b/Engine/Source/Programs/UnrealBuildTool/Modes/BuildMode.cs
--- a/Engine/Source/Programs/UnrealBuildTool/Modes/BuildMode.cs
+++ b/Engine/Source/Programs/UnrealBuildTool/Modes/BuildMode.cs
@@ -561,7 +561,7 @@ namespace UnrealBuildTool
 					if (PreprocessActions.Count > 0)
 					{
 						Logger.LogInformation("Updating module dependencies...");
-						await ActionGraph.ExecuteActionsAsync(BuildConfiguration, PreprocessActions, TargetDescriptors, Logger, null, executor);
+						await ActionGraph.ExecuteActionsAsync(BuildConfiguration, PreprocessActions, TargetDescriptors, Logger, executor);
 
 						foreach (FileItem ProducedItem in PreprocessActions.SelectMany(x => x.ProducedItems))
 						{
@@ -797,7 +797,7 @@ namespace UnrealBuildTool
 					//using ITimelineEvent ExcuteTimer = Timeline.ScopeEvent("Executing ActionGraph");
 					using (GlobalTracer.Instance.BuildSpan("ActionGraph.ExecuteActions()").StartActive())
 					{
-						await ActionGraph.ExecuteActionsAsync(BuildConfiguration, MergedActionsToExecute, TargetDescriptors, Logger, null, executor);
+						await ActionGraph.ExecuteActionsAsync(BuildConfiguration, MergedActionsToExecute, TargetDescriptors, Logger, executor);
 					}
 
 					// these are parsed by external tools wishing to open this file directly
