Remove unused parameter in UBT left over from removing unused artifact cache system.

#Synced-CL 48546779

diff --git a/Engine/Source/Programs/UnrealBuildTool/Executors/ImmediateActionQueue.cs b/Engine/Source/Programs/UnrealBuildTool/Executors/ImmediateActionQueue.cs
--- a/Engine/Source/Programs/UnrealBuildTool/Executors/ImmediateActionQueue.cs
+++ b/Engine/Source/Programs/UnrealBuildTool/Executors/ImmediateActionQueue.cs
@@ -254,12 +254,11 @@ namespace UnrealBuildTool
 		/// Construct a new instance of the action queue
 		/// </summary>
 		/// <param name="actions">Collection of actions</param>
-		/// <param name="maxActionArtifactCacheTasks">Max number of concurrent artifact cache tasks</param>
 		/// <param name="progressWriterText">Text to be displayed with the progress writer</param>
 		/// <param name="writeToolOutput">Action to invoke when writing tool output</param>
 		/// <param name="flushToolOutput">Action to invoke when flushing tool output</param>
 		/// <param name="logger">Logging interface</param>
-		public ImmediateActionQueue(IEnumerable<LinkedAction> actions, int maxActionArtifactCacheTasks, string progressWriterText, Action<string> writeToolOutput, System.Action flushToolOutput, ILogger logger)
+		public ImmediateActionQueue(IEnumerable<LinkedAction> actions, string progressWriterText, Action<string> writeToolOutput, System.Action flushToolOutput, ILogger logger)
 		:	base(actions.Count(), progressWriterText, writeToolOutput, flushToolOutput, logger)
 		{
 			CommandLine.ParseArguments(Environment.GetCommandLineArgs(), this, logger);
diff --git a/Engine/Source/Programs/UnrealBuildTool/Executors/ParallelExecutor.cs b/Engine/Source/Programs/UnrealBuildTool/Executors/ParallelExecutor.cs
--- a/Engine/Source/Programs/UnrealBuildTool/Executors/ParallelExecutor.cs
+++ b/Engine/Source/Programs/UnrealBuildTool/Executors/ParallelExecutor.cs
@@ -173,12 +173,11 @@ namespace UnrealBuildTool
 		/// Create an action queue
 		/// </summary>
 		/// <param name="actionsToExecute">Actions to be executed</param>
-		/// <param name="maxActionArtifactCacheTasks">Max artifact tasks that can execute in parallel</param>
 		/// <param name="logger">Logging interface</param>
 		/// <returns>Action queue</returns>
-		public ImmediateActionQueue CreateActionQueue(IEnumerable<LinkedAction> actionsToExecute, int maxActionArtifactCacheTasks, ILogger logger)
+		public ImmediateActionQueue CreateActionQueue(IEnumerable<LinkedAction> actionsToExecute, ILogger logger)
 		{
-			return new(actionsToExecute, maxActionArtifactCacheTasks, "Compiling C++ source code...", x => WriteToolOutput(x), () => FlushToolOutput(), logger)
+			return new(actionsToExecute, "Compiling C++ source code...", x => WriteToolOutput(x), () => FlushToolOutput(), logger)
 			{
 				ShowCompilationTimes = bShowCompilationTimes,
 				ShowCPUUtilization = bShowCPUUtilization,
@@ -206,7 +205,7 @@ namespace UnrealBuildTool
 			bool useAutomaticQueue = true;
 			if (useAutomaticQueue)
 			{
-				using ImmediateActionQueue queue = CreateActionQueue(ActionsToExecute, NumParallelProcesses, Logger);
+				using ImmediateActionQueue queue = CreateActionQueue(ActionsToExecute, Logger);
 				int actionLimit = Math.Min(NumParallelProcesses, queue.TotalActions);
 				queue.CreateAutomaticRunner(action => RunAction(queue, action), bUseActionWeights, actionLimit, NumParallelProcesses);
 				queue.Start();
@@ -218,7 +217,7 @@ namespace UnrealBuildTool
 			}
 			else
 			{
-				using ImmediateActionQueue queue = CreateActionQueue(ActionsToExecute, NumParallelProcesses, Logger);
+				using ImmediateActionQueue queue = CreateActionQueue(ActionsToExecute, Logger);
 				int actionLimit = Math.Min(NumParallelProcesses, queue.TotalActions);
 				ImmediateActionQueueRunner runner = queue.CreateManualRunner(action => RunAction(queue, action), bUseActionWeights, actionLimit, actionLimit);
 				queue.Start();
