Make `UncookedOnly` plugin modules automatically disabled for AutoRTFM.

#Synced-CL 48547807

diff --git a/Engine/Source/Programs/UnrealBuildTool/Configuration/Rules/ModuleRules.cs b/Engine/Source/Programs/UnrealBuildTool/Configuration/Rules/ModuleRules.cs
--- a/Engine/Source/Programs/UnrealBuildTool/Configuration/Rules/ModuleRules.cs
+++ b/Engine/Source/Programs/UnrealBuildTool/Configuration/Rules/ModuleRules.cs
@@ -1159,8 +1159,8 @@ namespace UnrealBuildTool
 		/// The default value of this depends on whether any `TargetType` can be deduced for the module.
 		/// It will default to `true` if:
 		/// - The module has `SupportedTargetTypes` specified and `TargetType.Server` is not one of the target types.
-		/// - The module is part of a plugin that has a "Type" that is one of "Editor", "EditorAndProgram", or
-		///   "EditorNoCommandlet".
+		/// - The module is part of a plugin that has a "Type" that is one of "Editor", "EditorAndProgram",
+		///   "EditorNoCommandlet", or "UncookedOnly".
 		/// - The module is part of a plugin where the module has `"TargetAllowList"` specified and `"Server"` is not
 		///   one of the allowed target types.
 		/// - The module is part of a plugin where the module has `"TargetDenyList"` specified and `"Server"` is one
@@ -1189,7 +1189,8 @@ namespace UnrealBuildTool
 					case ModuleHostType.Editor:
 					case ModuleHostType.EditorAndProgram:
                     case ModuleHostType.EditorNoCommandlet:
-                        return true;
+					case ModuleHostType.UncookedOnly:
+						return true;
 				}
 
                 if (descriptor?.TargetAllowList != null && descriptor.TargetAllowList.Any() && !descriptor.TargetAllowList.Any(x => x == TargetType.Server))
