UnrealBuildTool: Fixes for VS2026 solution

#rnx

#Synced-CL 47605596

diff --git a/Engine/Source/Programs/UnrealBuildTool/Modes/GenerateProjectFilesMode.cs b/Engine/Source/Programs/UnrealBuildTool/Modes/GenerateProjectFilesMode.cs
--- a/Engine/Source/Programs/UnrealBuildTool/Modes/GenerateProjectFilesMode.cs
+++ b/Engine/Source/Programs/UnrealBuildTool/Modes/GenerateProjectFilesMode.cs
@@ -25,6 +25,7 @@ namespace UnrealBuildTool
 		[CommandLine("-ProjectFileFormat")]
 		[CommandLine("-VisualStudio", Value = nameof(ProjectFileFormat.VisualStudio))]
 		[CommandLine("-2022", Value = nameof(ProjectFileFormat.VisualStudio2022))] // + override compiler
+		[CommandLine("-2026", Value = nameof(ProjectFileFormat.VisualStudio2026))] // + override compiler
 		[CommandLine("-Makefile", Value = nameof(ProjectFileFormat.Make))]
 		[CommandLine("-CMakefile", Value = nameof(ProjectFileFormat.CMake))]
 		[CommandLine("-QMakefile", Value = nameof(ProjectFileFormat.QMake))]
diff --git a/Engine/Source/Programs/UnrealBuildTool/ProjectFiles/VisualStudio/VCProjectFileGenerator.cs b/Engine/Source/Programs/UnrealBuildTool/ProjectFiles/VisualStudio/VCProjectFileGenerator.cs
--- a/Engine/Source/Programs/UnrealBuildTool/ProjectFiles/VisualStudio/VCProjectFileGenerator.cs
+++ b/Engine/Source/Programs/UnrealBuildTool/ProjectFiles/VisualStudio/VCProjectFileGenerator.cs
@@ -148,6 +148,11 @@ namespace UnrealBuildTool
 				BuildToolOverride = "-2022";
 			}
 
+			if (InArguments.HasOption("-2026"))
+			{
+				BuildToolOverride = "-2026";
+			}
+
 			// Allow generating the solution even if the only installed toolchain is banned.			
 			MicrosoftPlatformSDK.IgnoreToolchainErrors = true;
 		}
@@ -730,7 +735,7 @@ namespace UnrealBuildTool
 				VCSolutionFileContent.AppendLine();
 				VCSolutionFileContent.AppendLine("Microsoft Visual Studio Solution File, Format Version 12.00");
 				VCSolutionFileContent.AppendLine("# Visual Studio Version 18");
-				VCSolutionFileContent.AppendLine("VisualStudioVersion = 18.0.11005.162");
+				VCSolutionFileContent.AppendLine("VisualStudioVersion = 18.0.11205.157 d18.0");
 				VCSolutionFileContent.AppendLine("MinimumVisualStudioVersion = 10.0.40219.1");
 			}
 			else
@@ -983,8 +988,27 @@ namespace UnrealBuildTool
 			// Save the solution file
 			if (bSuccess)
 			{
-				string SolutionFilePath = FileReference.Combine(PrimaryProjectPath, SolutionFileName).FullName;
-				bSuccess = WriteFileIfChanged(SolutionFilePath, VCSolutionFileContent.ToString(), Logger);
+				FileReference SolutionFilePath = FileReference.Combine(PrimaryProjectPath, SolutionFileName);
+				bSuccess = WriteFileIfChanged(SolutionFilePath.FullName, VCSolutionFileContent.ToString(), Logger);
+
+				// Generate a slnx for VS2026 until properly supported
+				if (bSuccess && Settings.ProjectFileFormat == VCProjectFileFormat.VisualStudio2026 && Environment.Version.Major >= 10)
+				{
+					FileReference TempSolutionFilePath = FileReference.Combine(new DirectoryReference(Path.GetTempPath()), SolutionFileName);
+					FileReference.Copy(SolutionFilePath, TempSolutionFilePath, true);
+					using ManagedProcessGroup group = new();
+					using ManagedProcess process = new(group, Unreal.DotnetPath.FullName, $"sln migrate \"{TempSolutionFilePath}\"", TempSolutionFilePath.Directory.FullName, null, System.Diagnostics.ProcessPriorityClass.Normal);
+					process.ReadAllLines();
+					process.WaitForExit();
+					bSuccess = process.ExitCode == 0;
+					if (bSuccess)
+					{
+						string contents = FileReference.ReadAllText(TempSolutionFilePath.ChangeExtension(".slnx"));
+						bSuccess = WriteFileIfChanged(SolutionFilePath.ChangeExtension(".slnx").FullName, contents, Logger);
+						FileReference.Delete(TempSolutionFilePath);
+						FileReference.Delete(TempSolutionFilePath.ChangeExtension(".slnx"));
+					}
+				}
 			}
 
 			// Save a solution config file which selects the development editor configuration by default.
