UnrealBuildTool: Force an aggressive garbage collection before running actions to free up as much memory as possible for compilation

#rnx
#Synced-CL 47605596

diff --git a/Engine/Source/Programs/UnrealBuildTool/Actions/ActionGraph.cs b/Engine/Source/Programs/UnrealBuildTool/Actions/ActionGraph.cs
--- a/Engine/Source/Programs/UnrealBuildTool/Actions/ActionGraph.cs
+++ b/Engine/Source/Programs/UnrealBuildTool/Actions/ActionGraph.cs
@@ -8,6 +8,7 @@ using System.Diagnostics;
 using System.Globalization;
 using System.IO;
 using System.Linq;
+using System.Runtime;
 using System.Text;
 using System.Text.Json;
 using System.Text.Json.Serialization;
@@ -506,6 +507,13 @@ namespace UnrealBuildTool
 		{
 			Logger.LogInformation("Using {ExecutorName} executor to run {ActionCount} action(s)", Executor.Name, ActionsToExecute.Count);
 
+			// Force an aggressive garbage collection in the background before running actions to free up as much memory as possible as compiling and linking can use a significant amount of memory
+			Task gcTask = Task.Run(() => {
+				using ITimelineEvent _ = Timeline.ScopeEvent("GarbageCollect");
+				GCSettings.LargeObjectHeapCompactionMode = GCLargeObjectHeapCompactionMode.CompactOnce;
+				GC.Collect(GC.MaxGeneration, GCCollectionMode.Aggressive, true, true);
+			});
+
 			// Execute the build
 			Stopwatch Timer = Stopwatch.StartNew();
 			bool Result = await Executor.ExecuteActionsAsync(ActionsToExecute, Logger);
