Only show linux toolchain info once per process

#rnx
#Synced-CL 48608416

diff --git a/Engine/Source/Programs/UnrealBuildTool/Platform/Linux/LinuxToolChain.cs b/Engine/Source/Programs/UnrealBuildTool/Platform/Linux/LinuxToolChain.cs
--- a/Engine/Source/Programs/UnrealBuildTool/Platform/Linux/LinuxToolChain.cs
+++ b/Engine/Source/Programs/UnrealBuildTool/Platform/Linux/LinuxToolChain.cs
@@ -5,6 +5,7 @@ using System.Collections.Generic;
 using System.Diagnostics;
 using System.IO;
 using System.Linq;
+using System.Threading;
 using EpicGames.Core;
 using Microsoft.Extensions.Logging;
 using UnrealBuildBase;
@@ -782,7 +783,7 @@ namespace UnrealBuildTool
 		/// <summary>
 		/// Tracks that information about used C++ library is only printed once
 		/// </summary>
-		private bool bHasPrintedBuildDetails = false;
+		static int s_hasPrintedBuildDetails = 0; // per-process static so this is only printed once per-process
 		protected void PrintBuildDetails(CppCompileEnvironment CompileEnvironment, ILogger Logger)
 		{
 			Logger.LogInformation("------- Build details --------");
@@ -874,11 +875,9 @@ namespace UnrealBuildTool
 
 			//UEBuildPlatform BuildPlatform = UEBuildPlatform.GetBuildPlatform(CompileEnvironment.Platform);
 
-			if (!bHasPrintedBuildDetails)
+			if (Interlocked.Exchange(ref s_hasPrintedBuildDetails, 1) == 0)
 			{
 				PrintBuildDetails(CompileEnvironment, Logger);
-
-				bHasPrintedBuildDetails = true;
 			}
 
 			// Create a compile action for each source file.
