[UBA]
* Added config Storage.SkipCleanupOnShutdown to UBAExecutor when instantiated through BuildMode. With this config we let the process release some file mappings instead of doing it from code which speeds up shutdown a bit when large file mappings have been used

#Synced-CL 48618938

diff --git a/Engine/Source/Programs/UnrealBuildTool/Executors/UnrealBuildAccelerator/UBAExecutor.cs b/Engine/Source/Programs/UnrealBuildTool/Executors/UnrealBuildAccelerator/UBAExecutor.cs
--- a/Engine/Source/Programs/UnrealBuildTool/Executors/UnrealBuildAccelerator/UBAExecutor.cs
+++ b/Engine/Source/Programs/UnrealBuildTool/Executors/UnrealBuildAccelerator/UBAExecutor.cs
@@ -139,12 +139,18 @@ namespace UnrealBuildTool
 			_cancellationSource = new();
 			if (useInitTask)
 			{
-				_initTask = new(() => { Init(targetDescriptors, logger); }, TaskCreationOptions.LongRunning);
+				_initTask = new(() =>
+				{
+					// We know the BuildMode is the only one using init tasks which means we can let process exit do cleanup
+					// instead of explicitly doing it with the executor shutdown
+					Init(targetDescriptors, logger, false);
+				}, TaskCreationOptions.LongRunning);
+
 				_initTask.Start();
 			}
 		}
 
-		private void Init(IEnumerable<TargetDescriptor> targetDescriptors, ILogger logger)
+		private void Init(IEnumerable<TargetDescriptor> targetDescriptors, ILogger logger, bool cleanupOnShutdown)
 		{
 			XmlConfig.ApplyTo(this);
 			XmlConfig.ApplyTo(UBAConfig);
@@ -268,6 +274,11 @@ namespace UnrealBuildTool
 				config.AddValue("Storage", "CompressionLevel", UBAConfig.CompressionLevel);
 			}
 
+			if (!cleanupOnShutdown)
+			{
+				config.AddValue("Storage", "SkipCleanupOnShutdown", "true"); // For faster shutdown
+			}
+
 			EpicGames.UBA.ILogger ubaLogger = EpicGames.UBA.ILogger.GlobalLogger!;
 			_server = IServer.CreateServer(UBAConfig.MaxWorkers, UBAConfig.SendSize, ubaLogger, UBAConfig.bUseQuic);
 			_storageServer = IStorageServer.CreateStorageServer(_server, ubaLogger, new StorageServerCreateInfo(_rootDirRef.FullName, ((ulong)UBAConfig.StoreCapacityGb) * 1000 * 1000 * 1000, !UBAConfig.bStoreRaw, UBAConfig.Zone));
@@ -500,7 +511,7 @@ namespace UnrealBuildTool
 
 			if (_initTask == null)
 			{
-				Init(_targetDescriptors, logger);
+				Init(_targetDescriptors, logger, true);
 			}
 			else
 			{
