Small fix to Android Studio project generator if no Client Target

#Synced-CL 48491776

diff --git a/Engine/Source/Programs/UnrealBuildTool/ProjectFiles/AndroidStudio/AndroidStudioProjectFileGenerator.cs b/Engine/Source/Programs/UnrealBuildTool/ProjectFiles/AndroidStudio/AndroidStudioProjectFileGenerator.cs
--- a/Engine/Source/Programs/UnrealBuildTool/ProjectFiles/AndroidStudio/AndroidStudioProjectFileGenerator.cs
+++ b/Engine/Source/Programs/UnrealBuildTool/ProjectFiles/AndroidStudio/AndroidStudioProjectFileGenerator.cs
@@ -57,7 +57,7 @@ namespace UnrealBuildTool
 
 		protected override bool WriteProjectFiles(PlatformProjectGeneratorCollection PlatformProjectGenerators, ILogger Logger)
 		{
-			TargetRules ClientTargetRules = DefaultProject!.ProjectTargets.First(Project => Project.TargetRules!.Type == TargetType.Client).TargetRules!;
+			TargetRules ClientTargetRules = (DefaultProject!.ProjectTargets.FirstOrDefault(Project => Project.TargetRules!.Type == TargetType.Client) ?? DefaultProject!.ProjectTargets.First(Project => Project.TargetRules!.Type == TargetType.Game)).TargetRules!;
 			UnrealArchitectures Architectures = VCProjectFileGenerator.GetAndroidProjectArchitectures(OnlyGameProject);
 			UnrealArch Architecture = Architectures.Architectures.First();
 
