SA fixes SubmitToolCore

#rnx
#Synced-CL 45819516

diff --git a/Engine/Source/Developer/SubmitToolCore/Private/CommandLine/CmdLineParameters.cpp b/Engine/Source/Developer/SubmitToolCore/Private/CommandLine/CmdLineParameters.cpp
--- a/Engine/Source/Developer/SubmitToolCore/Private/CommandLine/CmdLineParameters.cpp
+++ b/Engine/Source/Developer/SubmitToolCore/Private/CommandLine/CmdLineParameters.cpp
@@ -259,12 +259,12 @@ bool FCmdLineParameters::GetValue(const FString& InKey, FString& OutValue) const
 	{
 		if(Tokens[i].Equals(InKey, ESearchCase::IgnoreCase) && i + 1 < Tokens.Num())
 		{
-			OutValue = (*Definition)->bIsFlag ? TEXT("true") : Tokens[i + 1];
-			if(Definition)
+			if (Definition != nullptr)
 			{
+				OutValue = (*Definition)->bIsFlag ? TEXT("true") : Tokens[i + 1];
 				(*Definition)->CustomParse(OutValue);
+				return true;
 			}
-			return true;
 		}
 	}
 
diff --git a/Engine/Source/Programs/UnrealBuildTool/Configuration/Rules/TargetRules.cs b/Engine/Source/Programs/UnrealBuildTool/Configuration/Rules/TargetRules.cs
--- a/Engine/Source/Programs/UnrealBuildTool/Configuration/Rules/TargetRules.cs
+++ b/Engine/Source/Programs/UnrealBuildTool/Configuration/Rules/TargetRules.cs
@@ -629,11 +629,13 @@ namespace UnrealBuildTool
 		/// </summary>
 		protected bool bIsTestTargetOverride { get; set; }
 
+#pragma warning disable CA2227
 		/// <summary>
 		/// The selection of test modules to be exclusively included in compilation (no other modules tests will be included). 
 		/// </summary>
-		[RequiresUniqueBuildEnvironment] 
+		[RequiresUniqueBuildEnvironment]
 		public List<string> SelectedTestModuleNames { get; protected set; } = [];
+#pragma warning restore CA2227
 
 		/// <summary>
 		/// Whether this is a test target explicitly defined.
diff --git a/Engine/Source/Programs/UnrealBuildTool/Configuration/Rules/TestTargetRules.cs b/Engine/Source/Programs/UnrealBuildTool/Configuration/Rules/TestTargetRules.cs
--- a/Engine/Source/Programs/UnrealBuildTool/Configuration/Rules/TestTargetRules.cs
+++ b/Engine/Source/Programs/UnrealBuildTool/Configuration/Rules/TestTargetRules.cs
@@ -376,10 +376,10 @@ namespace UnrealBuildTool
 
 			if (target.Arguments != null)
 			{
-				string? SelectedModules = null;
-				if (target.Arguments.TryGetValue("-SelectModules=", out SelectedModules))
+				string? selectedModules;
+				if (target.Arguments.TryGetValue("-SelectModules=", out selectedModules))
 				{
-					SelectedTestModuleNames = [.. SelectedModules.Split(',')];
+					SelectedTestModuleNames = [.. selectedModules.Split(',')];
 				}
 			}
 		}
