Replace the ParallelExecutor with UBAExecutor configured to disable detouring, as this should be equivalent in functionality, and allows us to reduce our code footprint.

#Synced-CL 48648194

diff --git a/Engine/Source/Programs/UnrealBuildTool/Actions/ActionGraph.cs b/Engine/Source/Programs/UnrealBuildTool/Actions/ActionGraph.cs
--- a/Engine/Source/Programs/UnrealBuildTool/Actions/ActionGraph.cs
+++ b/Engine/Source/Programs/UnrealBuildTool/Actions/ActionGraph.cs
@@ -388,33 +388,70 @@ namespace UnrealBuildTool
 			}
 		}
 
-		public static ActionExecutor? PrecreateExecutor(BuildConfiguration BuildConfiguration, ILogger Logger, List<TargetDescriptor> TargetDescriptors)
+		internal static UBAExecutor GetUBAExecutor(BuildConfiguration BuildConfiguration, ILogger Logger, List<TargetDescriptor> TargetDescriptors)
 		{
-			if (BuildConfiguration.bAllowXGE && XGE.IsAvailable(Logger))
-			{
-				return null;
-			}
-			if (BuildConfiguration.bAllowSNDBS && SNDBS.IsAvailable(Logger))
+			if (!UBAExecutor.IsAvailable())
 			{
-				return null;
+				throw new BuildLogEventException("UBA is not available - please ensure the UBA binaries exist for your host platform");
 			}
-			if (BuildConfiguration.bAllowFASTBuild && FASTBuild.IsAvailable(Logger))
+
+			UnrealBuildAcceleratorConfig Config = UBAExecutor.GetAppliedConfig(TargetDescriptors);
+			// We always use the UBA executor, but we disable detouring to mirror legacy behaviour if the config disables it.
+			if (!BuildConfiguration.bAllowUBAExecutor)
 			{
-				return null;
+				Config.bAllowDetour = false;
 			}
-			if (BuildConfiguration.bAllowUBAExecutor && UBAExecutor.IsAvailable())
+
+			return new UBAExecutor(
+				Config,
+				BuildConfiguration.MaxParallelActions,
+				BuildConfiguration.bAllCores,
+				BuildConfiguration.bCompactOutput,
+				Logger,
+				TargetDescriptors);
+		}
+
+		public static ActionExecutor? PrecreateExecutor(BuildConfiguration BuildConfiguration, ILogger Logger, List<TargetDescriptor> TargetDescriptors)
+		{
+			foreach (string Name in BuildConfiguration.RemoteExecutorPriority)
 			{
-				return new UBAExecutor(BuildConfiguration.MaxParallelActions, BuildConfiguration.bAllCores, BuildConfiguration.bCompactOutput, Logger, TargetDescriptors, true);
+				switch (Name)
+				{
+					case XGE.ExecutorName:
+						if (BuildConfiguration.bAllowXGE && XGE.IsAvailable(Logger))
+						{
+							return null;
+						}
+						break;
+					case SNDBS.ExecutorName:
+						if (BuildConfiguration.bAllowSNDBS && SNDBS.IsAvailable(Logger))
+						{
+							return null;
+						}
+						break;
+					case FASTBuild.ExecutorName:
+						if (BuildConfiguration.bAllowFASTBuild && FASTBuild.IsAvailable(Logger))
+						{
+							return null;
+						}
+						break;
+					case UBAExecutor.ExecutorName:
+						if (BuildConfiguration.bAllowUBAExecutor)
+						{
+							return GetUBAExecutor(BuildConfiguration, Logger, TargetDescriptors);
+						}
+						break;
+				}
 			}
 
-			return null;
+			return GetUBAExecutor(BuildConfiguration, Logger, TargetDescriptors);
 		}
 
 		private static ActionExecutor? GetRemoteExecutorByName(string Name, BuildConfiguration BuildConfiguration, int ActionCount, int MinActionsForRemote, List<TargetDescriptor> TargetDescriptors, ILogger Logger)
 		{
 			switch (Name)
 			{
-				case "XGE":
+				case XGE.ExecutorName:
 					{
 						if (ActionCount >= MinActionsForRemote && BuildConfiguration.bAllowXGE && XGE.IsAvailable(Logger) && ActionCount >= XGE.MinActions)
 						{
@@ -422,7 +459,7 @@ namespace UnrealBuildTool
 						}
 						return null;
 					}
-				case "SNDBS":
+				case SNDBS.ExecutorName:
 					{
 						if (ActionCount >= MinActionsForRemote && BuildConfiguration.bAllowSNDBS && SNDBS.IsAvailable(Logger))
 						{
@@ -430,20 +467,20 @@ namespace UnrealBuildTool
 						}
 						return null;
 					}
-				case "FASTBuild":
+				case FASTBuild.ExecutorName:
 					{
 						if (ActionCount >= MinActionsForRemote && BuildConfiguration.bAllowFASTBuild && FASTBuild.IsAvailable(Logger))
 						{
-							return new FASTBuild(BuildConfiguration.MaxParallelActions, BuildConfiguration.bAllCores, BuildConfiguration.bCompactOutput, Logger);
+							return new FASTBuild(BuildConfiguration, Logger, TargetDescriptors);
 						}
 						return null;
 					}
-				case "UBA":
+				case UBAExecutor.ExecutorName:
 					{
 						// Intentionally not checking MinActionsForRemote
-						if (BuildConfiguration.bAllowUBAExecutor && UBAExecutor.IsAvailable())
+						if (BuildConfiguration.bAllowUBAExecutor)
 						{
-							return new UBAExecutor(BuildConfiguration.MaxParallelActions, BuildConfiguration.bAllCores, BuildConfiguration.bCompactOutput, Logger, TargetDescriptors, true);
+							return GetUBAExecutor(BuildConfiguration, Logger, TargetDescriptors);
 						}
 						return null;
 					}
@@ -460,7 +497,8 @@ namespace UnrealBuildTool
 		/// </summary>
 		private static ActionExecutor SelectExecutor(BuildConfiguration BuildConfiguration, int ActionCount, List<TargetDescriptor> TargetDescriptors, ILogger Logger)
 		{
-			int MinActionsForRemote = ParallelExecutor.GetDefaultNumParallelProcesses(BuildConfiguration.MaxParallelActions, BuildConfiguration.bAllCores, Logger);
+			UnrealBuildAcceleratorConfig Config = UBAExecutor.GetAppliedConfig(TargetDescriptors);
+			int MinActionsForRemote = UBAExecutor.GetDefaultNumParallelProcesses(Config, BuildConfiguration.MaxParallelActions, BuildConfiguration.bAllCores, Logger);
 			foreach (string Name in BuildConfiguration.RemoteExecutorPriority)
 			{
 				ActionExecutor? Executor = GetRemoteExecutorByName(Name, BuildConfiguration, ActionCount, MinActionsForRemote, TargetDescriptors, Logger);
@@ -470,7 +508,7 @@ namespace UnrealBuildTool
 				}
 			}
 
-			return new ParallelExecutor(BuildConfiguration.MaxParallelActions, BuildConfiguration.bAllCores, BuildConfiguration.bCompactOutput, Logger);
+			return GetUBAExecutor(BuildConfiguration, Logger, TargetDescriptors);
 		}
 
 		/// <summary>
diff --git a/Engine/Source/Programs/UnrealBuildTool/Configuration/BuildConfiguration.cs b/Engine/Source/Programs/UnrealBuildTool/Configuration/BuildConfiguration.cs
--- a/Engine/Source/Programs/UnrealBuildTool/Configuration/BuildConfiguration.cs
+++ b/Engine/Source/Programs/UnrealBuildTool/Configuration/BuildConfiguration.cs
@@ -44,7 +44,7 @@ namespace UnrealBuildTool
 		/// Priority order for remote executors (XGE, SNDBS, FASTBuild, UBA)
 		/// </summary>
 		[XmlConfigFile]
-		public string[] RemoteExecutorPriority { get; set; } = ["XGE", "SNDBS", "FASTBuild", "UBA"];
+		public string[] RemoteExecutorPriority { get; set; } = [XGE.ExecutorName, SNDBS.ExecutorName, FASTBuild.ExecutorName, UBAExecutor.ExecutorName];
 
 		/// <summary>
 		/// Whether the UnrealBuildAccelerator executor will be used.
diff --git a/Engine/Source/Programs/UnrealBuildTool/Executors/Experimental/FASTBuild.cs b/Engine/Source/Programs/UnrealBuildTool/Executors/Experimental/FASTBuild.cs
--- a/Engine/Source/Programs/UnrealBuildTool/Executors/Experimental/FASTBuild.cs
+++ b/Engine/Source/Programs/UnrealBuildTool/Executors/Experimental/FASTBuild.cs
@@ -72,7 +72,7 @@ namespace UnrealBuildTool
 		/// <summary>
 		/// Executor to use for local actions
 		/// </summary>
-		readonly ActionExecutor LocalExecutor;
+		readonly UBAExecutor LocalExecutor;
 
 		public static readonly string DefaultExecutableBasePath = Path.Combine(Unreal.EngineDirectory.FullName, "Extras", "ThirdPartyNotUE", "FASTBuild");
 
@@ -161,15 +161,33 @@ namespace UnrealBuildTool
 		/// <summary>
 		/// Constructor
 		/// </summary>
-		public FASTBuild(int MaxLocalActions, bool bAllCores, bool bCompactOutput, ILogger Logger)
+		public FASTBuild(BuildConfiguration BuildConfiguration, ILogger Logger, IEnumerable<TargetDescriptor> TargetDescriptors)
 			: base(Logger)
 		{
 			XmlConfig.ApplyTo(this);
 
-			LocalExecutor = new ParallelExecutor(MaxLocalActions, bAllCores, bCompactOutput, Logger);
+			UnrealBuildAcceleratorConfig Config = UBAExecutor.GetAppliedConfig(TargetDescriptors);
+
+			// Disable remote, as this is intended as a local executor.
+			Config.bDisableRemote = true;
+
+			// We always use the UBA executor, but we disable detouring to mirror legacy behaviour if the config disables it.
+			if (!BuildConfiguration.bAllowUBAExecutor)
+			{
+				Config.bAllowDetour = false;
+			}
+			LocalExecutor = new UBAExecutor(
+				Config,
+				BuildConfiguration.MaxParallelActions,
+				BuildConfiguration.bAllCores,
+				BuildConfiguration.bCompactOutput,
+				Logger,
+				TargetDescriptors);
 		}
 
-		public override string Name => "FASTBuild";
+		internal const string ExecutorName = "FASTBuild";
+
+		public override string Name => ExecutorName;
 
 		public static string GetExecutableName()
 		{
@@ -1355,7 +1373,7 @@ namespace UnrealBuildTool
 			string ForceRemoteArgument = bForceRemote ? "-forceremote" : "";
 			string NoStopOnErrorArgument = bStopOnError ? "" : "-nostoponerror";
 			string IDEArgument = IsApple() ? "" : "-ide";
-			string MaxProcesses = "-j" + ((ParallelExecutor)LocalExecutor).NumParallelProcesses;
+			string MaxProcesses = "-j" + LocalExecutor.NumParallelProcesses;
 
 			// Interesting flags for FASTBuild:
 			// -nostoponerror, -verbose, -monitor (if FASTBuild Monitor Visual Studio Extension is installed!)
diff --git a/Engine/Source/Programs/UnrealBuildTool/Executors/ImmediateActionQueue.cs b/Engine/Source/Programs/UnrealBuildTool/Executors/ImmediateActionQueue.cs
deleted file mode 100644
diff --git a/Engine/Source/Programs/UnrealBuildTool/Executors/ParallelExecutor.cs b/Engine/Source/Programs/UnrealBuildTool/Executors/ParallelExecutor.cs
deleted file mode 100644
diff --git a/Engine/Source/Programs/UnrealBuildTool/Executors/SNDBS.cs b/Engine/Source/Programs/UnrealBuildTool/Executors/SNDBS.cs
--- a/Engine/Source/Programs/UnrealBuildTool/Executors/SNDBS.cs
+++ b/Engine/Source/Programs/UnrealBuildTool/Executors/SNDBS.cs
@@ -22,7 +22,9 @@ namespace UnrealBuildTool
 {
 	sealed class SNDBS : ActionExecutor
 	{
-		public override string Name => "SNDBS";
+		internal const string ExecutorName = "SNDBS";
+
+		public override string Name => ExecutorName;
 
 		private static readonly string? ProgramFilesx86 = Environment.GetEnvironmentVariable("ProgramFiles(x86)");
 		private static readonly string? SCERoot = Environment.GetEnvironmentVariable("SCE_ROOT_DIR");
diff --git a/Engine/Source/Programs/UnrealBuildTool/Executors/UnrealBuildAccelerator/UBAExecutor.cs b/Engine/Source/Programs/UnrealBuildTool/Executors/UnrealBuildAccelerator/UBAExecutor.cs
--- a/Engine/Source/Programs/UnrealBuildTool/Executors/UnrealBuildAccelerator/UBAExecutor.cs
+++ b/Engine/Source/Programs/UnrealBuildTool/Executors/UnrealBuildAccelerator/UBAExecutor.cs
@@ -17,17 +17,16 @@ using ILogger = Microsoft.Extensions.Logging.ILogger;
 
 namespace UnrealBuildTool
 {
-	class UBAExecutor : ParallelExecutor
+	class UBAExecutor : ActionExecutor
 	{
-		public UnrealBuildAcceleratorConfig UBAConfig { get; init; } = new();
+		public UnrealBuildAcceleratorConfig UBAConfig { get; init; }
 
 		public string Crypto { get; private set; } = String.Empty;
 
-		public bool CanRetryNatively { get; set; } = true;
-
 		readonly IEnumerable<TargetDescriptor> _targetDescriptors;
-		readonly Task? _initTask;
+		readonly Task _initTask;
 		readonly CancellationTokenSource _cancellationSource;
+		readonly bool _bCompactOutput;
 
 		IServer? _server;
 		IStorageServer? _storageServer;
@@ -57,11 +56,16 @@ namespace UnrealBuildTool
 		// Tracking for remote connection mode
 		string _remoteConnectionMode = "Local";
 
+		/// <summary>
+		/// Telemetry event for this executor
+		/// </summary>
+		private TelemetryExecutorEvent? telemetryEvent;
+
 		protected override void Dispose(bool disposing)
 		{
 			if (disposing)
 			{
-				_initTask?.Wait();
+				_initTask.Wait();
 
 				{
 					using ITimelineEvent a2 = Timeline.ScopeEvent("DisconnectCache");
@@ -94,9 +98,13 @@ namespace UnrealBuildTool
 			}
 		}
 
+		internal const string ExecutorName = "UBA";
+
 		public override string Name => !UBAConfig.bDisableRemote && _agentCoordinators.Count > 0 ? "Unreal Build Accelerator" : "Unreal Build Accelerator local";
 
-		public static new bool IsAvailable() => EpicGames.UBA.Utils.IsAvailable();
+		public int NumParallelProcesses { get; }
+
+		public static bool IsAvailable() => EpicGames.UBA.Utils.IsAvailable();
 
 		public delegate void SessionCreatedDelegate(ISessionServer Session);
 
@@ -132,45 +140,44 @@ namespace UnrealBuildTool
 			}
 		}
 
-		public UBAExecutor(int maxLocalActions, bool bAllCores, bool bCompactOutput, ILogger logger, IEnumerable<TargetDescriptor> targetDescriptors, bool useInitTask)
-			: base(maxLocalActions, bAllCores, bCompactOutput, logger)
+		internal static UnrealBuildAcceleratorConfig GetAppliedConfig(IEnumerable<TargetDescriptor> targetDescriptors)
 		{
-			_targetDescriptors = targetDescriptors;
-			_cancellationSource = new();
-			if (useInitTask)
-			{
-				_initTask = new(() =>
-				{
-					// We know the BuildMode is the only one using init tasks which means we can let process exit do cleanup
-					// instead of explicitly doing it with the executor shutdown
-					Init(targetDescriptors, logger, false);
-				}, TaskCreationOptions.LongRunning);
+			UnrealBuildAcceleratorConfig config = new();
+			XmlConfig.ApplyTo(config);
 
-				_initTask.Start();
+			foreach (TargetDescriptor targetDescriptor in targetDescriptors)
+			{
+				ConfigCache.ReadSettings(targetDescriptor.ProjectFile?.Directory, BuildHostPlatform.Current.Platform, config);
+				targetDescriptor.AdditionalArguments.ApplyTo(config);
 			}
+
+			return config;
+		}
+
+		public UBAExecutor(UnrealBuildAcceleratorConfig config, int maxLocalActions, bool bAllCores, bool bCompactOutput, ILogger logger, IEnumerable<TargetDescriptor> targetDescriptors)
+			: base(logger)
+		{
+			UBAConfig = config;
+
+			// Figure out how many processors to use
+			NumParallelProcesses = GetDefaultNumParallelProcesses(config, maxLocalActions, bAllCores, logger);
+
+			_bCompactOutput = bCompactOutput;
+
+			_targetDescriptors = targetDescriptors;
+			_cancellationSource = new();
+			_initTask = new(() => { Init(targetDescriptors, logger, true); }, TaskCreationOptions.LongRunning);
+			_initTask.Start();
 		}
 
 		private void Init(IEnumerable<TargetDescriptor> targetDescriptors, ILogger logger, bool cleanupOnShutdown)
 		{
 			XmlConfig.ApplyTo(this);
-			XmlConfig.ApplyTo(UBAConfig);
 			CommandLine.ParseArguments(Environment.GetCommandLineArgs(), this, logger);
 
 			foreach (TargetDescriptor targetDescriptor in targetDescriptors)
 			{
-				ConfigCache.ReadSettings(targetDescriptor.ProjectFile?.Directory, BuildHostPlatform.Current.Platform, UBAConfig);
 				targetDescriptor.AdditionalArguments.ApplyTo(this);
-				targetDescriptor.AdditionalArguments.ApplyTo(UBAConfig);
-			}
-
-			if (UBAConfig.bDisableRemote)
-			{
-				UBAConfig.Zone = "local";
-			}
-
-			if (UBAConfig.bForcedRetry)
-			{
-				UBAConfig.bAllowRetry = true;
 			}
 
 			_trace = ITrace.GlobalTrace;
@@ -380,6 +387,17 @@ namespace UnrealBuildTool
 			}
 		}
 
+		internal static int GetDefaultNumParallelProcesses(UnrealBuildAcceleratorConfig config, int maxLocalActions, bool bAllCores, ILogger logger)
+		{
+			double MemoryPerActionBytesComputed = Math.Max(config.MemoryPerActionBytes, MemoryPerActionBytesOverride);
+			if (MemoryPerActionBytesComputed > config.MemoryPerActionBytes)
+			{
+				logger.LogInformation("Overriding MemoryPerAction with target-defined value of {Memory} bytes", MemoryPerActionBytesComputed / 1024 / 1024 / 1024);
+			}
+
+			return Utils.GetMaxActionsToExecuteInParallel(maxLocalActions, bAllCores ? 1.0f : config.ProcessorCountMultiplier, bAllCores, Convert.ToInt64(MemoryPerActionBytesComputed));
+		}
+
 		public void UpdateStatus(uint statusRow, uint statusColumn, string statusText, LogEntryType statusType, string? statusLink)
 		{
 			_trace?.UpdateStatus(statusRow, statusColumn, statusText, statusType, statusLink);
@@ -495,6 +513,9 @@ namespace UnrealBuildTool
 			}
 		}
 
+		/// <inheritdoc/>
+		public override TelemetryExecutorEvent? GetTelemetryEvent() => telemetryEvent;
+
 		public override async Task<bool> ExecuteActionsAsync(IEnumerable<LinkedAction> inputActions, Microsoft.Extensions.Logging.ILogger logger)
 		{
 			// Disable cache server when using hot reload path (there would likely never be any cache hits anyway)
@@ -509,14 +530,7 @@ namespace UnrealBuildTool
 				_cancellationSource.Cancel();
 			}
 
-			if (_initTask == null)
-			{
-				Init(_targetDescriptors, logger, true);
-			}
-			else
-			{
-				_initTask?.Wait();
-			}
+			_initTask.Wait();
 
 			PrintConfiguration(logger);
 
@@ -729,12 +743,12 @@ namespace UnrealBuildTool
 
 			using ActionLogger actionLogger = new(inputActions.Count(), "Compiling C++ source code...", x => WriteToolOutput(x), () => FlushToolOutput(), logger)
 			{
-				ShowCompilationTimes = bShowCompilationTimes,
-				ShowCPUUtilization = bShowCPUUtilization,
-				PrintActionTargetNames = bPrintActionTargetNames,
-				LogActionCommandLines = bLogActionCommandLines,
-				ShowPerActionCompilationTimes = bShowPerActionCompilationTimes,
-				CompactOutput = bCompactOutput,
+				ShowCompilationTimes = UBAConfig.bShowCompilationTimes,
+				ShowCPUUtilization = UBAConfig.bShowCPUUtilization,
+				PrintActionTargetNames = UBAConfig.bPrintActionTargetNames,
+				LogActionCommandLines = UBAConfig.bLogActionCommandLines,
+				ShowPerActionCompilationTimes = UBAConfig.bShowPerActionCompilationTimes,
+				CompactOutput = _bCompactOutput,
 			};
 
 			using IScheduler scheduler = IScheduler.CreateScheduler(session, cache, NumParallelProcesses, UBAConfig.bForceBuildAllRemote);
@@ -846,7 +860,7 @@ namespace UnrealBuildTool
 							{
 								// If files are compressed (pch, obj etc), then native will not work so must retry detoured
 								// If vfs paths are used, then native will not work
-								bool mustRunDetoured = !CanRetryNatively || UBAConfig.bStoreObjFilesCompressed || action.RootPaths.bUseVfs;
+								bool mustRunDetoured = !UBAConfig.CanRetryNatively || UBAConfig.bStoreObjFilesCompressed || action.RootPaths.bUseVfs;
 
 								if (mustRunDetoured)
 								{
@@ -931,14 +945,14 @@ namespace UnrealBuildTool
 						}
 					}
 
-					if (bStopCompilationAfterErrors || StopCompilationAfterNumErrors <= 0)
+					if (UBAConfig.bStopCompilationAfterErrors || UBAConfig.StopCompilationAfterNumErrors <= 0)
 					{
 						cancelReason = "stop compilation after any errors";
 						shouldCancel = true;
 					}
-					else if (failedActions >= StopCompilationAfterNumErrors)
+					else if (failedActions >= UBAConfig.StopCompilationAfterNumErrors)
 					{
-						cancelReason = $"stop compilation after {StopCompilationAfterNumErrors} errors";
+						cancelReason = $"stop compilation after {UBAConfig.StopCompilationAfterNumErrors} errors";
 						shouldCancel = true;
 					}
 				}
@@ -1113,15 +1127,12 @@ namespace UnrealBuildTool
 			}
 
 			// Let pch compile at higher priority
-			System.Diagnostics.ProcessPriorityClass priority = ProcessPriority;
-			if (ProcessPriority == System.Diagnostics.ProcessPriorityClass.BelowNormal)
+			System.Diagnostics.ProcessPriorityClass priority = UBAConfig.ProcessPriority;
+			if (priority == System.Diagnostics.ProcessPriorityClass.BelowNormal && action.StatusDescription.Contains("PCH."))
 			{
-				if (action.StatusDescription.Contains("PCH."))
-				{
 					priority = System.Diagnostics.ProcessPriorityClass.Normal;
-				}
 			}
-				
+
 			ProcessStartInfo startInfo = new()
 			{
 				Application = action.CommandPath.FullName,
diff --git a/Engine/Source/Programs/UnrealBuildTool/Executors/UnrealBuildAccelerator/UnrealBuildAcceleratorConfig.cs b/Engine/Source/Programs/UnrealBuildTool/Executors/UnrealBuildAccelerator/UnrealBuildAcceleratorConfig.cs
--- a/Engine/Source/Programs/UnrealBuildTool/Executors/UnrealBuildAccelerator/UnrealBuildAcceleratorConfig.cs
+++ b/Engine/Source/Programs/UnrealBuildTool/Executors/UnrealBuildAccelerator/UnrealBuildAcceleratorConfig.cs
@@ -2,7 +2,9 @@
 
 using System;
 using System.Collections.Generic;
+using System.Diagnostics;
 using EpicGames.Core;
+using UnrealBuildBase;
 
 namespace UnrealBuildTool
 {
@@ -80,12 +82,17 @@ namespace UnrealBuildTool
 		[CommandLine("-UBAForceRemote")]
 		public bool bForceBuildAllRemote { get; set; } = false;
 
+		private bool _allowRetry = true;
 		/// <summary>
 		/// When set to true, actions that fail locally with certain error codes will retry without uba.
 		/// </summary>
 		[XmlConfigFile(Category = "UnrealBuildAccelerator")]
 		[CommandLine("-UBANoRetry", Value = "false")]
-		public bool bAllowRetry { get; set; } = true;
+		public bool bAllowRetry
+		{
+			get => bForcedRetry ? true : _allowRetry;
+			set => _allowRetry = value;
+		}
 
 		/// <summary>
 		/// When set to true, all actions that fail locally with UBA will be retried without UBA.
@@ -250,12 +257,17 @@ namespace UnrealBuildTool
 		[CommandLine("-UBANoCustomMalloc", Value = "true")]
 		public bool bDisableCustomAlloc { get; set; } = false;
 
+		private string _zone = String.Empty;
 		/// <summary>
 		/// The zone to use for UBA.
 		/// </summary>
 		[XmlConfigFile(Category = "UnrealBuildAccelerator")]
 		[CommandLine("-UBAZone=")]
-		public string Zone { get; set; } = String.Empty;
+		public string Zone
+		{
+			get => bDisableRemote ? "local" : _zone;
+			set => _zone = value;
+		}
 
 		/// <summary>
 		/// Set to true to enable encryption when transferring files over the network.
@@ -332,5 +344,76 @@ namespace UnrealBuildTool
 		[XmlConfigFile(Category = "UnrealBuildAccelerator")]
 		[CommandLine("-UBACompressionLevel")]
 		public string CompressionLevel { get; set; } = String.Empty;
+
+		/// <summary>
+		/// Processor count multiplier for local execution. Can be below 1 to reserve CPU for other tasks.
+		/// When using the local executor (not XGE), run a single action on each CPU core. Note that you can set this to a larger value
+		/// to get slightly faster build times in many cases, but your computer's responsiveness during compiling may be much worse.
+		/// This value is ignored if the CPU does not support hyper-threading.
+		/// </summary>
+		[XmlConfigFile(Category = "ParallelExecutor")]
+		public double ProcessorCountMultiplier { get; set; } = 1.0;
+
+		/// <summary>
+		/// Free memory per action in bytes, used to limit the number of parallel actions if the machine is memory starved.
+		/// Set to 0 to disable free memory checking.
+		/// </summary>
+		[XmlConfigFile(Category = "ParallelExecutor")]
+		public double MemoryPerActionBytes { get; set; } = 1.5 * 1024 * 1024 * 1024;
+
+		/// <summary>
+		/// The priority to set for spawned processes.
+		/// Valid Settings: Idle, BelowNormal, Normal, AboveNormal, High
+		/// Default: BelowNormal or Normal for an Asymmetrical processor as BelowNormal can cause scheduling issues.
+		/// </summary>
+		[XmlConfigFile(Category = "ParallelExecutor")]
+		public ProcessPriorityClass ProcessPriority { get; set; } = Utils.IsAsymmetricalProcessor() ? ProcessPriorityClass.Normal : ProcessPriorityClass.BelowNormal;
+
+		/// <summary>
+		/// When enabled, will stop compiling targets after a compile error occurs.
+		/// </summary>
+		[XmlConfigFile(Category = "ParallelExecutor")]
+		public bool bStopCompilationAfterErrors { get; set; } = false;
+
+		/// <summary>
+		/// When set, will stop compiling targets after a set number of errors occurs.
+		/// </summary>
+		[XmlConfigFile(Category = "ParallelExecutor")]
+		public int StopCompilationAfterNumErrors { get; set; } = Int32.MaxValue;
+
+		/// <summary>
+		/// Whether to show compilation times along with worst offenders or not.
+		/// </summary>
+		[XmlConfigFile(Category = "ParallelExecutor")]
+		public bool bShowCompilationTimes { get; set; } = Unreal.IsBuildMachine();
+
+		/// <summary>
+		/// Whether to show compilation times for each executed action
+		/// </summary>
+		[XmlConfigFile(Category = "ParallelExecutor")]
+		public bool bShowPerActionCompilationTimes { get; set; } = Unreal.IsBuildMachine();
+
+		/// <summary>
+		/// Whether to log command lines for actions being executed
+		/// </summary>
+		[XmlConfigFile(Category = "ParallelExecutor")]
+		public bool bLogActionCommandLines { get; set; } = false;
+
+		/// <summary>
+		/// Add target names for each action executed
+		/// </summary>
+		[XmlConfigFile(Category = "ParallelExecutor")]
+		public bool bPrintActionTargetNames { get; set; } = false;
+
+		/// <summary>
+		/// Whether to show CPU utilization after the work is complete.
+		/// </summary>
+		[XmlConfigFile(Category = "ParallelExecutor")]
+		public bool bShowCPUUtilization { get; set; } = Unreal.IsBuildMachine();
+
+		/// <summary>
+		/// Whether UBA can retry actions natively if they fail when detoured locally
+		/// </summary>
+		public bool CanRetryNatively { get; set; } = true;
 	}
 }
\ No newline at end of file
diff --git a/Engine/Source/Programs/UnrealBuildTool/Executors/XGE.cs b/Engine/Source/Programs/UnrealBuildTool/Executors/XGE.cs
--- a/Engine/Source/Programs/UnrealBuildTool/Executors/XGE.cs
+++ b/Engine/Source/Programs/UnrealBuildTool/Executors/XGE.cs
@@ -44,38 +44,38 @@ namespace UnrealBuildTool
 		/// <summary>
 		/// When set to false, XGE will not be enabled when running connected to the coordinator over VPN. Configure VPN-assigned subnets via the VpnSubnets parameter.
 		/// </summary>
-		[XmlConfigFile(Category = "XGE")]
+		[XmlConfigFile(Category = ExecutorName)]
 		static bool bAllowOverVpn = true;
 
 		/// <summary>
 		/// List of subnets containing IP addresses assigned by VPN
 		/// </summary>
-		[XmlConfigFile(Category = "XGE")]
+		[XmlConfigFile(Category = ExecutorName)]
 		static string[]? VpnSubnets = null;
 
 		/// <summary>
 		/// Whether to allow remote linking
 		/// </summary>
-		[XmlConfigFile(Category = "XGE")]
+		[XmlConfigFile(Category = ExecutorName)]
 		static bool bAllowRemoteLinking = false;
 
 		/// <summary>
 		/// Whether to enable the VCCompiler=true setting. This requires an additional license for VC tools. 
 		/// </summary>
-		[XmlConfigFile(Category = "XGE")]
+		[XmlConfigFile(Category = ExecutorName)]
 		static bool bUseVCCompilerMode = false;
 
 		/// <summary>
 		/// Minimum number of actions to use XGE execution.
 		/// </summary>
-		[XmlConfigFile(Category = "XGE")]
+		[XmlConfigFile(Category = ExecutorName)]
 		public static int MinActions = 2;
 
 		/// <summary>
 		/// Check for a concurrent XGE build and treat the XGE executor as unavailable if it's in use.
 		/// This will allow UBT to fall back to another executor such as the parallel executor. 
 		/// </summary>
-		[XmlConfigFile(Category = "XGE")]
+		[XmlConfigFile(Category = ExecutorName)]
 		static bool bUnavailableIfInUse = false;
 #pragma warning restore IDE0044
 
@@ -99,7 +99,9 @@ namespace UnrealBuildTool
 			XmlConfig.ApplyTo(this);
 		}
 
-		public override string Name => "XGE";
+		internal const string ExecutorName = "XGE";
+
+		public override string Name => ExecutorName;
 
 		public static bool TryGetXgConsoleExecutable([NotNullWhen(true)] out string? OutXgConsoleExe)
 		{
