UnrealBuildTool: Suppress xml documentation analysis warnings for target and module rules
Rules projects inherit Configurations settings from shared property sheet

#rnx
#Synced-CL 48635954

diff --git a/Engine/Source/Programs/UnrealBuildTool/ProjectFiles/ProjectFileGenerator.cs b/Engine/Source/Programs/UnrealBuildTool/ProjectFiles/ProjectFileGenerator.cs
--- a/Engine/Source/Programs/UnrealBuildTool/ProjectFiles/ProjectFileGenerator.cs
+++ b/Engine/Source/Programs/UnrealBuildTool/ProjectFiles/ProjectFileGenerator.cs
@@ -782,9 +782,8 @@ namespace UnrealBuildTool
 
 				Writer.WriteLine("  <PropertyGroup>");
 				Writer.WriteLine($"    <TargetFramework>{TargetFramework}</TargetFramework>");
-				Writer.WriteLine("    <AppendTargetFrameworkToOutputPath>false</AppendTargetFrameworkToOutputPath>"); // Shorten intermediate filepath slightly
-				Writer.WriteLine("    <Configurations>Debug;Release;Development;Analyze</Configurations>"); // VCSharpProject requires at least Debug & Development configurations
 				Writer.WriteLine($"    <DefineConstants>$(DefineConstants);{String.Join(';', RulesAssembly.PreprocessorDefines!)}</DefineConstants>");
+				Writer.WriteLine($"    <NoWarn>$(NoWarn);CS1587;CS1591</NoWarn>");
 				Writer.WriteLine("  </PropertyGroup>");
 
 				Writer.WriteLine("  <ItemGroup>");
