UnrealBuildTool: Cleanup BuildHostPlatform.cs

#rnx
#Synced-CL 48679523

diff --git a/Engine/Source/Programs/UnrealBuildTool/System/BuildHostPlatform.cs b/Engine/Source/Programs/UnrealBuildTool/System/BuildHostPlatform.cs
--- a/Engine/Source/Programs/UnrealBuildTool/System/BuildHostPlatform.cs
+++ b/Engine/Source/Programs/UnrealBuildTool/System/BuildHostPlatform.cs
@@ -5,6 +5,7 @@ using System.Collections.Generic;
 using System.Diagnostics;
 using System.IO;
 using System.Linq;
+using System.Runtime.Versioning;
 using EpicGames.Core;
 
 namespace UnrealBuildTool
@@ -30,37 +31,27 @@ namespace UnrealBuildTool
 	/// </summary>
 	public abstract class BuildHostPlatform
 	{
-		private static BuildHostPlatform? CurrentPlatform;
-
 		/// <summary>
 		/// Host platform singleton.
 		/// </summary>
-		public static BuildHostPlatform Current
+		public static BuildHostPlatform Current => s_currentPlatform.Value;
+
+		private static readonly Lazy<BuildHostPlatform> s_currentPlatform = new(() =>
 		{
-			get
+			if (OperatingSystem.IsWindows())
 			{
-				if (CurrentPlatform == null)
-				{
-					if (OperatingSystem.IsWindows())
-					{
-						CurrentPlatform = new WindowsBuildHostPlatform();
-					}
-					else if (OperatingSystem.IsMacOS())
-					{
-						CurrentPlatform = new MacBuildHostPlatform();
-					}
-					else if (OperatingSystem.IsLinux())
-					{
-						CurrentPlatform = new LinuxBuildHostPlatform();
-					}
-					else
-					{
-						throw new NotImplementedException();
-					}
-				}
-				return CurrentPlatform;
+				return new WindowsBuildHostPlatform();
 			}
-		}
+			else if (OperatingSystem.IsMacOS())
+			{
+				return new MacBuildHostPlatform();
+			}
+			else if (OperatingSystem.IsLinux())
+			{
+				return new LinuxBuildHostPlatform();
+			}
+			throw new PlatformNotSupportedException();
+		});
 
 		/// <summary>
 		/// Gets the current host platform type.
@@ -85,55 +76,53 @@ namespace UnrealBuildTool
 		/// <summary>
 		/// Class that holds information about a running process
 		/// </summary>
-		public class ProcessInfo
+		[DebuggerDisplay("{PID}, {Name}, {Filename}")]
+		public record ProcessInfo
 		{
 			/// <summary>
 			/// Process ID
 			/// </summary>
-			public int PID;
+			public int PID { get; init; }
 
 			/// <summary>
 			/// Name of the process
 			/// </summary>
-			public string Name;
+			public string Name { get; init; }
 
 			/// <summary>
 			/// Filename of the process binary
 			/// </summary>
-			public string Filename;
+			public string Filename { get; init; }
 
 			/// <summary>
 			/// Constructor
 			/// </summary>
-			/// <param name="InPID">The process ID</param>
-			/// <param name="InName">The process name</param>
-			/// <param name="InFilename">The process filename</param>
-			public ProcessInfo(int InPID, string InName, string InFilename)
+			/// <param name="pid">The process ID</param>
+			/// <param name="name">The process name</param>
+			/// <param name="filename">The process filename</param>
+			public ProcessInfo(int pid, string name, string filename)
 			{
-				PID = InPID;
-				Name = InName;
-				Filename = InFilename;
+				PID = pid;
+				Name = name;
+				Filename = filename;
 			}
 
 			/// <summary>
 			/// Constructor
 			/// </summary>
-			/// <param name="Proc">Process to take information from</param>
-			public ProcessInfo(Process Proc)
+			/// <param name="proc">Process to take information from</param>
+			public ProcessInfo(Process proc)
 			{
-				PID = Proc.Id;
-				Name = Proc.ProcessName;
-				Filename = Proc.MainModule?.FileName != null ? Path.GetFullPath(Proc.MainModule.FileName) : String.Empty;
+				PID = proc.Id;
+				Name = proc.ProcessName;
+				Filename = proc.MainModule?.FileName != null ? Path.GetFullPath(proc.MainModule.FileName) : String.Empty;
 			}
 
-			/// <summary>
-			/// Format as a string for debugging
-			/// </summary>
-			/// <returns>String containing process info</returns>
-			public override string ToString()
-			{
-				return String.Format("{0}, {1}", Name, Filename);
-			}
+			/// <inheritdoc/>
+			public override int GetHashCode() => HashCode.Combine(PID, Name, Filename);
+
+			/// <inheritdoc/>
+			public override string? ToString() => $"{Name}, {Filename}";
 		}
 
 		/// <summary>
@@ -150,7 +139,7 @@ namespace UnrealBuildTool
 				{
 					if (!Proc.HasExited)
 					{
-						Result.Add(new ProcessInfo(Proc));
+						Result.Add(new(Proc));
 					}
 				}
 				catch { }
@@ -161,65 +150,40 @@ namespace UnrealBuildTool
 		/// <summary>
 		/// Gets a process by name.
 		/// </summary>
-		/// <param name="Name">Name of the process to get information for.</param>
+		/// <param name="name">Name of the process to get information for.</param>
 		/// <returns></returns>
-		public virtual ProcessInfo? GetProcessByName(string Name)
-		{
-			ProcessInfo[] AllProcess = GetProcesses();
-			foreach (ProcessInfo Info in AllProcess)
-			{
-				if (Info.Name == Name)
-				{
-					return Info;
-				}
-			}
-			return null;
-		}
+		public virtual ProcessInfo? GetProcessByName(string name) => GetProcesses().FirstOrDefault(x => x.Name == name);
 
 		/// <summary>
 		/// Gets processes by name.
 		/// </summary>
-		/// <param name="Name">Name of the process to get information for.</param>
+		/// <param name="name">Name of the process to get information for.</param>
 		/// <returns></returns>
-		public virtual ProcessInfo[] GetProcessesByName(string Name)
-		{
-			ProcessInfo[] AllProcess = GetProcesses();
-			List<ProcessInfo> Result = new List<ProcessInfo>();
-			foreach (ProcessInfo Info in AllProcess)
-			{
-				if (Info.Name == Name)
-				{
-					Result.Add(Info);
-				}
-			}
-			return [.. Result];
-		}
+		public virtual ProcessInfo[] GetProcessesByName(string name) => [.. GetProcesses().Where(x => x.Name == name)];
 
 		/// <summary>
 		/// Gets the filenames of all modules associated with a process
 		/// </summary>
-		/// <param name="PID">Process ID</param>
-		/// <param name="Filename">Filename of the binary associated with the process.</param>
+		/// <param name="pid">Process ID</param>
+		/// <param name="filename">Filename of the binary associated with the process.</param>
 		/// <returns>An array of all module filenames associated with the process. Can be empty of the process is no longer running.</returns>
-		public virtual string[] GetProcessModules(int PID, string Filename)
+		public virtual string[] GetProcessModules(int pid, string filename)
 		{
-			List<string> Modules = new List<string>();
 			try
 			{
-				Process Proc = Process.GetProcessById(PID);
+				Process Proc = Process.GetProcessById(pid);
 				if (Proc != null)
 				{
-					foreach (ProcessModule Module in Proc.Modules.Cast<System.Diagnostics.ProcessModule>())
-					{
-						if (Module.FileName != null)
-						{
-							Modules.Add(Path.GetFullPath(Module.FileName));
-						}
-					}
+					return [..Proc.Modules
+						.Cast<ProcessModule>()
+						.Where(module => module.FileName != null)
+						.Select(module => Path.GetFullPath(module.FileName))
+						.Order()
+					];
 				}
 			}
 			catch { }
-			return [.. Modules];
+			return [];
 		}
 
 		/// <summary>
@@ -229,11 +193,12 @@ namespace UnrealBuildTool
 		internal abstract IEnumerable<ProjectFileFormat> GetDefaultProjectFileFormats();
 	}
 
+	[SupportedOSPlatform("windows")]
 	class WindowsBuildHostPlatform : BuildHostPlatform
 	{
 		public override UnrealTargetPlatform Platform => UnrealTargetPlatform.Win64;
 
-		public override FileReference Shell => new FileReference(Environment.GetEnvironmentVariable("COMSPEC")!);
+		public override FileReference Shell { get; } = new(Environment.GetEnvironmentVariable("ComSpec") ?? Path.Combine(Environment.SystemDirectory, "cmd.exe"));
 
 		public override ShellType ShellType => ShellType.Cmd;
 
@@ -248,111 +213,69 @@ namespace UnrealBuildTool
 		}
 	}
 
+	[SupportedOSPlatform("macos")]
 	class MacBuildHostPlatform : BuildHostPlatform
 	{
 		public override UnrealTargetPlatform Platform => UnrealTargetPlatform.Mac;
 
-		public override FileReference Shell => new FileReference("/bin/sh");
+		public override FileReference Shell => new("/bin/sh");
 
 		public override ShellType ShellType => ShellType.Sh;
 
 		public override string BinarySuffix => String.Empty;
 
 		/// <summary>
-		/// (needs confirmation) Currently returns incomplete process names in Process.GetProcesses() so we need to parse 'ps' output.
-		/// </summary>
-		/// <returns></returns>
-		public override ProcessInfo[] GetProcesses()
-		{
-			List<ProcessInfo> Result = new List<ProcessInfo>();
-
-			string TempFile = Path.Combine("/var/tmp", Path.GetTempFileName());
-			ProcessStartInfo StartInfo = new ProcessStartInfo();
-			StartInfo.FileName = "/bin/sh";
-			StartInfo.Arguments = "-c \"ps -eaw -o pid,comm > " + TempFile + "\"";
-			StartInfo.CreateNoWindow = true;
-
-			Process Proc = new Process();
-			Proc.StartInfo = StartInfo;
-			try
-			{
-				Proc.Start();
-				foreach (string FileLine in File.ReadAllLines(TempFile))
-				{
-					string Line = FileLine.Trim();
-					int PIDEnd = Line.IndexOf(' ');
-					string PIDString = Line.Substring(0, PIDEnd);
-					if (PIDString != "PID")
-					{
-						string Filename = Line.Substring(PIDEnd + 1);
-						int Pid = Int32.Parse(PIDString);
-						try
-						{
-							Process ExistingProc = Process.GetProcessById(Pid);
-							if (ExistingProc != null && Pid != Environment.ProcessId && ExistingProc.HasExited == false)
-							{
-								ProcessInfo ProcInfo = new ProcessInfo(ExistingProc.Id, Path.GetFileName(Filename), Filename);
-								Result.Add(ProcInfo);
-							}
-						}
-						catch { }
-					}
-				}
-				File.Delete(TempFile);
-				Proc.WaitForExit();
-			}
-			catch { }
-			return [.. Result];
-		}
-
-		/// <summary>
-		/// (needs confirmation) Currently returns incomplete list of modules for Process.Modules so we need to parse vmmap output.
+		/// Currently returns incomplete list of modules for Process.Modules so we need to parse vmmap output.
+		/// [REDACTED_PROFILE_URL]/runtime/issues/64042
 		/// </summary>
-		/// <param name="PID"></param>
-		/// <param name="Filename"></param>
+		/// <param name="pid"></param>
+		/// <param name="filename"></param>
 		/// <returns></returns>
-		public override string[] GetProcessModules(int PID, string Filename)
+		public override string[] GetProcessModules(int pid, string filename)
 		{
-			HashSet<string> Modules = new HashSet<string>();
 			// Add the process file name to the module list. This is to make it compatible with the results of Process.Modules on Windows.
-			Modules.Add(Filename);
+			HashSet<string> Modules = [filename];
 
-			ProcessStartInfo StartInfo = new ProcessStartInfo();
-			StartInfo.FileName = "vmmap";
-			StartInfo.Arguments = String.Format("{0} -w", PID);
-			StartInfo.CreateNoWindow = true;
-			StartInfo.UseShellExecute = false;
-			StartInfo.RedirectStandardOutput = true;
+			ProcessStartInfo startInfo = new ProcessStartInfo();
+			startInfo.FileName = "vmmap";
+			startInfo.Arguments = $"{pid} -w";
+			startInfo.CreateNoWindow = true;
+			startInfo.UseShellExecute = false;
+			startInfo.RedirectStandardOutput = true;
 
-			Process Proc = new Process();
-			Proc.StartInfo = StartInfo;
+			Process process = new Process();
+			process.StartInfo = startInfo;
 			try
 			{
-				Proc.Start();
+				process.Start();
 				// Start processing output before vmmap exits otherwise it's going to hang
-				while (!Proc.WaitForExit(1))
+				while (!process.WaitForExit(1))
 				{
-					ProcessVMMapOutput(Proc, Modules);
+					ProcessVMMapOutput(process, Modules);
 				}
-				ProcessVMMapOutput(Proc, Modules);
+				ProcessVMMapOutput(process, Modules);
 			}
 			catch { }
-			return [.. Modules];
+			return [.. Modules.Order()];
 		}
-		private void ProcessVMMapOutput(Process Proc, HashSet<string> Modules)
+
+		private void ProcessVMMapOutput(Process process, HashSet<string> modules)
 		{
-			for (string? Line = Proc.StandardOutput.ReadLine(); Line != null; Line = Proc.StandardOutput.ReadLine())
+			for (string? line = process.StandardOutput.ReadLine(); line != null; line = process.StandardOutput.ReadLine())
 			{
-				Line = Line.Trim();
-				if (Line.EndsWith(".dylib"))
+				line = line.Trim();
+				if (line.EndsWith(".dylib"))
 				{
-					const int SharingModeLength = 6;
-					int SMStart = Line.IndexOf("SM=");
-					int PathStart = SMStart + SharingModeLength;
-					string Module = Line.Substring(PathStart).Trim();
-					if (!Modules.Contains(Module))
+					const int sharingModeLength = 6;
+					int smStart = line.IndexOf("SM=");
+					if (smStart != -1)
 					{
-						Modules.Add(Module);
+						int pathStart = smStart + sharingModeLength;
+						string module = line.Substring(pathStart).Trim();
+						if (Path.Exists(module))
+						{
+							modules.Add(module);
+						}
 					}
 				}
 			}
@@ -367,38 +290,31 @@ namespace UnrealBuildTool
 		}
 	}
 
+	[SupportedOSPlatform("linux")]
 	class LinuxBuildHostPlatform : BuildHostPlatform
 	{
 		public override UnrealTargetPlatform Platform => UnrealTargetPlatform.Linux;
 
-		public override FileReference Shell => new FileReference("/bin/sh");
+		public override FileReference Shell => new("/bin/sh");
 
 		public override ShellType ShellType => ShellType.Sh;
 
 		public override string BinarySuffix => String.Empty;
 
 		/// <summary>
-		/// (needs confirmation) Currently returns incomplete process names in Process.GetProcesses() so we need to use /proc
-		/// (also, locks up during process traversal sometimes, trying to open /dev/snd/pcm*)
-		/// </summary>
-		/// <returns></returns>
-		public override ProcessInfo[] GetProcesses()
-		{
-			// @TODO: Implement for Linux
-			return [];
-		}
-
-		/// <summary>
-		/// (needs confirmation) Currently returns incomplete list of modules for Process.Modules so we need to parse /proc/PID/maps.
+		/// Currently returns incomplete list of modules for Process.Modules so we need to parse /proc/PID/maps.
 		/// (also, locks up during process traversal sometimes, trying to open /dev/snd/pcm*)
+		/// [REDACTED_PROFILE_URL]/runtime/issues/64042
 		/// </summary>
-		/// <param name="PID"></param>
-		/// <param name="Filename"></param>
+		/// <param name="pid"></param>
+		/// <param name="filename"></param>
 		/// <returns></returns>
-		public override string[] GetProcessModules(int PID, string Filename)
+		public override string[] GetProcessModules(int pid, string filename)
 		{
+			// Add the process file name to the module list. This is to make it compatible with the results of Process.Modules on Windows.
+			HashSet<string> modules = [filename];
 			// @TODO: Implement for Linux
-			return [];
+			return [.. modules.Order()];
 		}
 
 		internal override IEnumerable<ProjectFileFormat> GetDefaultProjectFileFormats()
