Add "long" integer support to plist support

#Synced-CL Unknown

diff --git a/Engine/Source/Programs/UnrealBuildTool/Platform/IOS/PListHelper.cs b/Engine/Source/Programs/UnrealBuildTool/Platform/IOS/PListHelper.cs
--- a/Engine/Source/Programs/UnrealBuildTool/Platform/IOS/PListHelper.cs
+++ b/Engine/Source/Programs/UnrealBuildTool/Platform/IOS/PListHelper.cs
@@ -160,6 +160,11 @@ namespace UnrealBuildTool
 				ValueElement = Doc.CreateElement("real");
 				ValueElement.InnerText = vdouble.ToString();
 			}
+			else if (Value is long vlong)
+			{
+				ValueElement = Doc.CreateElement("integer");
+				ValueElement.InnerText = vlong.ToString();
+			}
 			else if (Value is int vint)
 			{
 				ValueElement = Doc.CreateElement("integer");
