Control UBA's resource disposal through the UnrealBuildAcceleratorConfig class. With the default value of false, we let the process release some file mappings instead of doing it from code which speeds up shutdown a bit when large file mappings have been used.

#Synced-CL 48713665

diff --git a/Engine/Source/Programs/UnrealBuildTool/Executors/UnrealBuildAccelerator/UBAExecutor.cs b/Engine/Source/Programs/UnrealBuildTool/Executors/UnrealBuildAccelerator/UBAExecutor.cs
--- a/Engine/Source/Programs/UnrealBuildTool/Executors/UnrealBuildAccelerator/UBAExecutor.cs
+++ b/Engine/Source/Programs/UnrealBuildTool/Executors/UnrealBuildAccelerator/UBAExecutor.cs
@@ -166,11 +166,11 @@ namespace UnrealBuildTool
 
 			_targetDescriptors = targetDescriptors;
 			_cancellationSource = new();
-			_initTask = new(() => { Init(targetDescriptors, logger, true); }, TaskCreationOptions.LongRunning);
+			_initTask = new(() => { Init(targetDescriptors, logger); }, TaskCreationOptions.LongRunning);
 			_initTask.Start();
 		}
 
-		private void Init(IEnumerable<TargetDescriptor> targetDescriptors, ILogger logger, bool cleanupOnShutdown)
+		private void Init(IEnumerable<TargetDescriptor> targetDescriptors, ILogger logger)
 		{
 			XmlConfig.ApplyTo(this);
 			CommandLine.ParseArguments(Environment.GetCommandLineArgs(), this, logger);
@@ -281,7 +281,7 @@ namespace UnrealBuildTool
 				config.AddValue("Storage", "CompressionLevel", UBAConfig.CompressionLevel);
 			}
 
-			if (!cleanupOnShutdown)
+			if (!UBAConfig.bCleanupOnDispose)
 			{
 				config.AddValue("Storage", "SkipCleanupOnShutdown", "true"); // For faster shutdown
 			}
diff --git a/Engine/Source/Programs/UnrealBuildTool/Executors/UnrealBuildAccelerator/UnrealBuildAcceleratorConfig.cs b/Engine/Source/Programs/UnrealBuildTool/Executors/UnrealBuildAccelerator/UnrealBuildAcceleratorConfig.cs
--- a/Engine/Source/Programs/UnrealBuildTool/Executors/UnrealBuildAccelerator/UnrealBuildAcceleratorConfig.cs
+++ b/Engine/Source/Programs/UnrealBuildTool/Executors/UnrealBuildAccelerator/UnrealBuildAcceleratorConfig.cs
@@ -411,6 +411,14 @@ namespace UnrealBuildTool
 		[XmlConfigFile(Category = "ParallelExecutor")]
 		public bool bShowCPUUtilization { get; set; } = Unreal.IsBuildMachine();
 
+		/// <summary>
+		/// Whether to cleanup the UBA executor's resources on dispose. This will result in resources being freed sooner,
+		/// rather than waiting until process shutdown, at the expense of taking longer before the UBA process terminates.
+		/// Unless you know your UBT process will significantly outlive the UBA executor and need to free up resources now,
+		/// it's suggested you leave this as false.
+		/// </summary>
+		public bool bCleanupOnDispose { get; set; } = false;
+
 		/// <summary>
 		/// Whether UBA can retry actions natively if they fail when detoured locally
 		/// </summary>
