Explicitly mark UBT as non-CLS-compliant to fix issue CA1014.

#Synced-CL 48713665

diff --git a/Engine/Source/Programs/UnrealBuildTool/Properties/AssemblyInfo.cs b/Engine/Source/Programs/UnrealBuildTool/Properties/AssemblyInfo.cs
--- a/Engine/Source/Programs/UnrealBuildTool/Properties/AssemblyInfo.cs
+++ b/Engine/Source/Programs/UnrealBuildTool/Properties/AssemblyInfo.cs
@@ -1,5 +1,6 @@
 // Copyright Epic Games, Inc. All Rights Reserved.
 
+using System;
 using System.Reflection;
 using System.Runtime.CompilerServices;
 
@@ -11,3 +12,6 @@ using System.Runtime.CompilerServices;
 [assembly: AssemblyConfiguration("")]
 
 [assembly: InternalsVisibleTo("UnrealBuildTool.Tests")]
+
+// We are not CLS compliant, as we make extensive use of e.g. uint.
+[assembly: CLSCompliant(false)]
