[Backout] - CL48650844
#rnx
Original CL Desc
-----------------------------------------------------------------
Replace the ParallelExecutor with UBAExecutor configured to disable detouring, as this should be equivalent in functionality, and allows us to reduce our code footprint.

#Synced-CL 48648194

diff --git a/Engine/Source/Programs/UnrealBuildTool/Actions/ActionGraph.cs b/Engine/Source/Programs/UnrealBuildTool/Actions/ActionGraph.cs
--- a/Engine/Source/Programs/UnrealBuildTool/Actions/ActionGraph.cs
+++ b/Engine/Source/Programs/UnrealBuildTool/Actions/ActionGraph.cs
@@ -388,70 +388,33 @@ namespace UnrealBuildTool
 			}
 		}
 
-		internal static UBAExecutor GetUBAExecutor(BuildConfiguration BuildConfiguration, ILogger Logger, List<TargetDescriptor> TargetDescriptors)
+		public static ActionExecutor? PrecreateExecutor(BuildConfiguration BuildConfiguration, ILogger Logger, List<TargetDescriptor> TargetDescriptors)
 		{
-			if (!UBAExecutor.IsAvailable())
+			if (BuildConfiguration.bAllowXGE && XGE.IsAvailable(Logger))
 			{
-				throw new BuildLogEventException("UBA is not available - please ensure the UBA binaries exist for your host platform");
+				return null;
 			}
-
-			UnrealBuildAcceleratorConfig Config = UBAExecutor.GetAppliedConfig(TargetDescriptors);
-			// We always use the UBA executor, but we disable detouring to mirror legacy behaviour if the config disables it.
-			if (!BuildConfiguration.bAllowUBAExecutor)
+			if (BuildConfiguration.bAllowSNDBS && SNDBS.IsAvailable(Logger))
 			{
-				Config.bAllowDetour = false;
+				return null;
 			}
-
-			return new UBAExecutor(
-				Config,
-				BuildConfiguration.MaxParallelActions,
-				BuildConfiguration.bAllCores,
-				BuildConfiguration.bCompactOutput,
-				Logger,
-				TargetDescriptors);
-		}
-
-		public static ActionExecutor? PrecreateExecutor(BuildConfiguration BuildConfiguration, ILogger Logger, List<TargetDescriptor> TargetDescriptors)
-		{
-			foreach (string Name in BuildConfiguration.RemoteExecutorPriority)
+			if (BuildConfiguration.bAllowFASTBuild && FASTBuild.IsAvailable(Logger))
 			{
-				switch (Name)
-				{
-					case XGE.ExecutorName:
-						if (BuildConfiguration.bAllowXGE && XGE.IsAvailable(Logger))
-						{
-							return null;
-						}
-						break;
-					case SNDBS.ExecutorName:
-						if (BuildConfiguration.bAllowSNDBS && SNDBS.IsAvailable(Logger))
-						{
-							return null;
-						}
-						break;
-					case FASTBuild.ExecutorName:
-						if (BuildConfiguration.bAllowFASTBuild && FASTBuild.IsAvailable(Logger))
-						{
-							return null;
-						}
-						break;
-					case UBAExecutor.ExecutorName:
-						if (BuildConfiguration.bAllowUBAExecutor)
-						{
-							return GetUBAExecutor(BuildConfiguration, Logger, TargetDescriptors);
-						}
-						break;
-				}
+				return null;
+			}
+			if (BuildConfiguration.bAllowUBAExecutor && UBAExecutor.IsAvailable())
+			{
+				return new UBAExecutor(BuildConfiguration.MaxParallelActions, BuildConfiguration.bAllCores, BuildConfiguration.bCompactOutput, Logger, TargetDescriptors, true);
 			}
 
-			return GetUBAExecutor(BuildConfiguration, Logger, TargetDescriptors);
+			return null;
 		}
 
 		private static ActionExecutor? GetRemoteExecutorByName(string Name, BuildConfiguration BuildConfiguration, int ActionCount, int MinActionsForRemote, List<TargetDescriptor> TargetDescriptors, ILogger Logger)
 		{
 			switch (Name)
 			{
-				case XGE.ExecutorName:
+				case "XGE":
 					{
 						if (ActionCount >= MinActionsForRemote && BuildConfiguration.bAllowXGE && XGE.IsAvailable(Logger) && ActionCount >= XGE.MinActions)
 						{
@@ -459,7 +422,7 @@ namespace UnrealBuildTool
 						}
 						return null;
 					}
-				case SNDBS.ExecutorName:
+				case "SNDBS":
 					{
 						if (ActionCount >= MinActionsForRemote && BuildConfiguration.bAllowSNDBS && SNDBS.IsAvailable(Logger))
 						{
@@ -467,20 +430,20 @@ namespace UnrealBuildTool
 						}
 						return null;
 					}
-				case FASTBuild.ExecutorName:
+				case "FASTBuild":
 					{
 						if (ActionCount >= MinActionsForRemote && BuildConfiguration.bAllowFASTBuild && FASTBuild.IsAvailable(Logger))
 						{
-							return new FASTBuild(BuildConfiguration, Logger, TargetDescriptors);
+							return new FASTBuild(BuildConfiguration.MaxParallelActions, BuildConfiguration.bAllCores, BuildConfiguration.bCompactOutput, Logger);
 						}
 						return null;
 					}
-				case UBAExecutor.ExecutorName:
+				case "UBA":
 					{
 						// Intentionally not checking MinActionsForRemote
-						if (BuildConfiguration.bAllowUBAExecutor)
+						if (BuildConfiguration.bAllowUBAExecutor && UBAExecutor.IsAvailable())
 						{
-							return GetUBAExecutor(BuildConfiguration, Logger, TargetDescriptors);
+							return new UBAExecutor(BuildConfiguration.MaxParallelActions, BuildConfiguration.bAllCores, BuildConfiguration.bCompactOutput, Logger, TargetDescriptors, true);
 						}
 						return null;
 					}
@@ -497,8 +460,7 @@ namespace UnrealBuildTool
 		/// </summary>
 		private static ActionExecutor SelectExecutor(BuildConfiguration BuildConfiguration, int ActionCount, List<TargetDescriptor> TargetDescriptors, ILogger Logger)
 		{
-			UnrealBuildAcceleratorConfig Config = UBAExecutor.GetAppliedConfig(TargetDescriptors);
-			int MinActionsForRemote = UBAExecutor.GetDefaultNumParallelProcesses(Config, BuildConfiguration.MaxParallelActions, BuildConfiguration.bAllCores, Logger);
+			int MinActionsForRemote = ParallelExecutor.GetDefaultNumParallelProcesses(BuildConfiguration.MaxParallelActions, BuildConfiguration.bAllCores, Logger);
 			foreach (string Name in BuildConfiguration.RemoteExecutorPriority)
 			{
 				ActionExecutor? Executor = GetRemoteExecutorByName(Name, BuildConfiguration, ActionCount, MinActionsForRemote, TargetDescriptors, Logger);
@@ -508,7 +470,7 @@ namespace UnrealBuildTool
 				}
 			}
 
-			return GetUBAExecutor(BuildConfiguration, Logger, TargetDescriptors);
+			return new ParallelExecutor(BuildConfiguration.MaxParallelActions, BuildConfiguration.bAllCores, BuildConfiguration.bCompactOutput, Logger);
 		}
 
 		/// <summary>
diff --git a/Engine/Source/Programs/UnrealBuildTool/Configuration/BuildConfiguration.cs b/Engine/Source/Programs/UnrealBuildTool/Configuration/BuildConfiguration.cs
--- a/Engine/Source/Programs/UnrealBuildTool/Configuration/BuildConfiguration.cs
+++ b/Engine/Source/Programs/UnrealBuildTool/Configuration/BuildConfiguration.cs
@@ -44,7 +44,7 @@ namespace UnrealBuildTool
 		/// Priority order for remote executors (XGE, SNDBS, FASTBuild, UBA)
 		/// </summary>
 		[XmlConfigFile]
-		public string[] RemoteExecutorPriority { get; set; } = [XGE.ExecutorName, SNDBS.ExecutorName, FASTBuild.ExecutorName, UBAExecutor.ExecutorName];
+		public string[] RemoteExecutorPriority { get; set; } = ["XGE", "SNDBS", "FASTBuild", "UBA"];
 
 		/// <summary>
 		/// Whether the UnrealBuildAccelerator executor will be used.
diff --git a/Engine/Source/Programs/UnrealBuildTool/Executors/Experimental/FASTBuild.cs b/Engine/Source/Programs/UnrealBuildTool/Executors/Experimental/FASTBuild.cs
--- a/Engine/Source/Programs/UnrealBuildTool/Executors/Experimental/FASTBuild.cs
+++ b/Engine/Source/Programs/UnrealBuildTool/Executors/Experimental/FASTBuild.cs
@@ -72,7 +72,7 @@ namespace UnrealBuildTool
 		/// <summary>
 		/// Executor to use for local actions
 		/// </summary>
-		readonly UBAExecutor LocalExecutor;
+		readonly ActionExecutor LocalExecutor;
 
 		public static readonly string DefaultExecutableBasePath = Path.Combine(Unreal.EngineDirectory.FullName, "Extras", "ThirdPartyNotUE", "FASTBuild");
 
@@ -161,33 +161,15 @@ namespace UnrealBuildTool
 		/// <summary>
 		/// Constructor
 		/// </summary>
-		public FASTBuild(BuildConfiguration BuildConfiguration, ILogger Logger, IEnumerable<TargetDescriptor> TargetDescriptors)
+		public FASTBuild(int MaxLocalActions, bool bAllCores, bool bCompactOutput, ILogger Logger)
 			: base(Logger)
 		{
 			XmlConfig.ApplyTo(this);
 
-			UnrealBuildAcceleratorConfig Config = UBAExecutor.GetAppliedConfig(TargetDescriptors);
-
-			// Disable remote, as this is intended as a local executor.
-			Config.bDisableRemote = true;
-
-			// We always use the UBA executor, but we disable detouring to mirror legacy behaviour if the config disables it.
-			if (!BuildConfiguration.bAllowUBAExecutor)
-			{
-				Config.bAllowDetour = false;
-			}
-			LocalExecutor = new UBAExecutor(
-				Config,
-				BuildConfiguration.MaxParallelActions,
-				BuildConfiguration.bAllCores,
-				BuildConfiguration.bCompactOutput,
-				Logger,
-				TargetDescriptors);
+			LocalExecutor = new ParallelExecutor(MaxLocalActions, bAllCores, bCompactOutput, Logger);
 		}
 
-		internal const string ExecutorName = "FASTBuild";
-
-		public override string Name => ExecutorName;
+		public override string Name => "FASTBuild";
 
 		public static string GetExecutableName()
 		{
@@ -1373,7 +1355,7 @@ namespace UnrealBuildTool
 			string ForceRemoteArgument = bForceRemote ? "-forceremote" : "";
 			string NoStopOnErrorArgument = bStopOnError ? "" : "-nostoponerror";
 			string IDEArgument = IsApple() ? "" : "-ide";
-			string MaxProcesses = "-j" + LocalExecutor.NumParallelProcesses;
+			string MaxProcesses = "-j" + ((ParallelExecutor)LocalExecutor).NumParallelProcesses;
 
 			// Interesting flags for FASTBuild:
 			// -nostoponerror, -verbose, -monitor (if FASTBuild Monitor Visual Studio Extension is installed!)
diff --git a/Engine/Source/Programs/UnrealBuildTool/Executors/ImmediateActionQueue.cs b/Engine/Source/Programs/UnrealBuildTool/Executors/ImmediateActionQueue.cs
new file mode 100644
--- /dev/null
+++ b/Engine/Source/Programs/UnrealBuildTool/Executors/ImmediateActionQueue.cs
@@ -0,0 +1,1010 @@
+// Copyright Epic Games, Inc. All Rights Reserved.
+
+using System;
+using System.Collections.Generic;
+using System.Linq;
+using System.Threading;
+using System.Threading.Tasks;
+using EpicGames.Core;
+using Microsoft.CodeAnalysis;
+using Microsoft.Extensions.Logging;
+using UnrealBuildBase;
+
+namespace UnrealBuildTool
+{
+	/// <summary>
+	/// Defines the type of the runner.
+	/// </summary>
+	enum ImmediateActionQueueRunnerType : byte
+	{
+
+		/// <summary>
+		/// Will be used to queue jobs as part of general requests
+		/// </summary>
+		Automatic,
+
+		/// <summary>
+		/// Will only be used for manual requests to queue jobs
+		/// </summary>
+		Manual,
+	}
+
+	/// <summary>
+	/// Actions are assigned a runner when they need executing
+	/// </summary>
+	/// <param name="Type">Type of the runner</param>
+	/// <param name="RunAction">Function to be run queue running a task</param>
+	/// <param name="UseActionWeights">If true, use the action weight as a secondary limit</param>
+	/// <param name="MaxActions">Maximum number of action actions</param>
+	/// <param name="MaxActionWeight">Maximum weight of actions</param>
+	record ImmediateActionQueueRunner(ImmediateActionQueueRunnerType Type, Func<LinkedAction, Func<Task>?> RunAction, bool UseActionWeights = false, int MaxActions = Int32.MaxValue, double MaxActionWeight = Int32.MaxValue)
+	{
+		/// <summary>
+		/// Current number of active actions
+		/// </summary>
+		public int ActiveActions = 0;
+
+		/// <summary>
+		/// Current weight of actions
+		/// </summary>
+		public double ActiveActionWeight = 0;
+
+		/// <summary>
+		/// Used to track if the runner has already scanned the action list
+		/// for a given change count but found nothing to run.
+		/// </summary>
+		public int LastActionChange = 0;
+
+		/// <summary>
+		/// True if the current limits have not been reached.
+		/// </summary>
+		public bool IsUnderLimits => ActiveActions < MaxActions && (!UseActionWeights || ActiveActionWeight < MaxActionWeight);
+	}
+
+	/// <summary>
+	/// Helper class to manage the action queue
+	///
+	/// Running the queue can be done with a mixture of automatic and manual runners. Runners are responsible for performing
+	/// the work associated with an action. Automatic runners will have actions automatically queued to them up to the point
+	/// that any runtime limits aren't exceeded (such as maximum number of concurrent processes).  Manual runners must have
+	/// jobs queued to them by calling TryStartOneAction or StartManyActions with the runner specified.
+	/// 
+	/// For example:
+	/// 
+	///		ParallelExecutor uses an automatic runner exclusively.
+	///		UBAExecutor uses an automatic runner to run jobs locally and a manual runner to run jobs remotely as processes 
+	///			become available.
+	/// </summary>
+	class ImmediateActionQueue : ActionLogger, IDisposable
+	{
+
+		/// <summary>
+		/// Number of second of no completed actions to trigger an action stall report.
+		/// If zero, stall reports will not be enabled.
+		/// </summary>
+		[CommandLine("-ActionStallReportTime=")]
+		[XmlConfigFile(Category = "BuildConfiguration")]
+		public int ActionStallReportTime = 0;
+
+		/// <summary>
+		/// Number of second of no completed actions to trigger to terminate the queue.
+		/// If zero, force termination not be enabled.
+		/// </summary>
+		[CommandLine("-ActionStallTerminateTime=")]
+		[XmlConfigFile(Category = "BuildConfiguration")]
+		public int ActionStallTerminateTime = 0;
+
+		/// <summary>
+		/// Running status of the action
+		/// </summary>
+		private enum ActionStatus : byte
+		{
+			/// <summary>
+			/// Queued waiting for execution
+			/// </summary>
+			Queued,
+
+			/// <summary>
+			/// Action is running
+			/// </summary>
+			Running,
+
+			/// <summary>
+			/// Action has successfully finished
+			/// </summary>
+			Finished,
+
+			/// <summary>
+			/// Action has finished with an error
+			/// </summary>
+			Error,
+		}
+
+		/// <summary>
+		/// Used to track the state of each action
+		/// </summary>
+		private struct ActionState
+		{
+			/// <summary>
+			/// Action to be executed
+			/// </summary>
+			public LinkedAction Action;
+
+			/// <summary>
+			/// Current status of the execution
+			/// </summary>
+			public ActionStatus Status;
+
+			/// <summary>
+			/// Runner assigned to the action
+			/// </summary>
+			public ImmediateActionQueueRunner? Runner;
+
+			/// <summary>
+			/// Indices to prereq actions
+			/// </summary>
+			public int[] PrereqActionsSortIndex;
+		};
+
+		/// <summary>
+		/// State information about each action
+		/// </summary>
+		private readonly ActionState[] Actions;
+
+		/// <summary>
+		/// Total number of actions
+		/// </summary>
+		public int TotalActions => Actions.Length;
+
+		/// <summary>
+		/// Source for the cancellation token
+		/// </summary>
+		public readonly CancellationTokenSource CancellationTokenSource = new();
+
+		/// <summary>
+		/// Cancellation token
+		/// </summary>
+		public CancellationToken CancellationToken => CancellationTokenSource.Token;
+
+		/// <summary>
+		/// Overall success of the action queue
+		/// </summary>
+		public bool Success = true;
+
+		/// <summary>
+		/// When enabled, will stop compiling targets after a compile error occurs.
+		/// </summary>
+		public bool StopCompilationAfterErrors = false;
+
+		/// <summary>
+		/// When enabled, will stop compiling targets after a number of compile errors occur.
+		/// </summary>
+		public int StopCompilationAfterNumErrors = Int32.MaxValue;
+
+		public int CompletedActions => _completedActions;
+
+		/// <summary>
+		/// Return true if the queue is done
+		/// </summary>
+		public bool IsDone => _doneTaskSource.Task.IsCompleted;
+
+		/// <summary>
+		///  Action that can be to notify when artifacts have been read for an action
+		/// </summary>
+		public Action<LinkedAction>? OnArtifactsRead = null;
+
+		/// <summary>
+		///  Action that can be to notify when artifacts have been missed
+		/// </summary>
+		public Action<LinkedAction>? OnArtifactsMiss = null;
+
+		/// <summary>
+		/// Collection of available runners
+		/// </summary>
+		private readonly List<ImmediateActionQueueRunner> _runners = new();
+
+		/// <summary>
+		/// First action to start scanning for actions to run
+		/// </summary>
+		private int _firstPendingAction;
+
+		/// <summary>
+		/// Tracks the number of completed actions.
+		/// </summary>
+		private int _completedActions = 0;
+
+		/// <summary>
+		/// Flags used to track how StartManyActions should run
+		/// </summary>
+		private int _startManyFlags = 0;
+
+		/// <summary>
+		/// Number of changes to the action list
+		/// </summary>
+		private int _lastActionChange = 1;
+
+		/// <summary>
+		/// If true, an action stall has been reported for the current change count
+		/// </summary>
+		private bool _lastActionStallReported = false;
+
+		/// <summary>
+		/// If true, the queue has been cancelled due to an action stall
+		/// </summary>
+		private bool _lastActionStallCanceled = false;
+
+		/// <summary>
+		/// Time of the last change to the action count.  This is updated by the timer.
+		/// </summary>
+		private DateTime _lastActionChangeTime;
+
+		/// <summary>
+		/// Copy of _lastActionChange used to detect updates by the time.
+		/// </summary>
+		private int _lastActionStallChange = 1;
+
+		/// <summary>
+		/// Used to terminate the run with status
+		/// </summary>
+		private readonly TaskCompletionSource _doneTaskSource = new();
+
+		private bool _disposedValue;
+
+		/// <summary>
+		/// Construct a new instance of the action queue
+		/// </summary>
+		/// <param name="actions">Collection of actions</param>
+		/// <param name="progressWriterText">Text to be displayed with the progress writer</param>
+		/// <param name="writeToolOutput">Action to invoke when writing tool output</param>
+		/// <param name="flushToolOutput">Action to invoke when flushing tool output</param>
+		/// <param name="logger">Logging interface</param>
+		public ImmediateActionQueue(IEnumerable<LinkedAction> actions, string progressWriterText, Action<string> writeToolOutput, System.Action flushToolOutput, ILogger logger)
+		:	base(actions.Count(), progressWriterText, writeToolOutput, flushToolOutput, logger)
+		{
+			CommandLine.ParseArguments(Environment.GetCommandLineArgs(), this, logger);
+			XmlConfig.ApplyTo(this);
+
+			int count = actions.Count();
+			Actions = new ActionState[count];
+
+			int index = 0;
+			foreach (LinkedAction action in actions)
+			{
+				action.SortIndex = index;
+				Actions[index++] = new ActionState
+				{
+					Action = action,
+					Status = ActionStatus.Queued,
+					PrereqActionsSortIndex = action.PrerequisiteActions.Select(a => a.SortIndex).ToArray()
+				};
+			}
+
+			// Cancel the queue when Ctrl-C is pressed
+			Console.CancelKeyPress += CancelKeyPress;
+		}
+
+		/// <summary>
+		/// Event handler for the Console.CancelKeyPress event
+		/// </summary>
+		/// <param name="sender"></param>
+		/// <param name="e"></param>
+		private void CancelKeyPress(object? sender, ConsoleCancelEventArgs e)
+		{
+			Console.CancelKeyPress -= CancelKeyPress;
+
+			if (!CancellationTokenSource.IsCancellationRequested)
+			{
+				Logger.LogWarning("Canceling actions...");
+				CancellationTokenSource.Cancel();
+				e.Cancel = true;
+			}
+
+			// We must do this and can't rely on that there are active processes that are cancelled causing a cascading cancel (force remote actions and no remote workers)
+			int completedActions = 0;
+			lock (Actions)
+			{
+				for (int actionIndex = _firstPendingAction; actionIndex != Actions.Length; ++actionIndex)
+				{
+					if (Actions[actionIndex].Status == ActionStatus.Queued)
+					{
+						Actions[actionIndex].Status = ActionStatus.Error;
+						++completedActions;
+					}
+				}
+			}
+			AddCompletedActions(completedActions);
+		}
+
+		/// <summary>
+		/// Create a new automatic runner
+		/// </summary>
+		/// <param name="runAction">Action to be run queue running a task</param>
+		/// <param name="useActionWeights">If true, use the action weight as a secondary limit</param>
+		/// <param name="maxActions">Maximum number of action actions</param>
+		/// <param name="maxActionWeight">Maximum weight of actions</param>
+		/// <returns>Created runner</returns>
+		public ImmediateActionQueueRunner CreateAutomaticRunner(Func<LinkedAction, Func<Task>?> runAction, bool useActionWeights, int maxActions, double maxActionWeight)
+		{
+			ImmediateActionQueueRunner runner = new(ImmediateActionQueueRunnerType.Automatic, runAction, useActionWeights, maxActions, maxActionWeight);
+			_runners.Add(runner);
+			return runner;
+		}
+
+		/// <summary>
+		/// Create a manual runner
+		/// </summary>
+		/// <param name="runAction">Action to be run queue running a task</param>
+		/// <param name="useActionWeights">If true, use the action weight as a secondary limit</param>
+		/// <param name="maxActions">Maximum number of action actions</param>
+		/// <param name="maxActionWeight">Maximum weight of actions</param>
+		/// <returns>Created runner</returns>
+		public ImmediateActionQueueRunner CreateManualRunner(Func<LinkedAction, Func<Task>?> runAction, bool useActionWeights = false, int maxActions = Int32.MaxValue, double maxActionWeight = Double.MaxValue)
+		{
+			ImmediateActionQueueRunner runner = new(ImmediateActionQueueRunnerType.Manual, runAction, useActionWeights, maxActions, maxActionWeight);
+			_runners.Add(runner);
+			return runner;
+		}
+
+		/// <summary>
+		/// Start the process of running all the actions
+		/// </summary>
+		public void Start()
+		{
+			if (ShowCPUUtilization || ActionStallReportTime > 0)
+			{
+				_lastActionChangeTime = DateTime.Now;
+				_cpuUtilizationTimer = new(x =>
+				{
+					if (ShowCPUUtilization)
+					{
+						lock (_cpuUtilization)
+						{
+							if (Utils.GetTotalCpuUtilization(out float cpuUtilization))
+							{
+								_cpuUtilization.Add(cpuUtilization);
+							}
+						}
+					}
+
+					if (ActionStallReportTime > 0 || ActionStallTerminateTime > 0)
+					{
+						lock (Actions)
+						{
+
+							// If there has been an action count change, reset the timer and enable the report again
+							if (_lastActionStallChange != _lastActionChange)
+							{
+								_lastActionStallChange = _lastActionChange;
+								_lastActionChangeTime = DateTime.Now;
+								_lastActionStallReported = false;
+							}
+
+							// Otherwise, if we haven't already generated a report, test for a timeout in seconds and generate one on timeout.
+							else if (ActionStallReportTime > 0 && !_lastActionStallReported && (DateTime.Now - _lastActionChangeTime).TotalSeconds > ActionStallReportTime)
+							{
+								_lastActionStallReported = true;
+								GenerateStallReport();
+							}
+							else if (ActionStallTerminateTime > 0 && !_lastActionStallCanceled && (DateTime.Now - _lastActionChangeTime).TotalSeconds > ActionStallTerminateTime)
+							{
+								_lastActionStallCanceled = true;
+								CancelStalledActions();
+							}
+						}
+					}
+				}, null, 1000, 1000);
+			}
+
+			Logger.LogInformation("------ Building {TotalActions} action(s) started ------", TotalActions);
+		}
+
+		/// <summary>
+		/// Run the actions until complete
+		/// </summary>
+		/// <returns>True if completed successfully</returns>
+		public async Task<bool> RunTillDone()
+		{
+			await _doneTaskSource.Task;
+			TraceSummary(Success);
+			return Success;
+		}
+
+		/// <summary>
+		/// Get the number of actions that were run, and how many succeeded or failed
+		/// </summary>
+		/// <param name="totalActions">Out parameter, the total number of actions</param>
+		/// <param name="succeededActions">Out parameter, the number of successful actions</param>
+		/// <param name="failedActions">Out parameter, the number of failed actions</param>
+		public void GetActionResultCounts(out int totalActions, out int succeededActions, out int failedActions)
+		{
+			totalActions = Actions.Length;
+			succeededActions = Results.Where(x => x?.ExitCode == 0).Count();
+			failedActions = Results.Where(x => x != null && x.ExitCode != 0).Count();
+		}
+
+		/// <summary>
+		/// Return an enumeration of ready compile tasks.  This is not executed under a lock and 
+		/// does not modify the state of any actions.
+		/// </summary>
+		/// <returns>Enumerations of all ready to compile actions.</returns>
+		public IEnumerable<LinkedAction> EnumerateReadyToCompileActions()
+		{
+			for (int actionIndex = _firstPendingAction; actionIndex != Actions.Length; ++actionIndex)
+			{
+				ActionState actionState = Actions[actionIndex];
+				if (actionState.Status == ActionStatus.Queued &&
+					GetActionReadyState(actionState) == ActionReadyState.Ready)
+				{
+					yield return actionState.Action;
+				}
+			}
+		}
+
+		/// <summary>
+		/// Try to start one action
+		/// </summary>
+		/// <param name="runner">If specified, tasks will only be queued to this runner.  Otherwise all manual runners will be used.</param>
+		/// <returns>True if an action was run, false if not</returns>
+		public bool TryStartOneAction(ImmediateActionQueueRunner? runner = null)
+		{
+
+			// Loop until we have an action or no completed actions (i.e. no propagated errors)
+			for (; ; )
+			{
+				bool hasCanceled = false;
+
+				// Try to get the next action
+				Func<Task>? runAction = null;
+				LinkedAction? action = null;
+				int completedActions = 0;
+				lock (Actions)
+				{
+					hasCanceled = CancellationTokenSource.IsCancellationRequested;
+
+					// Don't bother if we have been canceled or the specified running has already failed to find
+					// anything at the given change number.
+					if (!hasCanceled && (runner == null || runner.LastActionChange != _lastActionChange))
+					{
+
+						// Try to get an action
+						(runAction, action, completedActions) = TryStartOneActionInternal(runner);
+
+						// If we have completed actions (i.e. error propagations), then increment the change counter
+						if (completedActions != 0)
+						{
+							_lastActionChange++;
+						}
+
+						// Otherwise if nothing was found, remember that we have already scanned at this change.
+						else if ((runAction == null || action == null) && runner != null)
+						{
+							runner.LastActionChange = _lastActionChange;
+						}
+					}
+
+					// If we have no completed actions and have nothing to run, then we need to check for a stall
+					if (completedActions == 0 && (runAction == null || action == null))
+					{
+						// We found nothing, check to see if we have no active running tasks and no manual runners.
+						// If we don't, then it is assumed we can't queue any more work.
+						bool prematureDone = true;
+						foreach (ImmediateActionQueueRunner tryRunner in _runners)
+						{
+							if ((tryRunner.Type == ImmediateActionQueueRunnerType.Manual && !hasCanceled) || tryRunner.ActiveActions != 0)
+							{
+								prematureDone = false;
+								break;
+							}
+						}
+						if (prematureDone)
+						{
+							AddCompletedActions(Int32.MaxValue);
+						}
+						return false;
+					}
+				}
+
+				// If we have an action, run it and account for any completed actions
+				if (runAction != null && action != null)
+				{
+					if (runner != null && runner.Type == ImmediateActionQueueRunnerType.Manual)
+					{
+						try
+						{
+							runAction().Wait(CancellationToken);
+						}
+						catch (Exception ex)
+						{
+							HandleException(action, ex);
+						}
+					}
+					else
+					{
+						Task.Factory.StartNew(() =>
+						{
+							try
+							{
+								runAction().Wait(CancellationToken);
+							}
+							catch (Exception ex)
+							{
+								HandleException(action, ex);
+							}
+						}, CancellationToken, TaskCreationOptions.LongRunning | TaskCreationOptions.PreferFairness, TaskScheduler.Default);
+					}
+					AddCompletedActions(completedActions);
+					return true;
+				}
+
+				// Otherwise add the completed actions and test again for the possibility that errors still need propagating
+				AddCompletedActions(completedActions);
+			}
+		}
+
+		/// <summary>
+		/// Handle an exception running a task
+		/// </summary>
+		/// <param name="action">Action that has been run</param>
+		/// <param name="ex">Thrown exception</param>
+		void HandleException(LinkedAction action, Exception ex)
+		{
+			if (ex is AggregateException aggregateEx)
+			{
+				if (aggregateEx.InnerException != null)
+				{
+					HandleException(action, aggregateEx.InnerException);
+				}
+				else
+				{
+					ExecuteResults results = new(new List<string>(), Int32.MaxValue, TimeSpan.Zero, TimeSpan.Zero, 0, action);
+					OnActionCompleted(action, false, results);
+				}
+			}
+			else if (ex is OperationCanceledException)
+			{
+				ExecuteResults results = new(new List<string>(), Int32.MaxValue, TimeSpan.Zero, TimeSpan.Zero, 0, action);
+				OnActionCompleted(action, false, results);
+			}
+			else
+			{
+				List<string> text = new()
+				{
+					ExceptionUtils.FormatException(ex),
+					ExceptionUtils.FormatExceptionDetails(ex),
+				};
+				ExecuteResults results = new(text, Int32.MaxValue, TimeSpan.Zero, TimeSpan.Zero, 0, action);
+				OnActionCompleted(action, false, results);
+			}
+		}
+
+		/// <summary>
+		/// Try to start one action.
+		/// </summary>
+		/// <param name="runner">If specified, tasks will only be queued to this runner.  Otherwise all manual runners will be used.</param>
+		/// <returns>Action to be run if something was queued and the number of completed actions</returns>
+		public (Func<Task>?, LinkedAction?, int) TryStartOneActionInternal(ImmediateActionQueueRunner? runner = null)
+		{
+			int completedActions = 0;
+
+			// If we are starting deeper in the action collection, then never set the first pending action location
+			bool foundFirstPending = false;
+
+			// Loop through the items
+			for (int actionIndex = _firstPendingAction; actionIndex != Actions.Length; ++actionIndex)
+			{
+
+				// If the action has already reached the compilation state, then just ignore
+				if (Actions[actionIndex].Status != ActionStatus.Queued)
+				{
+					continue;
+				}
+
+				// If needed, update the first valid slot for searching for actions to run
+				if (!foundFirstPending)
+				{
+					_firstPendingAction = actionIndex;
+					foundFirstPending = true;
+				}
+
+				// Based on the ready state, use this action or mark as an error
+				switch (GetActionReadyState(Actions[actionIndex]))
+				{
+					case ActionReadyState.NotReady:
+						break;
+
+					case ActionReadyState.Error:
+						Actions[actionIndex].Status = ActionStatus.Error;
+						completedActions++;
+						break;
+
+					case ActionReadyState.Ready:
+						ImmediateActionQueueRunner? selectedRunner = null;
+						Func<Task>? runAction = null;
+						if (runner != null)
+						{
+							if (runner.IsUnderLimits)
+							{
+								runAction = runner.RunAction(Actions[actionIndex].Action);
+								if (runAction != null)
+								{
+									selectedRunner = runner;
+								}
+							}
+						}
+						else
+						{
+							foreach (ImmediateActionQueueRunner tryRunner in _runners)
+							{
+								if (tryRunner.Type == ImmediateActionQueueRunnerType.Automatic &&
+									tryRunner.IsUnderLimits)
+								{
+									runAction = tryRunner.RunAction(Actions[actionIndex].Action);
+									if (runAction != null)
+									{
+										selectedRunner = tryRunner;
+										break;
+									}
+								}
+							}
+						}
+
+						if (runAction != null && selectedRunner != null)
+						{
+							Actions[actionIndex].Status = ActionStatus.Running;
+							Actions[actionIndex].Runner = selectedRunner;
+							selectedRunner.ActiveActions++;
+							selectedRunner.ActiveActionWeight += Actions[actionIndex].Action.Weight;
+							return (runAction, Actions[actionIndex].Action, completedActions);
+						}
+						break;
+
+					default:
+						throw new ApplicationException("Unexpected action ready state");
+				}
+			}
+			return (null, null, completedActions);
+		}
+
+		/// <summary>
+		/// Start as many actions as possible.
+		/// </summary>
+		/// <param name="runner">If specified, all actions will be limited to the runner</param>
+		public void StartManyActions(ImmediateActionQueueRunner? runner = null)
+		{
+			const int Running = 1 << 1;
+			const int ScanRequested = 1 << 0;
+
+			// If both flags were clear, I need to start running actions
+			int old = Interlocked.Or(ref _startManyFlags, Running | ScanRequested);
+			if (old == 0)
+			{
+				for (; ; )
+				{
+
+					// Clear the changed flag since we are about to scan
+					Interlocked.And(ref _startManyFlags, Running);
+
+					// If nothing started
+					if (!TryStartOneAction(runner))
+					{
+
+						// If we only have the running flag (nothing new changed), then exit
+						if (Interlocked.CompareExchange(ref _startManyFlags, 0, Running) == Running)
+						{
+							return;
+						}
+					}
+				}
+			}
+		}
+
+		protected override void Dispose(bool disposing)
+		{
+			base.Dispose(disposing);
+
+			if (!_disposedValue)
+			{
+				if (disposing)
+				{
+					Console.CancelKeyPress -= CancelKeyPress;
+					_cpuUtilizationTimer?.Dispose();
+					_cpuUtilizationTimer = null;
+					CancellationTokenSource.Dispose();
+				}
+
+				_disposedValue = true;
+			}
+		}
+
+		/// <summary>
+		/// Reset the status of the given action to the queued state.
+		/// </summary>
+		/// <param name="action">Action being re-queued</param>
+		public void RequeueAction(LinkedAction action)
+		{
+			SetActionState(action, ActionStatus.Queued, null);
+		}
+
+		/// <summary>
+		/// Notify the system that an action has been completed.
+		/// </summary>
+		/// <param name="action">Action being completed</param>
+		/// <param name="success">If true, the action succeeded or false if it failed.</param>
+		/// <param name="results">Optional execution results</param>
+		public void OnActionCompleted(LinkedAction action, bool success, ExecuteResults? results)
+		{
+			// Delete produced items on error if requested
+			if (action.bDeleteProducedItemsOnError && (results?.ExitCode ?? 0) != 0)
+			{
+				foreach (FileItem output in action.ProducedItems.Distinct().Where(x => FileReference.Exists(x.Location)))
+				{
+					FileReference.Delete(output.Location);
+				}
+			}
+
+			// On success invalidate the file info for all the produced items.  This needs to be
+			// done so the artifact system sees any compiled files.
+			if (success)
+			{
+				foreach (FileItem output in action.ProducedItems)
+				{
+					output.ResetCachedInfo();
+				}
+			}
+
+			SetActionState(action, success ? ActionStatus.Finished : ActionStatus.Error, results);
+		}
+
+		/// <summary>
+		/// Set the new state of an action.  Can only set state to Queued, Finished, or Error
+		/// </summary>
+		/// <param name="action">Action being set</param>
+		/// <param name="status">Status to set</param>
+		/// <param name="results">Optional results</param>
+		/// <exception cref="BuildException"></exception>
+		private void SetActionState(LinkedAction action, ActionStatus status, ExecuteResults? results)
+		{
+			int actionIndex = action.SortIndex;
+			int completedActions = 0;
+
+			// Update the actions data
+			lock (Actions)
+			{
+				ImmediateActionQueueRunner runner = Actions[actionIndex].Runner ?? throw new BuildException("Attempting to update action state but runner isn't set");
+				runner.ActiveActions--;
+				runner.ActiveActionWeight -= Actions[actionIndex].Action.Weight;
+
+				// Use to track that action states have changed
+				_lastActionChange++;
+
+				Actions[actionIndex].Status = status;
+				Actions[actionIndex].Runner = null;
+				if (results != null)
+				{
+					Results[actionIndex] = results;
+				}
+
+				// Add the action to the logging queue
+				if (status != ActionStatus.Queued)
+				{
+					lock (_actionsToLog)
+					{
+						_actionsToLog.Add(action.SortIndex);
+						_actionsToLogTask ??= Task.Run(LogActions);
+					}
+				}
+
+				switch (status)
+				{
+					case ActionStatus.Queued:
+						if (actionIndex < _firstPendingAction)
+						{
+							_firstPendingAction = actionIndex;
+						}
+						break;
+
+					case ActionStatus.Finished:
+						completedActions = 1;
+						break;
+
+					case ActionStatus.Error:
+						Success = false;
+						completedActions = 1;
+						break;
+
+					default:
+						throw new BuildException("Unexpected action status set");
+				}
+			}
+
+			// Outside of lock, update the completed actions
+			AddCompletedActions(completedActions);
+
+			// Since something has been completed or returned to the queue, try to run actions again
+			StartManyActions();
+		}
+
+		/// <summary>
+		/// Add the number of completed actions and signal done if all actions complete.
+		/// </summary>
+		/// <param name="count">Number of completed actions to add.  If int.MaxValue, then the number is immediately set to total actions</param>
+		private void AddCompletedActions(int count)
+		{
+			if (count == 0)
+			{
+				// do nothing
+			}
+			else if (count != Int32.MaxValue)
+			{
+				if (Interlocked.Add(ref _completedActions, count) == TotalActions)
+				{
+					_doneTaskSource.SetResult();
+				}
+			}
+			else
+			{
+				if (Interlocked.Exchange(ref _completedActions, TotalActions) != TotalActions)
+				{
+					_doneTaskSource.SetResult();
+				}
+			}
+		}
+
+		/// <summary>
+		/// Returns the number of queued actions left (including queued actions that will do artifact check first)
+		/// Note, this method is lockless and will not always return accurate count
+		/// </summary>
+		/// <param name="filterFunc">Optional function to filter out actions. Return false if action should not be included</param>
+		public uint GetQueuedActionsCount(Func<LinkedAction, bool>? filterFunc = null)
+		{
+			uint count = 0;
+
+			for (int actionIndex = _firstPendingAction; actionIndex != Actions.Length; ++actionIndex)
+			{
+
+				if (Actions[actionIndex].Status != ActionStatus.Queued)
+				{
+					continue;
+				}
+
+				if (filterFunc != null && !filterFunc(Actions[actionIndex].Action))
+				{
+					continue;
+				}
+
+				++count;
+			}
+			return count;
+		}
+
+		/// <summary>
+		/// Purge the pending logging actions
+		/// </summary>
+		private void LogActions()
+		{
+			int numErrors = 0;
+			for (; ; )
+			{
+				int[]? actionsToLog = null;
+				lock (_actionsToLog)
+				{
+					if (_actionsToLog.Count == 0)
+					{
+						_actionsToLogTask = null;
+					}
+					else
+					{
+						actionsToLog = [.. _actionsToLog];
+						_actionsToLog.Clear();
+					}
+				}
+
+				if (actionsToLog == null)
+				{
+					return;
+				}
+
+				foreach (int index in actionsToLog)
+				{
+					ExecuteResults? result = Results[index];
+					LogAction(Actions[index].Action, result, Actions.Length);
+
+					// Cancel all other pending tasks
+					if (result == null || result.ExitCode != 0)
+					{
+						numErrors++;
+						if (StopCompilationAfterErrors || numErrors >= StopCompilationAfterNumErrors)
+						{
+							CancellationTokenSource.Cancel();
+						}
+					}
+				}
+			}
+		}
+
+		private enum ActionReadyState
+		{
+			NotReady,
+			Error,
+			Ready,
+		}
+
+		/// <summary>
+		/// Get the ready state of an action
+		/// </summary>
+		/// <param name="action">Action in question</param>
+		/// <returns>Action ready state</returns>
+		private ActionReadyState GetActionReadyState(ActionState action)
+		{
+			foreach (int prereqIndex in action.PrereqActionsSortIndex)
+			{
+				// Respect the compile status of the action
+				switch (Actions[prereqIndex].Status)
+				{
+					case ActionStatus.Finished:
+						continue;
+
+					case ActionStatus.Error:
+						return ActionReadyState.Error;
+
+					default:
+						return ActionReadyState.NotReady;
+				}
+			}
+			return ActionReadyState.Ready;
+		}
+
+		private void GenerateStallReport()
+		{
+			Logger.LogInformation("Action stall detected:");
+			foreach (ImmediateActionQueueRunner runner in _runners)
+			{
+				Logger.LogInformation("Runner Type: {Type}, Running Actions: {ActionCount}", runner.Type.ToString(), runner.ActiveActions);
+				if (runner.ActiveActions > 0)
+				{
+					int count = 0;
+					foreach (ActionState state in Actions)
+					{
+						if (state.Runner == runner && state.Status == ActionStatus.Running)
+						{
+							string description = $"{(state.Action.CommandDescription ?? state.Action.CommandPath.GetFileNameWithoutExtension())} {state.Action.StatusDescription}".Trim();
+							Logger.LogInformation("    Action[{Index}]: {Description}", count++, description);
+						}
+					}
+				}
+			}
+			{
+				int queued = 0;
+				int running = 0;
+				int error = 0;
+				int finished = 0;
+				foreach (ActionState state in Actions)
+				{
+					switch (state.Status)
+					{
+						case ActionStatus.Error:
+							error++;
+							break;
+						case ActionStatus.Finished:
+							finished++;
+							break;
+						case ActionStatus.Queued:
+							queued++;
+							break;
+						case ActionStatus.Running:
+							running++;
+							break;
+					}
+				}
+				Logger.LogInformation("Queue Counts: Queued = {Queued}, Running = {Running}, Finished = {Finished}, Error = {Error}", queued, running, finished, error);
+			}
+		}
+
+		private void CancelStalledActions()
+		{
+			GenerateStallReport();
+			Logger.LogInformation("Action stall terminate time exceeded, canceling remaining actions...");
+			CancellationTokenSource.Cancel();
+		}
+	}
+}
diff --git a/Engine/Source/Programs/UnrealBuildTool/Executors/ParallelExecutor.cs b/Engine/Source/Programs/UnrealBuildTool/Executors/ParallelExecutor.cs
new file mode 100644
--- /dev/null
+++ b/Engine/Source/Programs/UnrealBuildTool/Executors/ParallelExecutor.cs
@@ -0,0 +1,288 @@
+// Copyright Epic Games, Inc. All Rights Reserved.
+
+using System;
+using System.Collections.Generic;
+using System.Diagnostics;
+using System.IO;
+using System.Linq;
+using System.Threading;
+using System.Threading.Tasks;
+using EpicGames.Core;
+using Microsoft.Extensions.Logging;
+using UnrealBuildBase;
+
+namespace UnrealBuildTool
+{
+
+	/// <summary>
+	/// This executor uses async Tasks to process the action graph
+	/// </summary>
+	class ParallelExecutor : ActionExecutor
+	{
+#pragma warning disable IDE0044 // Make field readonly - these private static fields are set by XML configuration.
+		/// <summary>
+		/// Maximum processor count for local execution. 
+		/// </summary>
+		[XmlConfigFile]
+		[Obsolete("ParallelExecutor.MaxProcessorCount is deprecated. Please update xml to use BuildConfiguration.MaxParallelActions")]
+#pragma warning disable 0169
+#pragma warning disable CA1823 // Avoid unused private fields - part of the public API
+#pragma warning disable IDE0051 // Remove unused private member
+		private static int MaxProcessorCount;
+#pragma warning restore 0169
+#pragma warning restore CA1823
+#pragma warning restore IDE0051
+
+		/// <summary>
+		/// Processor count multiplier for local execution. Can be below 1 to reserve CPU for other tasks.
+		/// When using the local executor (not XGE), run a single action on each CPU core. Note that you can set this to a larger value
+		/// to get slightly faster build times in many cases, but your computer's responsiveness during compiling may be much worse.
+		/// This value is ignored if the CPU does not support hyper-threading.
+		/// </summary>
+		[XmlConfigFile]
+		private static double ProcessorCountMultiplier = 1.0;
+
+		/// <summary>
+		/// Free memory per action in bytes, used to limit the number of parallel actions if the machine is memory starved.
+		/// Set to 0 to disable free memory checking.
+		/// </summary>
+		[XmlConfigFile]
+		private static double MemoryPerActionBytes = 1.5 * 1024 * 1024 * 1024;
+#pragma warning restore IDE0044
+
+		/// <summary>
+		/// The priority to set for spawned processes.
+		/// Valid Settings: Idle, BelowNormal, Normal, AboveNormal, High
+		/// Default: BelowNormal or Normal for an Asymmetrical processor as BelowNormal can cause scheduling issues.
+		/// </summary>
+		[XmlConfigFile]
+		protected static ProcessPriorityClass ProcessPriority = Utils.IsAsymmetricalProcessor() ? ProcessPriorityClass.Normal : ProcessPriorityClass.BelowNormal;
+
+		/// <summary>
+		/// When enabled, will stop compiling targets after a compile error occurs.
+		/// </summary>
+		[XmlConfigFile]
+		protected static bool bStopCompilationAfterErrors = false;
+
+		/// <summary>
+		/// When set, will stop compiling targets after a set number of errors occurs.
+		/// </summary>
+		[XmlConfigFile]
+		protected static int StopCompilationAfterNumErrors = Int32.MaxValue;
+
+		/// <summary>
+		/// Whether to show compilation times along with worst offenders or not.
+		/// </summary>
+		[XmlConfigFile]
+		protected static bool bShowCompilationTimes = Unreal.IsBuildMachine();
+
+		/// <summary>
+		/// Whether to show compilation times for each executed action
+		/// </summary>
+		[XmlConfigFile]
+		protected static bool bShowPerActionCompilationTimes = Unreal.IsBuildMachine();
+
+		/// <summary>
+		/// Whether to log command lines for actions being executed
+		/// </summary>
+		[XmlConfigFile]
+		protected static bool bLogActionCommandLines = false;
+
+		/// <summary>
+		/// Add target names for each action executed
+		/// </summary>
+		[XmlConfigFile]
+		protected static bool bPrintActionTargetNames = false;
+
+		/// <summary>
+		/// Whether to take into account the Action's weight when determining to do more work or not.
+		/// </summary>
+		[XmlConfigFile]
+		protected static bool bUseActionWeights = false;
+
+		/// <summary>
+		/// Whether to show CPU utilization after the work is complete.
+		/// </summary>
+		[XmlConfigFile]
+		protected static bool bShowCPUUtilization = Unreal.IsBuildMachine();
+
+		/// <summary>
+		/// Collapse non-error output lines
+		/// </summary>
+		protected bool bCompactOutput = false;
+
+		/// <summary>
+		/// How many processes that will be executed in parallel
+		/// </summary>
+		public int NumParallelProcesses { get; private set; }
+
+		private static readonly char[] LineEndingSplit = new char[] { '\n', '\r' };
+
+		public static int GetDefaultNumParallelProcesses(int MaxLocalActions, bool bAllCores, ILogger Logger)
+		{
+			double MemoryPerActionBytesComputed = Math.Max(MemoryPerActionBytes, MemoryPerActionBytesOverride);
+			if (MemoryPerActionBytesComputed > MemoryPerActionBytes)
+			{
+				Logger.LogInformation("Overriding MemoryPerAction with target-defined value of {Memory} bytes", MemoryPerActionBytesComputed / 1024 / 1024 / 1024);
+			}
+
+			return Utils.GetMaxActionsToExecuteInParallel(MaxLocalActions, bAllCores ? 1.0f : ProcessorCountMultiplier, bAllCores, Convert.ToInt64(MemoryPerActionBytesComputed));
+		}
+
+		/// <summary>
+		/// Constructor
+		/// </summary>
+		/// <param name="MaxLocalActions">How many actions to execute in parallel</param>
+		/// <param name="bAllCores">Consider logical cores when determining how many total cpu cores are available</param>
+		/// <param name="bCompactOutput">Should output be written in a compact fashion</param>
+		/// <param name="Logger">Logger for output</param>
+		public ParallelExecutor(int MaxLocalActions, bool bAllCores, bool bCompactOutput, ILogger Logger)
+			: base(Logger)
+		{
+			XmlConfig.ApplyTo(this);
+
+			// Figure out how many processors to use
+			NumParallelProcesses = GetDefaultNumParallelProcesses(MaxLocalActions, bAllCores, Logger);
+
+			this.bCompactOutput = bCompactOutput;
+		}
+
+		/// <summary>
+		/// Returns the name of this executor
+		/// </summary>
+		public override string Name => "Parallel";
+
+		/// <summary>
+		/// Checks whether the task executor can be used
+		/// </summary>
+		/// <returns>True if the task executor can be used</returns>
+		public static bool IsAvailable()
+		{
+			return true;
+		}
+
+		/// <summary>
+		/// Telemetry event for this executor
+		/// </summary>
+		protected TelemetryExecutorEvent? telemetryEvent;
+
+		/// <inheritdoc/>
+		public override TelemetryExecutorEvent? GetTelemetryEvent() => telemetryEvent;
+
+		/// <summary>
+		/// Create an action queue
+		/// </summary>
+		/// <param name="actionsToExecute">Actions to be executed</param>
+		/// <param name="logger">Logging interface</param>
+		/// <returns>Action queue</returns>
+		public ImmediateActionQueue CreateActionQueue(IEnumerable<LinkedAction> actionsToExecute, ILogger logger)
+		{
+			return new(actionsToExecute, "Compiling C++ source code...", x => WriteToolOutput(x), () => FlushToolOutput(), logger)
+			{
+				ShowCompilationTimes = bShowCompilationTimes,
+				ShowCPUUtilization = bShowCPUUtilization,
+				PrintActionTargetNames = bPrintActionTargetNames,
+				LogActionCommandLines = bLogActionCommandLines,
+				ShowPerActionCompilationTimes = bShowPerActionCompilationTimes,
+				CompactOutput = bCompactOutput,
+				StopCompilationAfterErrors = bStopCompilationAfterErrors,
+				StopCompilationAfterNumErrors = StopCompilationAfterNumErrors,
+			};
+		}
+
+		/// <inheritdoc/>
+		public override async Task<bool> ExecuteActionsAsync(IEnumerable<LinkedAction> ActionsToExecute, ILogger Logger)
+		{
+			if (!ActionsToExecute.Any())
+			{
+				return true;
+			}
+
+			DateTime startTimeUTC = DateTime.UtcNow;
+			bool result;
+
+			// The "useAutomaticQueue" should always be true unless manual queue is being tested
+			bool useAutomaticQueue = true;
+			if (useAutomaticQueue)
+			{
+				using ImmediateActionQueue queue = CreateActionQueue(ActionsToExecute, Logger);
+				int actionLimit = Math.Min(NumParallelProcesses, queue.TotalActions);
+				queue.CreateAutomaticRunner(action => RunAction(queue, action), bUseActionWeights, actionLimit, NumParallelProcesses);
+				queue.Start();
+				queue.StartManyActions();
+				result = await queue.RunTillDone();
+
+				queue.GetActionResultCounts(out int totalActions, out int succeededActions, out int failedActions);
+				telemetryEvent = new TelemetryExecutorEvent(Name, startTimeUTC, result, totalActions, succeededActions, failedActions, cacheHitActions: 0, cacheMissActions: 0, DateTime.UtcNow, queue.ProcessGroup.PeakProcessMemoryUsed, queue.ProcessGroup.PeakJobMemoryUsed);
+			}
+			else
+			{
+				using ImmediateActionQueue queue = CreateActionQueue(ActionsToExecute, Logger);
+				int actionLimit = Math.Min(NumParallelProcesses, queue.TotalActions);
+				ImmediateActionQueueRunner runner = queue.CreateManualRunner(action => RunAction(queue, action), bUseActionWeights, actionLimit, actionLimit);
+				queue.Start();
+				using Timer timer = new((_) => queue.StartManyActions(runner), null, 0, 500);
+				queue.StartManyActions();
+				result = await queue.RunTillDone();
+
+				queue.GetActionResultCounts(out int totalActions, out int succeededActions, out int failedActions);
+				telemetryEvent = new TelemetryExecutorEvent(Name, startTimeUTC, result, totalActions, succeededActions, failedActions, cacheHitActions: 0, cacheMissActions: 0, DateTime.UtcNow, queue.ProcessGroup.PeakProcessMemoryUsed, queue.ProcessGroup.PeakJobMemoryUsed);
+			}
+
+			return result;
+		}
+
+		private static Func<Task>? RunAction(ImmediateActionQueue queue, LinkedAction action)
+		{
+			return async () =>
+			{
+				ExecuteResults results = await RunActionAsync(action, queue.ProcessGroup, queue.CancellationToken);
+				queue.OnActionCompleted(action, results.ExitCode == 0, results);
+			};
+		}
+
+		protected static async Task<ExecuteResults> RunActionAsync(LinkedAction Action, ManagedProcessGroup ProcessGroup, CancellationToken CancellationToken, string? AdditionalDescription = null)
+		{
+			CancellationToken.ThrowIfCancellationRequested();
+
+			using ManagedProcess Process = new ManagedProcess(ProcessGroup, Action.CommandPath.FullName, Action.CommandArguments, Action.WorkingDirectory.FullName, null, null, ProcessPriority);
+
+			using MemoryStream StdOutStream = new MemoryStream();
+			await Process.CopyToAsync(StdOutStream, CancellationToken);
+
+			CancellationToken.ThrowIfCancellationRequested();
+
+			await Process.WaitForExitAsync(CancellationToken);
+
+			List<string> LogLines = [.. Console.OutputEncoding.GetString(StdOutStream.GetBuffer(), 0, Convert.ToInt32(StdOutStream.Length)).Split(LineEndingSplit, StringSplitOptions.RemoveEmptyEntries)];
+			int ExitCode = Process.ExitCode;
+			TimeSpan ProcessorTime = Process.TotalProcessorTime;
+			TimeSpan ExecutionTime = Process.ExitTime - Process.StartTime;
+			long PeakMemoryUsed = Process.PeakMemoryUsed;
+
+			if (ExitCode == 0 && Action.bForceWarningsAsError && LogLines.Any(x => x.Contains(": warning: ", StringComparison.OrdinalIgnoreCase)))
+			{
+				ExitCode = 1;
+				LogLines = [.. LogLines.Select(x => x.Replace(": warning: ", ": error: ", StringComparison.OrdinalIgnoreCase))];
+			}
+
+			return new ExecuteResults(LogLines, ExitCode, ExecutionTime, ProcessorTime, PeakMemoryUsed, Action, AdditionalDescription);
+		}
+	}
+
+	/// <summary>
+	/// Publicly visible static class that allows external access to the parallel executor config
+	/// </summary>
+	public static class ParallelExecutorConfiguration
+	{
+		/// <summary>
+		/// Maximum number of processes that should be used for execution
+		/// </summary>
+		public static int GetMaxParallelProcesses(ILogger Logger) => ParallelExecutor.GetDefaultNumParallelProcesses(0, false, Logger);
+
+		/// <summary>
+		/// Maximum number of processes that should be used for execution
+		/// </summary>
+		public static int GetMaxParallelProcesses(int MaxLocalActions, bool bAllCores, ILogger Logger) => ParallelExecutor.GetDefaultNumParallelProcesses(MaxLocalActions, bAllCores, Logger);
+	}
+}
diff --git a/Engine/Source/Programs/UnrealBuildTool/Executors/SNDBS.cs b/Engine/Source/Programs/UnrealBuildTool/Executors/SNDBS.cs
--- a/Engine/Source/Programs/UnrealBuildTool/Executors/SNDBS.cs
+++ b/Engine/Source/Programs/UnrealBuildTool/Executors/SNDBS.cs
@@ -22,9 +22,7 @@ namespace UnrealBuildTool
 {
 	sealed class SNDBS : ActionExecutor
 	{
-		internal const string ExecutorName = "SNDBS";
-
-		public override string Name => ExecutorName;
+		public override string Name => "SNDBS";
 
 		private static readonly string? ProgramFilesx86 = Environment.GetEnvironmentVariable("ProgramFiles(x86)");
 		private static readonly string? SCERoot = Environment.GetEnvironmentVariable("SCE_ROOT_DIR");
diff --git a/Engine/Source/Programs/UnrealBuildTool/Executors/UnrealBuildAccelerator/UBAExecutor.cs b/Engine/Source/Programs/UnrealBuildTool/Executors/UnrealBuildAccelerator/UBAExecutor.cs
--- a/Engine/Source/Programs/UnrealBuildTool/Executors/UnrealBuildAccelerator/UBAExecutor.cs
+++ b/Engine/Source/Programs/UnrealBuildTool/Executors/UnrealBuildAccelerator/UBAExecutor.cs
@@ -17,16 +17,17 @@ using ILogger = Microsoft.Extensions.Logging.ILogger;
 
 namespace UnrealBuildTool
 {
-	class UBAExecutor : ActionExecutor
+	class UBAExecutor : ParallelExecutor
 	{
-		public UnrealBuildAcceleratorConfig UBAConfig { get; init; }
+		public UnrealBuildAcceleratorConfig UBAConfig { get; init; } = new();
 
 		public string Crypto { get; private set; } = String.Empty;
 
+		public bool CanRetryNatively { get; set; } = true;
+
 		readonly IEnumerable<TargetDescriptor> _targetDescriptors;
-		readonly Task _initTask;
+		readonly Task? _initTask;
 		readonly CancellationTokenSource _cancellationSource;
-		readonly bool _bCompactOutput;
 
 		IServer? _server;
 		IStorageServer? _storageServer;
@@ -56,16 +57,11 @@ namespace UnrealBuildTool
 		// Tracking for remote connection mode
 		string _remoteConnectionMode = "Local";
 
-		/// <summary>
-		/// Telemetry event for this executor
-		/// </summary>
-		private TelemetryExecutorEvent? telemetryEvent;
-
 		protected override void Dispose(bool disposing)
 		{
 			if (disposing)
 			{
-				_initTask.Wait();
+				_initTask?.Wait();
 
 				{
 					using ITimelineEvent a2 = Timeline.ScopeEvent("DisconnectCache");
@@ -98,13 +94,9 @@ namespace UnrealBuildTool
 			}
 		}
 
-		internal const string ExecutorName = "UBA";
-
 		public override string Name => !UBAConfig.bDisableRemote && _agentCoordinators.Count > 0 ? "Unreal Build Accelerator" : "Unreal Build Accelerator local";
 
-		public int NumParallelProcesses { get; }
-
-		public static bool IsAvailable() => EpicGames.UBA.Utils.IsAvailable();
+		public static new bool IsAvailable() => EpicGames.UBA.Utils.IsAvailable();
 
 		public delegate void SessionCreatedDelegate(ISessionServer Session);
 
@@ -140,44 +132,45 @@ namespace UnrealBuildTool
 			}
 		}
 
-		internal static UnrealBuildAcceleratorConfig GetAppliedConfig(IEnumerable<TargetDescriptor> targetDescriptors)
+		public UBAExecutor(int maxLocalActions, bool bAllCores, bool bCompactOutput, ILogger logger, IEnumerable<TargetDescriptor> targetDescriptors, bool useInitTask)
+			: base(maxLocalActions, bAllCores, bCompactOutput, logger)
 		{
-			UnrealBuildAcceleratorConfig config = new();
-			XmlConfig.ApplyTo(config);
-
-			foreach (TargetDescriptor targetDescriptor in targetDescriptors)
-			{
-				ConfigCache.ReadSettings(targetDescriptor.ProjectFile?.Directory, BuildHostPlatform.Current.Platform, config);
-				targetDescriptor.AdditionalArguments.ApplyTo(config);
-			}
-
-			return config;
-		}
-
-		public UBAExecutor(UnrealBuildAcceleratorConfig config, int maxLocalActions, bool bAllCores, bool bCompactOutput, ILogger logger, IEnumerable<TargetDescriptor> targetDescriptors)
-			: base(logger)
-		{
-			UBAConfig = config;
-
-			// Figure out how many processors to use
-			NumParallelProcesses = GetDefaultNumParallelProcesses(config, maxLocalActions, bAllCores, logger);
-
-			_bCompactOutput = bCompactOutput;
-
 			_targetDescriptors = targetDescriptors;
 			_cancellationSource = new();
-			_initTask = new(() => { Init(targetDescriptors, logger, true); }, TaskCreationOptions.LongRunning);
-			_initTask.Start();
+			if (useInitTask)
+			{
+				_initTask = new(() =>
+				{
+					// We know the BuildMode is the only one using init tasks which means we can let process exit do cleanup
+					// instead of explicitly doing it with the executor shutdown
+					Init(targetDescriptors, logger, false);
+				}, TaskCreationOptions.LongRunning);
+
+				_initTask.Start();
+			}
 		}
 
 		private void Init(IEnumerable<TargetDescriptor> targetDescriptors, ILogger logger, bool cleanupOnShutdown)
 		{
 			XmlConfig.ApplyTo(this);
+			XmlConfig.ApplyTo(UBAConfig);
 			CommandLine.ParseArguments(Environment.GetCommandLineArgs(), this, logger);
 
 			foreach (TargetDescriptor targetDescriptor in targetDescriptors)
 			{
+				ConfigCache.ReadSettings(targetDescriptor.ProjectFile?.Directory, BuildHostPlatform.Current.Platform, UBAConfig);
 				targetDescriptor.AdditionalArguments.ApplyTo(this);
+				targetDescriptor.AdditionalArguments.ApplyTo(UBAConfig);
+			}
+
+			if (UBAConfig.bDisableRemote)
+			{
+				UBAConfig.Zone = "local";
+			}
+
+			if (UBAConfig.bForcedRetry)
+			{
+				UBAConfig.bAllowRetry = true;
 			}
 
 			_trace = ITrace.GlobalTrace;
@@ -387,17 +380,6 @@ namespace UnrealBuildTool
 			}
 		}
 
-		internal static int GetDefaultNumParallelProcesses(UnrealBuildAcceleratorConfig config, int maxLocalActions, bool bAllCores, ILogger logger)
-		{
-			double MemoryPerActionBytesComputed = Math.Max(config.MemoryPerActionBytes, MemoryPerActionBytesOverride);
-			if (MemoryPerActionBytesComputed > config.MemoryPerActionBytes)
-			{
-				logger.LogInformation("Overriding MemoryPerAction with target-defined value of {Memory} bytes", MemoryPerActionBytesComputed / 1024 / 1024 / 1024);
-			}
-
-			return Utils.GetMaxActionsToExecuteInParallel(maxLocalActions, bAllCores ? 1.0f : config.ProcessorCountMultiplier, bAllCores, Convert.ToInt64(MemoryPerActionBytesComputed));
-		}
-
 		public void UpdateStatus(uint statusRow, uint statusColumn, string statusText, LogEntryType statusType, string? statusLink)
 		{
 			_trace?.UpdateStatus(statusRow, statusColumn, statusText, statusType, statusLink);
@@ -513,9 +495,6 @@ namespace UnrealBuildTool
 			}
 		}
 
-		/// <inheritdoc/>
-		public override TelemetryExecutorEvent? GetTelemetryEvent() => telemetryEvent;
-
 		public override async Task<bool> ExecuteActionsAsync(IEnumerable<LinkedAction> inputActions, Microsoft.Extensions.Logging.ILogger logger)
 		{
 			// Disable cache server when using hot reload path (there would likely never be any cache hits anyway)
@@ -530,7 +509,14 @@ namespace UnrealBuildTool
 				_cancellationSource.Cancel();
 			}
 
-			_initTask.Wait();
+			if (_initTask == null)
+			{
+				Init(_targetDescriptors, logger, true);
+			}
+			else
+			{
+				_initTask?.Wait();
+			}
 
 			PrintConfiguration(logger);
 
@@ -743,12 +729,12 @@ namespace UnrealBuildTool
 
 			using ActionLogger actionLogger = new(inputActions.Count(), "Compiling C++ source code...", x => WriteToolOutput(x), () => FlushToolOutput(), logger)
 			{
-				ShowCompilationTimes = UBAConfig.bShowCompilationTimes,
-				ShowCPUUtilization = UBAConfig.bShowCPUUtilization,
-				PrintActionTargetNames = UBAConfig.bPrintActionTargetNames,
-				LogActionCommandLines = UBAConfig.bLogActionCommandLines,
-				ShowPerActionCompilationTimes = UBAConfig.bShowPerActionCompilationTimes,
-				CompactOutput = _bCompactOutput,
+				ShowCompilationTimes = bShowCompilationTimes,
+				ShowCPUUtilization = bShowCPUUtilization,
+				PrintActionTargetNames = bPrintActionTargetNames,
+				LogActionCommandLines = bLogActionCommandLines,
+				ShowPerActionCompilationTimes = bShowPerActionCompilationTimes,
+				CompactOutput = bCompactOutput,
 			};
 
 			using IScheduler scheduler = IScheduler.CreateScheduler(session, cache, NumParallelProcesses, UBAConfig.bForceBuildAllRemote);
@@ -860,7 +846,7 @@ namespace UnrealBuildTool
 							{
 								// If files are compressed (pch, obj etc), then native will not work so must retry detoured
 								// If vfs paths are used, then native will not work
-								bool mustRunDetoured = !UBAConfig.CanRetryNatively || UBAConfig.bStoreObjFilesCompressed || action.RootPaths.bUseVfs;
+								bool mustRunDetoured = !CanRetryNatively || UBAConfig.bStoreObjFilesCompressed || action.RootPaths.bUseVfs;
 
 								if (mustRunDetoured)
 								{
@@ -945,14 +931,14 @@ namespace UnrealBuildTool
 						}
 					}
 
-					if (UBAConfig.bStopCompilationAfterErrors || UBAConfig.StopCompilationAfterNumErrors <= 0)
+					if (bStopCompilationAfterErrors || StopCompilationAfterNumErrors <= 0)
 					{
 						cancelReason = "stop compilation after any errors";
 						shouldCancel = true;
 					}
-					else if (failedActions >= UBAConfig.StopCompilationAfterNumErrors)
+					else if (failedActions >= StopCompilationAfterNumErrors)
 					{
-						cancelReason = $"stop compilation after {UBAConfig.StopCompilationAfterNumErrors} errors";
+						cancelReason = $"stop compilation after {StopCompilationAfterNumErrors} errors";
 						shouldCancel = true;
 					}
 				}
@@ -1127,12 +1113,15 @@ namespace UnrealBuildTool
 			}
 
 			// Let pch compile at higher priority
-			System.Diagnostics.ProcessPriorityClass priority = UBAConfig.ProcessPriority;
-			if (priority == System.Diagnostics.ProcessPriorityClass.BelowNormal && action.StatusDescription.Contains("PCH."))
+			System.Diagnostics.ProcessPriorityClass priority = ProcessPriority;
+			if (ProcessPriority == System.Diagnostics.ProcessPriorityClass.BelowNormal)
 			{
+				if (action.StatusDescription.Contains("PCH."))
+				{
 					priority = System.Diagnostics.ProcessPriorityClass.Normal;
+				}
 			}
-
+				
 			ProcessStartInfo startInfo = new()
 			{
 				Application = action.CommandPath.FullName,
diff --git a/Engine/Source/Programs/UnrealBuildTool/Executors/UnrealBuildAccelerator/UnrealBuildAcceleratorConfig.cs b/Engine/Source/Programs/UnrealBuildTool/Executors/UnrealBuildAccelerator/UnrealBuildAcceleratorConfig.cs
--- a/Engine/Source/Programs/UnrealBuildTool/Executors/UnrealBuildAccelerator/UnrealBuildAcceleratorConfig.cs
+++ b/Engine/Source/Programs/UnrealBuildTool/Executors/UnrealBuildAccelerator/UnrealBuildAcceleratorConfig.cs
@@ -2,9 +2,7 @@
 
 using System;
 using System.Collections.Generic;
-using System.Diagnostics;
 using EpicGames.Core;
-using UnrealBuildBase;
 
 namespace UnrealBuildTool
 {
@@ -82,17 +80,12 @@ namespace UnrealBuildTool
 		[CommandLine("-UBAForceRemote")]
 		public bool bForceBuildAllRemote { get; set; } = false;
 
-		private bool _allowRetry = true;
 		/// <summary>
 		/// When set to true, actions that fail locally with certain error codes will retry without uba.
 		/// </summary>
 		[XmlConfigFile(Category = "UnrealBuildAccelerator")]
 		[CommandLine("-UBANoRetry", Value = "false")]
-		public bool bAllowRetry
-		{
-			get => bForcedRetry ? true : _allowRetry;
-			set => _allowRetry = value;
-		}
+		public bool bAllowRetry { get; set; } = true;
 
 		/// <summary>
 		/// When set to true, all actions that fail locally with UBA will be retried without UBA.
@@ -257,17 +250,12 @@ namespace UnrealBuildTool
 		[CommandLine("-UBANoCustomMalloc", Value = "true")]
 		public bool bDisableCustomAlloc { get; set; } = false;
 
-		private string _zone = String.Empty;
 		/// <summary>
 		/// The zone to use for UBA.
 		/// </summary>
 		[XmlConfigFile(Category = "UnrealBuildAccelerator")]
 		[CommandLine("-UBAZone=")]
-		public string Zone
-		{
-			get => bDisableRemote ? "local" : _zone;
-			set => _zone = value;
-		}
+		public string Zone { get; set; } = String.Empty;
 
 		/// <summary>
 		/// Set to true to enable encryption when transferring files over the network.
@@ -344,76 +332,5 @@ namespace UnrealBuildTool
 		[XmlConfigFile(Category = "UnrealBuildAccelerator")]
 		[CommandLine("-UBACompressionLevel")]
 		public string CompressionLevel { get; set; } = String.Empty;
-
-		/// <summary>
-		/// Processor count multiplier for local execution. Can be below 1 to reserve CPU for other tasks.
-		/// When using the local executor (not XGE), run a single action on each CPU core. Note that you can set this to a larger value
-		/// to get slightly faster build times in many cases, but your computer's responsiveness during compiling may be much worse.
-		/// This value is ignored if the CPU does not support hyper-threading.
-		/// </summary>
-		[XmlConfigFile(Category = "ParallelExecutor")]
-		public double ProcessorCountMultiplier { get; set; } = 1.0;
-
-		/// <summary>
-		/// Free memory per action in bytes, used to limit the number of parallel actions if the machine is memory starved.
-		/// Set to 0 to disable free memory checking.
-		/// </summary>
-		[XmlConfigFile(Category = "ParallelExecutor")]
-		public double MemoryPerActionBytes { get; set; } = 1.5 * 1024 * 1024 * 1024;
-
-		/// <summary>
-		/// The priority to set for spawned processes.
-		/// Valid Settings: Idle, BelowNormal, Normal, AboveNormal, High
-		/// Default: BelowNormal or Normal for an Asymmetrical processor as BelowNormal can cause scheduling issues.
-		/// </summary>
-		[XmlConfigFile(Category = "ParallelExecutor")]
-		public ProcessPriorityClass ProcessPriority { get; set; } = Utils.IsAsymmetricalProcessor() ? ProcessPriorityClass.Normal : ProcessPriorityClass.BelowNormal;
-
-		/// <summary>
-		/// When enabled, will stop compiling targets after a compile error occurs.
-		/// </summary>
-		[XmlConfigFile(Category = "ParallelExecutor")]
-		public bool bStopCompilationAfterErrors { get; set; } = false;
-
-		/// <summary>
-		/// When set, will stop compiling targets after a set number of errors occurs.
-		/// </summary>
-		[XmlConfigFile(Category = "ParallelExecutor")]
-		public int StopCompilationAfterNumErrors { get; set; } = Int32.MaxValue;
-
-		/// <summary>
-		/// Whether to show compilation times along with worst offenders or not.
-		/// </summary>
-		[XmlConfigFile(Category = "ParallelExecutor")]
-		public bool bShowCompilationTimes { get; set; } = Unreal.IsBuildMachine();
-
-		/// <summary>
-		/// Whether to show compilation times for each executed action
-		/// </summary>
-		[XmlConfigFile(Category = "ParallelExecutor")]
-		public bool bShowPerActionCompilationTimes { get; set; } = Unreal.IsBuildMachine();
-
-		/// <summary>
-		/// Whether to log command lines for actions being executed
-		/// </summary>
-		[XmlConfigFile(Category = "ParallelExecutor")]
-		public bool bLogActionCommandLines { get; set; } = false;
-
-		/// <summary>
-		/// Add target names for each action executed
-		/// </summary>
-		[XmlConfigFile(Category = "ParallelExecutor")]
-		public bool bPrintActionTargetNames { get; set; } = false;
-
-		/// <summary>
-		/// Whether to show CPU utilization after the work is complete.
-		/// </summary>
-		[XmlConfigFile(Category = "ParallelExecutor")]
-		public bool bShowCPUUtilization { get; set; } = Unreal.IsBuildMachine();
-
-		/// <summary>
-		/// Whether UBA can retry actions natively if they fail when detoured locally
-		/// </summary>
-		public bool CanRetryNatively { get; set; } = true;
 	}
 }
\ No newline at end of file
diff --git a/Engine/Source/Programs/UnrealBuildTool/Executors/XGE.cs b/Engine/Source/Programs/UnrealBuildTool/Executors/XGE.cs
--- a/Engine/Source/Programs/UnrealBuildTool/Executors/XGE.cs
+++ b/Engine/Source/Programs/UnrealBuildTool/Executors/XGE.cs
@@ -44,38 +44,38 @@ namespace UnrealBuildTool
 		/// <summary>
 		/// When set to false, XGE will not be enabled when running connected to the coordinator over VPN. Configure VPN-assigned subnets via the VpnSubnets parameter.
 		/// </summary>
-		[XmlConfigFile(Category = ExecutorName)]
+		[XmlConfigFile(Category = "XGE")]
 		static bool bAllowOverVpn = true;
 
 		/// <summary>
 		/// List of subnets containing IP addresses assigned by VPN
 		/// </summary>
-		[XmlConfigFile(Category = ExecutorName)]
+		[XmlConfigFile(Category = "XGE")]
 		static string[]? VpnSubnets = null;
 
 		/// <summary>
 		/// Whether to allow remote linking
 		/// </summary>
-		[XmlConfigFile(Category = ExecutorName)]
+		[XmlConfigFile(Category = "XGE")]
 		static bool bAllowRemoteLinking = false;
 
 		/// <summary>
 		/// Whether to enable the VCCompiler=true setting. This requires an additional license for VC tools. 
 		/// </summary>
-		[XmlConfigFile(Category = ExecutorName)]
+		[XmlConfigFile(Category = "XGE")]
 		static bool bUseVCCompilerMode = false;
 
 		/// <summary>
 		/// Minimum number of actions to use XGE execution.
 		/// </summary>
-		[XmlConfigFile(Category = ExecutorName)]
+		[XmlConfigFile(Category = "XGE")]
 		public static int MinActions = 2;
 
 		/// <summary>
 		/// Check for a concurrent XGE build and treat the XGE executor as unavailable if it's in use.
 		/// This will allow UBT to fall back to another executor such as the parallel executor. 
 		/// </summary>
-		[XmlConfigFile(Category = ExecutorName)]
+		[XmlConfigFile(Category = "XGE")]
 		static bool bUnavailableIfInUse = false;
 #pragma warning restore IDE0044
 
@@ -99,9 +99,7 @@ namespace UnrealBuildTool
 			XmlConfig.ApplyTo(this);
 		}
 
-		internal const string ExecutorName = "XGE";
-
-		public override string Name => ExecutorName;
+		public override string Name => "XGE";
 
 		public static bool TryGetXgConsoleExecutable([NotNullWhen(true)] out string? OutXgConsoleExe)
 		{
