In UBT, check for whether detouring is allowed to determine whether remote is disabled. This saves some initialization code from running unnecessarily.

#Synced-CL 48713665

diff --git a/Engine/Source/Programs/UnrealBuildTool/Executors/UnrealBuildAccelerator/UnrealBuildAcceleratorConfig.cs b/Engine/Source/Programs/UnrealBuildTool/Executors/UnrealBuildAccelerator/UnrealBuildAcceleratorConfig.cs
--- a/Engine/Source/Programs/UnrealBuildTool/Executors/UnrealBuildAccelerator/UnrealBuildAcceleratorConfig.cs
+++ b/Engine/Source/Programs/UnrealBuildTool/Executors/UnrealBuildAccelerator/UnrealBuildAcceleratorConfig.cs
@@ -66,12 +66,18 @@ namespace UnrealBuildTool
 		[ConfigFile(ConfigHierarchyType.Engine, "UbaController", "BuildMachineCacheProviders")]
 		public List<string> IniBuildMachineCacheProviders { get; set; } = [];
 
+		private bool _disableRemote = false;
 		/// <summary>
 		/// When set to true, UBA will not use any remote help
 		/// </summary>
 		[XmlConfigFile(Category = "UnrealBuildAccelerator")]
 		[CommandLine("-UBADisableRemote")]
-		public bool bDisableRemote { get; set; } = false;
+		public bool bDisableRemote
+		{
+			// Remote must be disabled if we don't allow detouring.
+			get => !bAllowDetour || _disableRemote;
+			set => _disableRemote = value;
+		}
 
 		/// <summary>
 		/// When set to true, UBA will force all actions that can be built remotely to be built remotely. This will hang if there are no remote agents available
