[UBT] Fix static analyzis warnings caused by CL 48746946

That CL moved legacy code into its own file, and made some properties and methods not used, but that we don't want to remove just yet.

This CL disables the warning on these. But also removed code that is indeed no longer moved and I missed it in my first pass.

#Synced-CL 48747215

diff --git a/Engine/Source/Programs/UnrealBuildTool/Platform/Apple/AppleExports.cs b/Engine/Source/Programs/UnrealBuildTool/Platform/Apple/AppleExports.cs
--- a/Engine/Source/Programs/UnrealBuildTool/Platform/Apple/AppleExports.cs
+++ b/Engine/Source/Programs/UnrealBuildTool/Platform/Apple/AppleExports.cs
@@ -15,10 +15,15 @@ namespace UnrealBuildTool
 	/// </summary>
 	public static class AppleExports
 	{
+#pragma warning disable CA1823 // unused private fields - This will be removed in a future version
+#pragma warning disable IDE0052 //  unused private member - This will be removed in a future version
 		[Obsolete("Modern Xcode is now the only option")]
 		private static readonly bool bForceModernXcode = Environment.CommandLine.Contains("-modernxcode", StringComparison.OrdinalIgnoreCase);
 		[Obsolete("Legacy Xcode is no longer supported")]
 		private static readonly bool bForceLegacyXcode = Environment.CommandLine.Contains("-legacy xcode", StringComparison.OrdinalIgnoreCase);
+#pragma warning restore IDE0052
+#pragma warning restore CA1823
+
 		private static readonly bool bNoEntitlements = Environment.CommandLine.Contains("-noEntitlements", StringComparison.OrdinalIgnoreCase);
 
 		private static DirectoryReference? _MobileProvisionDir;
diff --git a/Engine/Source/Programs/UnrealBuildTool/Platform/IOS/IOSToolChain.cs b/Engine/Source/Programs/UnrealBuildTool/Platform/IOS/IOSToolChain.cs
--- a/Engine/Source/Programs/UnrealBuildTool/Platform/IOS/IOSToolChain.cs
+++ b/Engine/Source/Programs/UnrealBuildTool/Platform/IOS/IOSToolChain.cs
@@ -5,9 +5,7 @@ using System.Collections.Generic;
 using System.Diagnostics;
 using System.IO;
 using System.Linq;
-using System.Security.Cryptography.X509Certificates;
 using System.Text;
-using System.Text.Json;
 using EpicGames.Core;
 using Ionic.Zip;
 using Ionic.Zlib;
diff --git a/Engine/Source/Programs/UnrealBuildTool/Platform/IOS/UEDeployIOS.cs b/Engine/Source/Programs/UnrealBuildTool/Platform/IOS/UEDeployIOS.cs
--- a/Engine/Source/Programs/UnrealBuildTool/Platform/IOS/UEDeployIOS.cs
+++ b/Engine/Source/Programs/UnrealBuildTool/Platform/IOS/UEDeployIOS.cs
@@ -29,8 +29,6 @@ namespace UnrealBuildTool
 		}
 		bool bForDistribution = false;
 
-		private static readonly DirectoryReference MobileProvisionDirRef = AppleExports.GetProvisionDirectory();
-
 		public class VersionUtilities
 		{
 			public static string? BuildDirectory
diff --git a/Engine/Source/Programs/UnrealBuildTool/Platform/Mac/MacToolChain.cs b/Engine/Source/Programs/UnrealBuildTool/Platform/Mac/MacToolChain.cs
--- a/Engine/Source/Programs/UnrealBuildTool/Platform/Mac/MacToolChain.cs
+++ b/Engine/Source/Programs/UnrealBuildTool/Platform/Mac/MacToolChain.cs
@@ -270,12 +270,6 @@ namespace UnrealBuildTool
 			Arguments.Add("-headerpad_max_install_names");
 		}
 
-		private void AppendMacLine(StreamWriter Writer, string Format, params object[] Arg)
-		{
-			string PreLine = String.Format(Format, Arg);
-			Writer.Write(PreLine + "\n");
-		}
-
 		public static string LoadEngineDisplayVersion(bool bIgnorePatchVersion = false)
 		{
 			BuildVersion? Version;
@@ -807,11 +801,6 @@ namespace UnrealBuildTool
 			return OutputFile;
 		}
 
-		static string FormatCopyCommand(string SourceFile, string TargetFile)
-		{
-			return String.Format("rsync --checksum \"{0}\" \"{1}\"", SourceFile, TargetFile);
-		}
-
 		/// <summary>
 		/// Generates debug info for a given executable
 		/// </summary>
@@ -895,59 +884,6 @@ namespace UnrealBuildTool
 			return OutputFile;
 		}
 
-		/// <summary>
-		/// Creates app bundle for a given executable
-		/// </summary>
-		/// <param name="LinkEnvironment"></param>
-		/// <param name="Executable">FileItem describing the executable to generate app bundle for</param>
-		/// <param name="Graph">List of actions to be executed. Additional actions will be added to this list.</param>
-		FileItem FinalizeAppBundle(LinkEnvironment LinkEnvironment, FileItem Executable, IActionGraphBuilder Graph)
-		{
-			// Make a file item for the source and destination files
-			string FullDestPath = Executable.AbsolutePath.Substring(0, Executable.AbsolutePath.IndexOf(".app") + 4);
-			FileItem DestFile = FileItem.GetItemByPath(FullDestPath);
-
-			// Make the compile action
-			Action FinalizeAppBundleAction = Graph.CreateAction(ActionType.CreateAppBundle);
-			FinalizeAppBundleAction.WorkingDirectory = GetMacDevSrcRoot(); // Path.GetFullPath(".");
-			FinalizeAppBundleAction.CommandPath = BuildHostPlatform.Current.Shell;
-			FinalizeAppBundleAction.CommandDescription = "";
-
-			// make path to the script
-			FileItem BundleScript = FileItem.GetItemByFileReference(FileReference.Combine(LinkEnvironment.IntermediateDirectory!, "FinalizeAppBundle.sh"));
-
-			FinalizeAppBundleAction.CommandArguments = "\"" + BundleScript.AbsolutePath + "\"";
-			FinalizeAppBundleAction.PrerequisiteItems.Add(Executable);
-			FinalizeAppBundleAction.ProducedItems.Add(DestFile);
-			FinalizeAppBundleAction.StatusDescription = String.Format("Finalizing app bundle: {0}.app", Path.GetFileName(Executable.AbsolutePath));
-			FinalizeAppBundleAction.bCanExecuteRemotely = false;
-
-			return DestFile;
-		}
-
-		FileItem CopyBundleResource(UEBuildBundleResource Resource, FileItem Executable, DirectoryReference BundleDirectory, IActionGraphBuilder Graph)
-		{
-			Action CopyAction = Graph.CreateAction(ActionType.CreateAppBundle);
-			CopyAction.WorkingDirectory = GetMacDevSrcRoot(); // Path.GetFullPath(".");
-			CopyAction.CommandPath = BuildHostPlatform.Current.Shell;
-			CopyAction.CommandDescription = "";
-
-			string BundlePath = BundleDirectory.FullName;
-			string SourcePath = Path.Combine(Path.GetFullPath("."), Resource.ResourcePath!);
-			string TargetPath = Path.Combine(BundlePath, "Contents", Resource.BundleContentsSubdir!, Path.GetFileName(Resource.ResourcePath)!);
-
-			FileItem TargetItem = FileItem.GetItemByPath(TargetPath);
-
-			CopyAction.CommandArguments = String.Format("-c \"cp -f -R \\\"{0}\\\" \\\"{1}\\\"; touch -c \\\"{2}\\\"\"", SourcePath, Path.GetDirectoryName(TargetPath)!.Replace('\\', '/') + "/", TargetPath.Replace('\\', '/'));
-			CopyAction.PrerequisiteItems.Add(Executable);
-			CopyAction.ProducedItems.Add(TargetItem);
-			CopyAction.bShouldOutputStatusDescription = Resource.bShouldLog;
-			CopyAction.StatusDescription = String.Format("Copying {0} to app bundle", Path.GetFileName(Resource.ResourcePath));
-			CopyAction.bCanExecuteRemotely = false;
-
-			return TargetItem;
-		}
-
 		private static readonly Dictionary<ReadOnlyTargetRules, DirectoryReference> BundleContentsDirectories = new();
 
 		public override void ModifyBuildProducts(ReadOnlyTargetRules Target, UEBuildBinary Binary, IEnumerable<string> Libraries, IEnumerable<UEBuildBundleResource> BundleResources, Dictionary<FileReference, BuildProductType> BuildProducts)
diff --git a/Engine/Source/Programs/UnrealBuildTool/Platform/TVOS/UEDeployTVOS.cs b/Engine/Source/Programs/UnrealBuildTool/Platform/TVOS/UEDeployTVOS.cs
--- a/Engine/Source/Programs/UnrealBuildTool/Platform/TVOS/UEDeployTVOS.cs
+++ b/Engine/Source/Programs/UnrealBuildTool/Platform/TVOS/UEDeployTVOS.cs
@@ -1,10 +1,7 @@
 // Copyright Epic Games, Inc. All Rights Reserved.
 
-using System;
 using System.Collections.Generic;
-using System.IO;
-using System.Text;
-using System.Xml.Linq;
+
 using EpicGames.Core;
 using Microsoft.Extensions.Logging;
 
@@ -12,7 +9,6 @@ namespace UnrealBuildTool
 {
 	class UEDeployTVOS : UEDeployIOS
 	{
-
 		public UEDeployTVOS(ILogger InLogger)
 			: base(InLogger)
 		{
