[UBT]
* Added timing scope around waiting on git process

#Synced-CL 48764782

diff --git a/Engine/Source/Programs/UnrealBuildTool/System/SourceFileWorkingSet.cs b/Engine/Source/Programs/UnrealBuildTool/System/SourceFileWorkingSet.cs
--- a/Engine/Source/Programs/UnrealBuildTool/System/SourceFileWorkingSet.cs
+++ b/Engine/Source/Programs/UnrealBuildTool/System/SourceFileWorkingSet.cs
@@ -116,12 +116,13 @@ namespace UnrealBuildTool
 
 			Logger.LogInformation("Using 'git status' to determine working set for adaptive non-unity build ({RootDir}).", RootDir);
 
+			string Arguments = "--no-optional-locks status --porcelain";
 			BackgroundProcess = new()
 			{
 				StartInfo = new()
 				{
 					FileName = GitPath,
-					Arguments = "--no-optional-locks status --porcelain",
+					Arguments = Arguments,
 					WorkingDirectory = RootDir.FullName,
 					RedirectStandardOutput = true,
 					RedirectStandardError = true,
@@ -138,6 +139,7 @@ namespace UnrealBuildTool
 			}
 			catch
 			{
+				Logger.LogInformation($"Failed to start git process using command line {GitPath} {Arguments}");
 				BackgroundProcess.Dispose();
 				BackgroundProcess = null;
 			}
@@ -149,6 +151,8 @@ namespace UnrealBuildTool
 					return true;
 				}
 
+				using ITimelineEvent _ = Timeline.ScopeEvent("GitWait");
+
 				if (BackgroundProcess.WaitForExit(500) == false)
 				{
 					Logger.LogInformation("Waiting for 'git status' command to complete");
