Apps launch by default with AppDistributor set to 'Other'. Adjust generated Xcode scheme and devicectl to launch as AppDistrubor AppStore by default. This it done by setting the devicectl launch environment variable and scheme additional option APP_DISTRIBUTOR_ID_OVERRIDE to com.app.AppStore.

#Synced-CL 48762232

diff --git a/Engine/Source/Programs/AutomationTool/IOS/IOSPlatform.Automation.cs b/Engine/Source/Programs/AutomationTool/IOS/IOSPlatform.Automation.cs
--- a/Engine/Source/Programs/AutomationTool/IOS/IOSPlatform.Automation.cs
+++ b/Engine/Source/Programs/AutomationTool/IOS/IOSPlatform.Automation.cs
@@ -1236,6 +1236,7 @@ public class IOSPlatform : ApplePlatform
 				// For now, update the app launch implementation to use devicectl. Windows may need to switch to a different 3rd party implementation.
 				Program = "xcrun";
 				Arguments = "devicectl device process launch";
+				Arguments += " -e '{\\\"APP_DISTRIBUTOR_ID_OVERRIDE\\\": \\\"com.apple.AppStore\\\"}'";
 				Arguments += " --terminate-existing";	// if it's already running on device, kill it first
 				Arguments += " --console";	// attach to the console so we can get log output
 				Arguments += " --device " + Params.DeviceNames[0];
diff --git a/Engine/Source/Programs/UnrealBuildTool/ProjectFiles/Xcode/XcodeProject.Scheme.cs b/Engine/Source/Programs/UnrealBuildTool/ProjectFiles/Xcode/XcodeProject.Scheme.cs
--- a/Engine/Source/Programs/UnrealBuildTool/ProjectFiles/Xcode/XcodeProject.Scheme.cs
+++ b/Engine/Source/Programs/UnrealBuildTool/ProjectFiles/Xcode/XcodeProject.Scheme.cs
@@ -195,6 +195,13 @@ namespace UnrealBuildTool.XcodeProjectXcconfig
 				Content.WriteLine("            ReferencedContainer = \"container:" + ProjectName + ".xcodeproj\">");
 				Content.WriteLine("         </BuildableReference>");
 				Content.WriteLine("      </BuildableProductRunnable>");
+				Content.WriteLine("      <AdditionalOptions>");
+				Content.WriteLine("        <AdditionalOption");
+				Content.WriteLine("          key = \"APP_DISTRIBUTOR_ID_OVERRIDE\"");
+				Content.WriteLine("          value = \"com.apple.AppStore\"");
+				Content.WriteLine("          isEnabled = \"YES\">");
+				Content.WriteLine("        </AdditionalOption>");
+				Content.WriteLine(".     </AdditionalOptions>");
 				if (String.IsNullOrEmpty(OldCommandLineArguments))
 				{
 					if (bHasEditorConfiguration && !String.IsNullOrEmpty(GameProjectPath))
