[UBT] Upgrade Xcode indexit targets to C++ 20

The current implementation was setting C++17 as the target version, but since UE 5.7 the minimum C++ version is 20.

This was causing compile errors in the indexing process

#Synced-CL 48538939

diff --git a/Engine/Source/Programs/UnrealBuildTool/ProjectFiles/Xcode/XcodeProject.cs b/Engine/Source/Programs/UnrealBuildTool/ProjectFiles/Xcode/XcodeProject.cs
--- a/Engine/Source/Programs/UnrealBuildTool/ProjectFiles/Xcode/XcodeProject.cs
+++ b/Engine/Source/Programs/UnrealBuildTool/ProjectFiles/Xcode/XcodeProject.cs
@@ -109,6 +109,7 @@ namespace UnrealBuildTool.XcodeProjectXcconfig
 		public void GenerateResponseFile()
 		{
 			StringBuilder ResponseFileContents = new();
+			ResponseFileContents.Append("-std=c++20 ");
 			ResponseFileContents.Append("-isystem");
 			ResponseFileContents.AppendJoin(" -isystem", SystemIncludePaths.Select(x => x.FullName.Contains(' ') ? $"\"{x.FullName}\"" : x.FullName));
 			ResponseFileContents.Append(" -I");
@@ -2901,7 +2902,7 @@ namespace UnrealBuildTool.XcodeProjectXcconfig
 		public override void WriteXcconfigFile(ILogger Logger)
 		{
 			// write out settings that for compiling natively
-			Xcconfig!.AppendLine("CLANG_CXX_LANGUAGE_STANDARD = c++17");
+			Xcconfig!.AppendLine("CLANG_CXX_LANGUAGE_STANDARD = c++20");
 			Xcconfig.AppendLine("GCC_WARN_INHIBIT_ALL_WARNINGS = YES");
 			Xcconfig.AppendLine("GCC_PRECOMPILE_PREFIX_HEADER = YES");
 			Xcconfig.AppendLine("GCC_OPTIMIZATION_LEVEL = 0");
