[UBT] Add support for AutoRTFM sanitizer builds

This sanitizer is a new iteration of the AutoRTFM memory validator, which works in a similar way to Clang's other sanitizers.

#Synced-CL 48740692

diff --git a/Engine/Source/Programs/UnrealBuildTool/Platform/Clang/ClangToolChain.cs b/Engine/Source/Programs/UnrealBuildTool/Platform/Clang/ClangToolChain.cs
--- a/Engine/Source/Programs/UnrealBuildTool/Platform/Clang/ClangToolChain.cs
+++ b/Engine/Source/Programs/UnrealBuildTool/Platform/Clang/ClangToolChain.cs
@@ -53,88 +53,93 @@ namespace UnrealBuildTool
 		/// </summary>
 		EnableMemorySanitizer = 1 << 5,
 
+		/// <summary>
+		/// Enable AutoRTFM sanitizer
+		/// </summary>
+		EnableAutoRTFMSanitizer = 1 << 6,
+		
 		/// <summary>
 		/// Enable Shared library for the Sanitizers otherwise defaults to Statically linked
 		/// </summary>
-		EnableSharedSanitizer = 1 << 6,
+		EnableSharedSanitizer = 1 << 7,
 
 		/// <summary>
 		/// Enables link time optimization (LTO). Link times will significantly increase.
 		/// </summary>
-		EnableLinkTimeOptimization = 1 << 7,
+		EnableLinkTimeOptimization = 1 << 8,
 
 		/// <summary>
 		/// Enable thin LTO
 		/// </summary>
-		EnableThinLTO = 1 << 8,
+		EnableThinLTO = 1 << 9,
 
 		/// <summary>
 		/// Enable thin LTO
 		/// </summary>
-		ThinLTODistributed = 1 << 9,
+		ThinLTODistributed = 1 << 10,
 
 		/// <summary>
 		/// Enable tuning of debug info for LLDB
 		/// </summary>
-		TuneDebugInfoForLLDB = 1 << 10,
+		TuneDebugInfoForLLDB = 1 << 11,
 
 		/// <summary>
 		/// Whether or not to preserve the portable symbol file produced by dump_syms
 		/// </summary>
-		PreservePSYM = 1 << 11,
+		PreservePSYM = 1 << 12,
 
 		/// <summary>
 		/// (Apple toolchains) Whether we're outputting a dylib instead of an executable
 		/// </summary>
-		OutputDylib = 1 << 12,
+		OutputDylib = 1 << 13,
 
 		/// <summary>
 		/// Enables the creation of custom symbol files used for runtime symbol resolution.
 		/// </summary>
-		GenerateSymbols = 1 << 13,
+		GenerateSymbols = 1 << 14,
 
 		/// <summary>
 		/// Enables dead code/data stripping and common code folding.
 		/// </summary>
-		EnableDeadStripping = 1 << 14,
+		EnableDeadStripping = 1 << 15,
 
 		/// <summary>
 		/// Indicates that the target is a moduler build i.e. Target.LinkType == TargetLinkType.Modular
 		/// </summary>
-		ModularBuild = 1 << 15,
+		ModularBuild = 1 << 16,
 
 		/// <summary>
 		/// Disable Dump Syms step for faster iteration
 		/// </summary>
-		DisableDumpSyms = 1 << 16,
+		DisableDumpSyms = 1 << 17,
 
 		/// <summary>
 		/// Indicates that the AutoRTFM Clang compiler should be used instead of the standard clang compiler
 		/// </summary>
-		UseAutoRTFMCompiler = 1 << 17,
+		UseAutoRTFMCompiler = 1 << 18,
 
 		/// <summary>
 		/// Enable LibFuzzer
 		/// </summary>
-		EnableLibFuzzer = 1 << 18,
+		EnableLibFuzzer = 1 << 19,
 
 		/// <summary>
 		/// Modify code generation to help with debugging optimized builds e.g. by extending lifetimes of local variables.
 		/// It may slightly reduce performance. Thus, it's meant to be used during development only.
 		/// Supported only on some platforms.
 		/// </summary>
-		OptimizeForDebugging = 1 << 19,
+		OptimizeForDebugging = 1 << 20,
 
 		/// <summary>
 		/// Enables compressing the debug sections if the platform supports this
 		/// </summary>
-		CompressDebugFile = 1 << 20,
+		CompressDebugFile = 1 << 21,
 
 		/// <summary>
 		/// Disable Symbol stripping, for use when an external part of the build pipeline requires symbols to be present
 		/// and will strip them in a later step
 		/// </summary>
-		DisableStripSymbols = 1 << 21,
+		DisableStripSymbols = 1 << 22,
 	}
 
 	abstract partial class ClangToolChain : ISPCToolChain
@@ -958,6 +963,12 @@ namespace UnrealBuildTool
 				Arguments.Add("-fsanitize=memory");
 			}
 
+			// AutoRTFMSan
+			if (Options.HasFlag(ClangToolChainOptions.EnableAutoRTFMSanitizer))
+			{
+				Arguments.Add("-fsanitize=autortfm");
+			}
+
 			// LibFuzzer
 			if (Options.HasFlag(ClangToolChainOptions.EnableLibFuzzer))
 			{
diff --git a/Engine/Source/Programs/UnrealBuildTool/Platform/Linux/LinuxToolChain.cs b/Engine/Source/Programs/UnrealBuildTool/Platform/Linux/LinuxToolChain.cs
--- a/Engine/Source/Programs/UnrealBuildTool/Platform/Linux/LinuxToolChain.cs
+++ b/Engine/Source/Programs/UnrealBuildTool/Platform/Linux/LinuxToolChain.cs
@@ -800,6 +800,7 @@ namespace UnrealBuildTool
 				Options.HasFlag(ClangToolChainOptions.EnableThreadSanitizer) ||
 				Options.HasFlag(ClangToolChainOptions.EnableUndefinedBehaviorSanitizer) ||
 				Options.HasFlag(ClangToolChainOptions.EnableMemorySanitizer) ||
+				Options.HasFlag(ClangToolChainOptions.EnableAutoRTFMSanitizer) ||
 				Options.HasFlag(ClangToolChainOptions.EnableLibFuzzer))
 			{
 				string SanitizerInfo = "Building with:";
@@ -809,6 +810,7 @@ namespace UnrealBuildTool
 				SanitizerInfo += Options.HasFlag(ClangToolChainOptions.EnableThreadSanitizer) ? StaticOrShared + " linked ThreadSanitizer" : "";
 				SanitizerInfo += Options.HasFlag(ClangToolChainOptions.EnableUndefinedBehaviorSanitizer) ? StaticOrShared + " linked UndefinedBehaviorSanitizer" : "";
 				SanitizerInfo += Options.HasFlag(ClangToolChainOptions.EnableMemorySanitizer) ? StaticOrShared + " linked MemorySanitizer" : "";
+				SanitizerInfo += Options.HasFlag(ClangToolChainOptions.EnableAutoRTFMSanitizer) ? " AutoRTFMSanitizer" : "";
 				SanitizerInfo += Options.HasFlag(ClangToolChainOptions.EnableLibFuzzer) ? StaticOrShared + " linked LibFuzzer" : "";
 
 				Logger.LogInformation("{SanitizerInfo}", SanitizerInfo);
diff --git a/Engine/Source/Programs/UnrealBuildTool/Platform/Linux/UEBuildLinux.cs b/Engine/Source/Programs/UnrealBuildTool/Platform/Linux/UEBuildLinux.cs
--- a/Engine/Source/Programs/UnrealBuildTool/Platform/Linux/UEBuildLinux.cs
+++ b/Engine/Source/Programs/UnrealBuildTool/Platform/Linux/UEBuildLinux.cs
@@ -103,6 +103,14 @@ namespace UnrealBuildTool
 		[XmlConfigFile(Category = "BuildConfiguration", Name = "bEnableMemorySanitizer")]
 		public bool bEnableMemorySanitizer = false;
 
+		/// <summary>
+		/// Enables AutoRTFM sanitizer (AutoRTFMSan)
+		/// </summary>
+		[CommandLine("-EnableAutoRTFMSan")]
+		[CommandLine("-EnableRTFMSan")]
+		[XmlConfigFile(Category = "BuildConfiguration", Name = "bEnableAutoRTFMSanitizer")]
+		public bool bEnableAutoRTFMSanitizer = false;
+
 		/// <summary>
 		/// Whether or not to preserve the portable symbol file produced by dump_syms
 		/// </summary>
@@ -171,6 +179,8 @@ namespace UnrealBuildTool
 
 		public bool bEnableMemorySanitizer => Inner.bEnableMemorySanitizer;
 
+		public bool bEnableAutoRTFMSanitizer => Inner.bEnableAutoRTFMSanitizer;
+
 		public bool bTuneDebugInfoForLLDB => Inner.bTuneDebugInfoForLLDB;
 
 		public bool bDisableDumpSyms => Inner.bDisableDumpSyms;
@@ -284,6 +294,10 @@ namespace UnrealBuildTool
 				{
 					SanitizerSuffix = "MSan";
 				}
+				else if (Target.LinuxPlatform.bEnableAutoRTFMSanitizer)
+				{
+					SanitizerSuffix = "AutoRTFMSan";
+				}
 				if (Target.LinuxPlatform.bEnableLibFuzzer)
 				{
 					SanitizerSuffix += "LibFuzzer";
@@ -673,6 +687,10 @@ namespace UnrealBuildTool
 					throw new BuildException("Memory Sanitizer (MSan) unsupported for non-monolithic builds");
 				}
 			}
+			if (Target.LinuxPlatform.bEnableAutoRTFMSanitizer)
+			{
+				Options |= ClangToolChainOptions.EnableAutoRTFMSanitizer;
+			}
 			if (Target.LinuxPlatform.bDisableDumpSyms)
 			{
 				Options |= ClangToolChainOptions.DisableDumpSyms;
diff --git a/Engine/Source/Programs/UnrealBuildTool/Platform/Windows/UEBuildWindows.cs b/Engine/Source/Programs/UnrealBuildTool/Platform/Windows/UEBuildWindows.cs
--- a/Engine/Source/Programs/UnrealBuildTool/Platform/Windows/UEBuildWindows.cs
+++ b/Engine/Source/Programs/UnrealBuildTool/Platform/Windows/UEBuildWindows.cs
@@ -546,7 +546,15 @@ namespace UnrealBuildTool
 		[XmlConfigFile(Category = "BuildConfiguration", Name = "bEnableUndefinedBehaviorSanitizer")]
 		[CommandLine("-EnableUBSan")]
 		public bool bEnableUndefinedBehaviorSanitizer = false;
-		
+
+		/// <summary>
+		/// Enables AutoRTFM sanitizer (AutoRTFMSan)
+		/// </summary>
+		[XmlConfigFile(Category = "BuildConfiguration", Name = "bEnableAutoRTFMSanitizer")]
+		[CommandLine("-EnableAutoRTFMSan")]
+		[CommandLine("-EnableRTFMSan")]
+		public bool bEnableAutoRTFMSanitizer = false;
+
 		/// <summary>
 		/// Enables instrumentation.
 		/// </summary>
@@ -1033,6 +1041,8 @@ namespace UnrealBuildTool
 
 		public bool bEnableUndefinedBehaviorSanitizer => Inner.bEnableUndefinedBehaviorSanitizer;
 
+		public bool bEnableAutoRTFMSanitizer => Inner.bEnableAutoRTFMSanitizer;
+
 		public bool bEnableInstrumentation => Inner.bEnableInstrumentation;
 
 		public bool bEnableLibFuzzer => Inner.bEnableLibFuzzer;
@@ -1333,6 +1343,11 @@ namespace UnrealBuildTool
 				}
 			}
 
+			if (Target.WindowsPlatform.bEnableAutoRTFMSanitizer)
+			{
+				Target.WindowsPlatform.Compiler = WindowsCompiler.ClangRTFM;
+			}
+
 			// Set the compiler version if necessary
 			if (Target.WindowsPlatform.Compiler == WindowsCompiler.Default)
 			{
diff --git a/Engine/Source/Programs/UnrealBuildTool/Platform/Windows/VCToolChain.cs b/Engine/Source/Programs/UnrealBuildTool/Platform/Windows/VCToolChain.cs
--- a/Engine/Source/Programs/UnrealBuildTool/Platform/Windows/VCToolChain.cs
+++ b/Engine/Source/Programs/UnrealBuildTool/Platform/Windows/VCToolChain.cs
@@ -110,6 +110,10 @@ namespace UnrealBuildTool
 				{
 					Lines.Add($" * Target.{nameof(Target.WindowsPlatform)}.{nameof(Target.WindowsPlatform.bEnableAddressSanitizer)} is enabled");
 				}
+				if (Target.WindowsPlatform.bEnableAutoRTFMSanitizer)
+				{
+					Lines.Add($" * Target.{nameof(Target.WindowsPlatform)}.{nameof(Target.WindowsPlatform.bEnableAutoRTFMSanitizer)} is enabled");
+				}
 				if (Target.bSupportEditAndContinue)
 				{
 					Lines.Add($" * Target.{nameof(Target.bSupportEditAndContinue)} is enabled");
@@ -845,6 +849,11 @@ namespace UnrealBuildTool
 				}
 			}
 
+			if (Target.WindowsPlatform.bEnableAutoRTFMSanitizer)
+			{
+				Arguments.Add("-fsanitize=autortfm");
+			}
+
 			if (Target.WindowsPlatform.bEnableLibFuzzer)
 			{
 				if (Target.WindowsPlatform.Compiler.IsClang())
