Types to support the storage and retrieval of action results. Rework the hash-addressable storage (which is currently unused) to use a much more performant hash algorithm and a simpler interface that's easier to use correctly for the most common use-case of committing transactions.

#Synced-CL 48842273

diff --git a/Engine/Source/Programs/UnrealBuildTool/Actions/ActionMetrics.cs b/Engine/Source/Programs/UnrealBuildTool/Actions/ActionMetrics.cs
new file mode 100644
--- /dev/null
+++ b/Engine/Source/Programs/UnrealBuildTool/Actions/ActionMetrics.cs
@@ -0,0 +1,13 @@
+// Copyright Epic Games, Inc. All Rights Reserved.
+
+using System;
+
+namespace UnrealBuildTool;
+
+/// <summary>
+/// Metrics for the execution of an action.
+/// </summary>
+/// <param name="ExecutionTime">The wall-clock time the action took to execute.</param>
+/// <param name="ProcessorTime">The amount of CPU time the action took to execute.</param>
+/// <param name="PeakMemoryBytes">The peak amount of memory, in bytes, used by the action.</param>
+internal readonly record struct ActionMetrics(TimeSpan ExecutionTime, TimeSpan ProcessorTime, ulong PeakMemoryBytes);
diff --git a/Engine/Source/Programs/UnrealBuildTool/Actions/ActionResult.cs b/Engine/Source/Programs/UnrealBuildTool/Actions/ActionResult.cs
new file mode 100644
--- /dev/null
+++ b/Engine/Source/Programs/UnrealBuildTool/Actions/ActionResult.cs
@@ -0,0 +1,21 @@
+// Copyright Epic Games, Inc. All Rights Reserved.
+
+using System.Collections.Generic;
+
+namespace UnrealBuildTool;
+
+/// <summary>
+/// The result of the execution of an action.
+/// </summary>
+/// <param name="ExitCode">The exit code the action returned after executing.</param>
+/// <param name="Success">Whether the action succeeded - actions may still fail even if the exit code was 0.</param>
+/// <param name="Output"></param>
+/// <param name="Metrics"></param>
+internal record class ActionResult(int ExitCode, bool Success, IReadOnlyList<string> Output, ActionMetrics Metrics)
+{
+	public static ActionResult From(ExecuteResults executeResults)
+	{
+		ActionMetrics metrics = new(executeResults.ExecutionTime, executeResults.ProcessorTime, (ulong)executeResults.PeakMemoryUsed);
+		return new(executeResults.ExitCode, executeResults.ExitCode == 0, executeResults.LogLines, metrics);
+	}
+}
diff --git a/Engine/Source/Programs/UnrealBuildTool/Actions/ResultStore/ContentAddressableActionResultStore.cs b/Engine/Source/Programs/UnrealBuildTool/Actions/ResultStore/ContentAddressableActionResultStore.cs
new file mode 100644
--- /dev/null
+++ b/Engine/Source/Programs/UnrealBuildTool/Actions/ResultStore/ContentAddressableActionResultStore.cs
@@ -0,0 +1,45 @@
+// Copyright Epic Games, Inc. All Rights Reserved.
+
+using System;
+using System.IO;
+using System.Text.Json;
+using System.Threading;
+using System.Threading.Tasks;
+using EpicGames.Core;
+using UnrealBuildTool.Storage;
+
+namespace UnrealBuildTool.Actions.ResultStore;
+
+internal class ContentAddressableActionResultStore(IStorageProvider storageProvider)
+	: IActionResultStore
+{
+	/// <inheritdoc />
+	public async Task<ActionResult?> LoadResultAsync(IExternalAction action, CancellationToken cancellationToken)
+	{
+		IoHash actionHash = GetActionHash(action);
+		using Stream? stream = storageProvider.CreateReader(actionHash);
+		if (stream is null)
+		{
+			return null;
+		}
+
+		return await JsonSerializer.DeserializeAsync<ActionResult>(stream, cancellationToken: cancellationToken);
+	}
+
+	/// <inheritdoc />
+	public async Task StoreResultAsync(IExternalAction action, ActionResult result, CancellationToken cancellationToken)
+	{
+		IoHash actionHash = GetActionHash(action);
+		using Stream stream = storageProvider.CreateWriter(actionHash, cancellationToken);
+		await JsonSerializer.SerializeAsync(stream, result, cancellationToken: cancellationToken);
+	}
+
+	private IoHash GetActionHash(IExternalAction action)
+	{
+		using Blake3.Hasher hasher = Blake3.Hasher.New();
+		hasher.Update(action.CommandPath.FullName.AsSpan());
+		hasher.Update(action.CommandArguments.AsSpan());
+		hasher.Update(action.CommandVersion.AsSpan());
+		return IoHash.FromBlake3(hasher);
+	}
+}
diff --git a/Engine/Source/Programs/UnrealBuildTool/Actions/ResultStore/IActionResultStore.cs b/Engine/Source/Programs/UnrealBuildTool/Actions/ResultStore/IActionResultStore.cs
new file mode 100644
--- /dev/null
+++ b/Engine/Source/Programs/UnrealBuildTool/Actions/ResultStore/IActionResultStore.cs
@@ -0,0 +1,29 @@
+// Copyright Epic Games, Inc. All Rights Reserved.
+
+using System.Threading;
+using System.Threading.Tasks;
+
+namespace UnrealBuildTool.Actions.ResultStore;
+
+/// <summary>
+/// Persists the result of executing actions.
+/// </summary>
+internal interface IActionResultStore
+{
+	/// <summary>
+	/// Load the result of executing a given action.
+	/// </summary>
+	/// <param name="action">The action to check results of.</param>
+	/// <param name="cancellationToken">A token to allow cancelling the operation.</param>
+	/// <returns>The result of the action, if it was previously persisted.</returns>
+	Task<ActionResult?> LoadResultAsync(IExternalAction action, CancellationToken cancellationToken = default);
+
+	/// <summary>
+	/// Store the result of executing a given action.
+	/// </summary>
+	/// <param name="action">The action to store results of.</param>
+	/// <param name="result">The result of the action to store.</param>
+	/// <param name="cancellationToken">A token to allow cancelling the operation.</param>
+	/// <returns></returns>
+	Task StoreResultAsync(IExternalAction action, ActionResult result, CancellationToken cancellationToken = default);
+}
diff --git a/Engine/Source/Programs/UnrealBuildTool/Actions/ResultStore/NullActionResultStore.cs b/Engine/Source/Programs/UnrealBuildTool/Actions/ResultStore/NullActionResultStore.cs
new file mode 100644
--- /dev/null
+++ b/Engine/Source/Programs/UnrealBuildTool/Actions/ResultStore/NullActionResultStore.cs
@@ -0,0 +1,21 @@
+// Copyright Epic Games, Inc. All Rights Reserved.
+
+using System.Threading;
+using System.Threading.Tasks;
+
+namespace UnrealBuildTool.Actions.ResultStore;
+
+internal class NullActionResultStore : IActionResultStore
+{
+	/// <inheritdoc />
+	public Task<ActionResult?> LoadResultAsync(IExternalAction action, CancellationToken cancellationToken)
+	{
+		return Task.FromResult<ActionResult?>(null);
+	}
+
+	/// <inheritdoc />
+	public Task StoreResultAsync(IExternalAction action, ActionResult result, CancellationToken cancellationToken)
+	{
+		return Task.CompletedTask;
+	}
+}
diff --git a/Engine/Source/Programs/UnrealBuildTool/Storage/IStorageProvider.cs b/Engine/Source/Programs/UnrealBuildTool/Storage/IStorageProvider.cs
--- a/Engine/Source/Programs/UnrealBuildTool/Storage/IStorageProvider.cs
+++ b/Engine/Source/Programs/UnrealBuildTool/Storage/IStorageProvider.cs
@@ -1,59 +1,28 @@
 // Copyright Epic Games, Inc. All Rights Reserved.
 
-using System;
 using System.IO;
+using System.Threading;
 using EpicGames.Core;
 
 namespace UnrealBuildTool.Storage
 {
 	/// <summary>
-	/// Interface for reads from content-addressible storage
-	/// </summary>
-	interface IStorageReader : IDisposable
-	{
-		/// <summary>
-		/// Accessor for the read stream
-		/// </summary>
-		Stream? Stream { get; }
-
-		/// <summary>
-		/// Whether the reader contains any data
-		/// </summary>
-		bool IsValid { get; }
-	}
-
-	/// <summary>
-	/// Interface for writes to content-addressible storage
-	/// </summary>
-	interface IStorageWriter : IDisposable
-	{
-		/// <summary>
-		/// Accessor for the write stream
-		/// </summary>
-		Stream? Stream { get; }
-
-		/// <summary>
-		/// Commits the written data to the cache
-		/// </summary>
-		void Commit();
-	}
-
-	/// <summary>
-	/// Interface for an artifact cache
+	/// Represents a read/write store addressed by a hash representing the contents to be stored.
 	/// </summary>
 	interface IStorageProvider
 	{
 		/// <summary>
-		/// Attempts to open a file from the output cache
+		/// Attempts to open a stream from the output cache for reading.
 		/// </summary>
-		/// <param name="Digest">Digest of the item to retrieve</param>
-		/// <returns>True if the item exists in the cache, false otherwise</returns>
-		IStorageReader CreateReader(ContentHash Digest);
+		/// <returns>A readable stream of the stored content for this hash, if it exists. <see langword="null" /> if it does not exist.</returns>
+		Stream? CreateReader(IoHash hash);
 
 		/// <summary>
-		/// Opens a stream for writing into the cache. The digest 
+		/// Opens a stream for writing into the cache. The stream is transactional and will be committed on dispose unless cancelled.
 		/// </summary>
-		/// <returns></returns>
-		IStorageWriter CreateWriter(ContentHash Digest);
+		/// <param name="hash">The hash addressing the contents to write.</param>
+		/// <param name="cancellationToken">A cancellation token, which, if cancelled before the returned stream is disposed,
+		/// will cancel the transaction.</param>
+		Stream CreateWriter(IoHash hash, CancellationToken cancellationToken);
 	}
 }
diff --git a/Engine/Source/Programs/UnrealBuildTool/Storage/Impl/FileSystemStorageProvider.cs b/Engine/Source/Programs/UnrealBuildTool/Storage/Impl/FileSystemStorageProvider.cs
--- a/Engine/Source/Programs/UnrealBuildTool/Storage/Impl/FileSystemStorageProvider.cs
+++ b/Engine/Source/Programs/UnrealBuildTool/Storage/Impl/FileSystemStorageProvider.cs
@@ -2,140 +2,76 @@
 
 using System;
 using System.IO;
+using System.Threading;
 using EpicGames.Core;
 using UnrealBuildBase;
 
 namespace UnrealBuildTool.Storage.Impl
 {
 	/// <summary>
-	/// Storage provider which uses the filesystem
+	/// Storage provider which uses the filesystem.
 	/// </summary>
-	class FileSystemStorageProvider : IStorageProvider
+	internal class FileSystemStorageProvider : IStorageProvider
 	{
 		/// <summary>
-		/// Cached copy of the current process id
+		/// Cached copy of the current process ID.
 		/// </summary>
-		static readonly int ProcessId = Environment.ProcessId;
+		private static readonly int ProcessId = Environment.ProcessId;
 
 		/// <summary>
-		/// Implements <see cref="IStorageReader"/> for writes to the backing storage
+		/// Represents a write transaction to the cache. If not cancelled, the transaction will be committed on dispose.
 		/// </summary>
-		class StorageReader : IStorageReader
+		private class TransactionalWriterStream(FileReference finalLocation, FileReference tempLocation, CancellationToken cancellationToken)
+			: FileStream(tempLocation.FullName, FileMode.Create, FileAccess.Write, FileShare.Read)
 		{
-			public Stream? Stream
+			protected override void Dispose(bool disposing)
 			{
-				get;
-				private set;
-			}
-
-			public bool IsValid => Stream != null;
-
-			public StorageReader(FileReference Location)
-			{
-				if (FileReference.Exists(Location))
-				{
-					Stream = FileReference.Open(Location, FileMode.Open, FileAccess.Read, FileShare.Read | FileShare.Delete);
-				}
-			}
+				base.Dispose(disposing);
 
-			public void Dispose()
-			{
-				if (Stream != null)
+				if (disposing && !cancellationToken.IsCancellationRequested)
 				{
-					Stream.Dispose();
-					Stream = null;
+					try
+					{
+						FileReference.Move(tempLocation, finalLocation);
+					}
+					catch
+					{
+						FileReference.Delete(tempLocation);
+						throw;
+					}
 				}
 			}
 		}
 
-		/// <summary>
-		/// Represents a write transaction to the cache. Commit() should be called once the transaction is complete.
-		/// </summary>
-		class StorageWriter : IStorageWriter
+		/// <inheritdoc />
+		public Stream? CreateReader(IoHash hash)
 		{
-			FileReference? TempLocation;
-			readonly FileReference FinalLocation;
-
-			public Stream? Stream
-			{
-				get;
-				private set;
-			}
-
-			public StorageWriter(FileReference Location)
-			{
-				FinalLocation = Location;
-				DirectoryReference.CreateDirectory(FinalLocation.Directory);
-
-				TempLocation = new FileReference(String.Format("{0}.{1}", Location.FullName, ProcessId));
-				Stream = FileReference.Open(TempLocation, FileMode.Create, FileAccess.Write, FileShare.Read | FileShare.Delete);
-			}
-
-			public void Commit()
-			{
-				if (TempLocation == null)
-				{
-					throw new InvalidOperationException("Item has already been committed");
-				}
-
-				Stream?.Close();
-
-				try
-				{
-					FileReference.Move(TempLocation, FinalLocation);
-				}
-				catch
-				{
-					FileReference.Delete(TempLocation);
-					TempLocation = null;
-					throw;
-				}
-
-				TempLocation = null;
-			}
-
-			public void Dispose()
+			FileReference fileRef = GetFileForDigest(hash);
+			if (FileReference.Exists(fileRef))
 			{
-				if (Stream != null)
-				{
-					Stream.Dispose();
-					Stream = null;
-				}
-				if (TempLocation != null)
-				{
-					FileReference.Delete(TempLocation);
-					TempLocation = null;
-				}
+				return FileReference.Open(fileRef, FileMode.Open, FileAccess.Read, FileShare.Read | FileShare.Delete);
 			}
-		}
 
-		/// <summary>
-		/// Attempts to open a file from the output cache
-		/// </summary>
-		/// <param name="Digest">Digest of the item to retrieve</param>
-		/// <returns>Reader interface for the file</returns>
-		public IStorageReader CreateReader(ContentHash Digest)
-		{
-			return new StorageReader(GetFileForDigest(Digest));
+			return null;
 		}
 
-		/// <summary>
-		/// Opens a stream for writing into the cache. The digest 
-		/// </summary>
-		/// <returns></returns>
-		public IStorageWriter CreateWriter(ContentHash Digest)
+		/// <inheritdoc />
+		public Stream CreateWriter(IoHash hash, CancellationToken cancellationToken)
 		{
-			return new StorageWriter(GetFileForDigest(Digest));
+			FileReference finalLocation = GetFileForDigest(hash);
+			FileReference tempLocation = new(String.Format("{0}.{1}", finalLocation.FullName, ProcessId));
+			DirectoryReference.CreateDirectory(finalLocation.Directory);
+			return new TransactionalWriterStream(finalLocation, tempLocation, cancellationToken);
 		}
 
 		/// <summary>
-		/// Gets the filename on disk to use for a particular digest
+		/// Gets the filename on disk to use for a particular hash.
 		/// </summary>
-		/// <param name="Digest">The digest to find a filename for</param>
-		/// <returns>Filename to use for the given digest</returns>
-		static FileReference GetFileForDigest(ContentHash Digest)
+		/// <param name="hash">The hash to find a filename for</param>
+		/// <returns>Filename to use for the given hash.</returns>
+		private static FileReference GetFileForDigest(IoHash hash)
 		{
-			string DigestText = Digest.ToString();
+			string DigestText = hash.ToString();
 			return FileReference.Combine(Unreal.EngineDirectory, "Saved", "UnrealBuildTool", "Cache", String.Format("{0}/{1}/{2}/{3}.bin", DigestText[0], DigestText[1], DigestText[2], DigestText));
 		}
 	}
