Move UBT XML / INI configuration support from the System directory to relevant subdirectories within Configuration.

#Synced-CL 48842273

diff --git a/Engine/Source/Programs/UnrealBuildTool/Configuration/Ini/ConfigCache.cs b/Engine/Source/Programs/UnrealBuildTool/Configuration/Ini/ConfigCache.cs
new file mode 100644
--- /dev/null
+++ b/Engine/Source/Programs/UnrealBuildTool/Configuration/Ini/ConfigCache.cs
@@ -0,0 +1,853 @@
+// Copyright Epic Games, Inc. All Rights Reserved.
+
+using System;
+using System.Collections.Concurrent;
+using System.Collections.Generic;
+using System.Diagnostics.CodeAnalysis;
+using System.IO;
+using System.Linq;
+using System.Reflection;
+using EpicGames.Core;
+using Microsoft.Extensions.Logging;
+using Microsoft.Extensions.Logging.Abstractions;
+
+namespace UnrealBuildTool
+{
+	/// <summary>
+	/// Caches config files and config file hierarchies
+	/// </summary>
+	public static class ConfigCache
+	{
+		/// <summary>
+		/// Delegate to add a value to an ICollection in a target object
+		/// </summary>
+		/// <param name="TargetObject">The object containing the field to be modified</param>
+		/// <param name="ValueObject">The value to add</param>
+		delegate void AddElementDelegate(object TargetObject, object? ValueObject);
+
+		/// <summary>
+		/// Caches information about a member with a [ConfigFile] attribute in a type
+		/// </summary>
+		abstract class ConfigMember
+		{
+			/// <summary>
+			/// The attribute instance
+			/// </summary>
+			public ConfigFileAttribute Attribute;
+
+			/// <summary>
+			/// For fields implementing ICollection, specifies the element type
+			/// </summary>
+			public Type? ElementType;
+
+			/// <summary>
+			/// For fields implementing ICollection, a callback to add an element type.
+			/// </summary>
+			public AddElementDelegate? AddElement;
+
+			/// <summary>
+			/// Constructor
+			/// </summary>
+			/// <param name="Attribute"></param>
+			public ConfigMember(ConfigFileAttribute Attribute)
+			{
+				this.Attribute = Attribute;
+			}
+
+			/// <summary>
+			/// Returns Reflection.MemberInfo describing the target class member.
+			/// </summary>
+			public abstract MemberInfo MemberInfo { get; }
+
+			/// <summary>
+			/// Returns Reflection.Type of the target class member.
+			/// </summary>
+			public abstract Type Type { get; }
+
+			/// <summary>
+			/// Returns the value setter of the target class member.
+			/// </summary>
+			public abstract Action<object?, object?> SetValue { get; }
+
+			/// <summary>
+			/// Returns the value getter of the target class member.
+			/// </summary>
+			public abstract Func<object?, object?> GetValue { get; }
+		}
+
+		/// <summary>
+		/// Caches information about a field with a [ConfigFile] attribute in a type
+		/// </summary>
+		class ConfigField : ConfigMember
+		{
+			/// <summary>
+			/// Reflection description of the field with the config attribute.
+			/// </summary>
+			private readonly FieldInfo FieldInfo;
+
+			/// <summary>
+			/// Constructor
+			/// </summary>
+			/// <param name="FieldInfo"></param>
+			/// <param name="Attribute"></param>
+			public ConfigField(FieldInfo FieldInfo, ConfigFileAttribute Attribute)
+				: base(Attribute)
+			{
+				this.FieldInfo = FieldInfo;
+			}
+
+			public override MemberInfo MemberInfo => FieldInfo;
+			public override Type Type => FieldInfo.FieldType;
+			public override Action<object?, object?> SetValue => FieldInfo.SetValue;
+			public override Func<object?, object?> GetValue => FieldInfo.GetValue;
+		}
+
+		/// <summary>
+		/// Caches information about a property with a [ConfigFile] attribute in a type
+		/// </summary>
+		class ConfigProperty : ConfigMember
+		{
+			/// <summary>
+			/// Reflection description of the property with the config attribute.
+			/// </summary>
+			private readonly PropertyInfo PropertyInfo;
+
+			/// <summary>
+			/// Constructor
+			/// </summary>
+			/// <param name="PropertyInfo"></param>
+			/// <param name="Attribute"></param>
+			public ConfigProperty(PropertyInfo PropertyInfo, ConfigFileAttribute Attribute)
+				: base(Attribute)
+			{
+				this.PropertyInfo = PropertyInfo;
+			}
+
+			public override MemberInfo MemberInfo => PropertyInfo;
+			public override Type Type => PropertyInfo.PropertyType;
+			public override Action<object?, object?> SetValue => PropertyInfo.SetValue;
+			public override Func<object?, object?> GetValue => PropertyInfo.GetValue;
+		}
+
+		/// <summary>
+		/// Allowed modification types allowed for default config files
+		/// </summary>
+		public enum ConfigDefaultUpdateType
+		{
+			/// <summary>
+			/// Used for non-array types, this will replace a setting, or it will add a setting if it doesn't exist
+			/// </summary>
+			SetValue,
+			/// <summary>
+			/// Used to add an array entry to the end of any existing array entries, or will add to the end of the section
+			/// </summary>
+			AddArrayEntry,
+		}
+
+		/// <summary>
+		/// Stores information identifying a unique config hierarchy
+		/// </summary>
+		class ConfigHierarchyKey
+		{
+			/// <summary>
+			/// The hierarchy type
+			/// </summary>
+			public ConfigHierarchyType Type;
+
+			/// <summary>
+			/// The project directory to read from
+			/// </summary>
+			public DirectoryReference? ProjectDir;
+
+			/// <summary>
+			/// Which platform-specific files to read
+			/// </summary>
+			public UnrealTargetPlatform Platform;
+
+			/// <summary>
+			/// Custom config subdirectory to load
+			/// </summary>
+			public string CustomConfig;
+
+			/// <summary>
+			/// Custom override strings
+			/// </summary>
+			public List<string> OverrideStrings;
+
+			/// <summary>
+			/// A hotfix directory where modifications/additions to config files can be read from
+			/// </summary>
+			public DirectoryReference? HotfixDir;
+
+			/// <summary>
+			/// If specified, this hierarchy will include plugins enabled for the given target type
+			/// </summary>
+			public TargetType? IncludePluginsForTargetType;
+
+			/// <summary>
+			/// Constructor
+			/// </summary>
+			/// <param name="Type">The hierarchy type</param>
+			/// <param name="ProjectDir">The project directory to read from</param>
+			/// <param name="Platform">Which platform-specific files to read</param>
+			/// <param name="CustomConfig">Custom config subdirectory to load</param>
+			/// <param name="OverrideStrings">Custom override strings</param>
+			/// <param name="HotfixDir">A hotfix directory where modifications/additions to config files can be read from</param>
+			/// <param name="IncludePluginsForTargetType">If specified, this hierarchy will unclude plugins enabled for the given target type</param>
+			public ConfigHierarchyKey(ConfigHierarchyType Type, DirectoryReference? ProjectDir, UnrealTargetPlatform Platform, string CustomConfig, List<string> OverrideStrings, DirectoryReference? HotfixDir, TargetType? IncludePluginsForTargetType)
+			{
+				this.Type = Type;
+				this.ProjectDir = ProjectDir;
+				this.Platform = Platform;
+				this.CustomConfig = CustomConfig;
+				this.OverrideStrings = OverrideStrings;
+				this.HotfixDir = HotfixDir;
+				this.IncludePluginsForTargetType = IncludePluginsForTargetType;
+			}
+
+			/// <summary>
+			/// Test whether this key is equal to another object
+			/// </summary>
+			/// <param name="Other">The object to compare against</param>
+			/// <returns>True if the objects match, false otherwise</returns>
+			public override bool Equals(object? Other)
+			{
+				ConfigHierarchyKey? OtherKey = Other as ConfigHierarchyKey;
+				return OtherKey != null &&
+					OtherKey.Type == Type &&
+					OtherKey.ProjectDir == ProjectDir &&
+					OtherKey.Platform == Platform &&
+					OtherKey.CustomConfig == CustomConfig &&
+					OtherKey.OverrideStrings.SequenceEqual(OverrideStrings) &&
+					OtherKey.HotfixDir == HotfixDir &&
+					OtherKey.IncludePluginsForTargetType == IncludePluginsForTargetType;
+			}
+
+			/// <summary>
+			/// Returns a stable hash of this object
+			/// </summary>
+			/// <returns>Hash value for this object</returns>
+			public override int GetHashCode()
+			{
+				int Hash = 17;
+				Hash = (Hash * 31) + Type.GetHashCode();
+				Hash = (Hash * 31) + ((ProjectDir == null) ? 0 : ProjectDir.GetHashCode());
+				Hash = (Hash * 31) + Platform.GetHashCode();
+				Hash = (Hash * 31) + CustomConfig.GetHashCode();
+				foreach (string OverrideString in OverrideStrings)
+				{
+					Hash = (Hash * 31) + OverrideString.GetHashCode();
+				}
+				Hash = (Hash * 31) + ((HotfixDir == null) ? 0 : HotfixDir.GetHashCode());
+				Hash = (Hash * 31) + ((IncludePluginsForTargetType == null) ? 0 : IncludePluginsForTargetType.GetHashCode());
+
+				return Hash;
+			}
+		}
+
+		/// <summary>
+		/// Cache of individual config files
+		/// </summary>
+		static readonly ConcurrentDictionary<FileReference, Lazy<ConfigFile?>> LocationToConfigFile = [];
+
+		/// <summary>
+		/// Cache of config hierarchies by project
+		/// </summary>
+		static readonly ConcurrentDictionary<ConfigHierarchyKey, Lazy<ConfigHierarchy>> HierarchyKeyToHierarchy = [];
+
+		/// <summary>
+		/// Cache of config fields by type
+		/// </summary>
+		static readonly ConcurrentDictionary<Type, Lazy<List<ConfigMember>>> TypeToConfigMembers = [];
+
+		/// <summary>
+		/// Attempts to read a config file (or retrieve it from the cache)
+		/// </summary>
+		/// <param name="Location">Location of the file to read</param>
+		/// <param name="ConfigFile">On success, receives the parsed config file</param>
+		/// <returns>True if the file exists and was read, false otherwise</returns>
+		internal static bool TryReadFile(FileReference Location, [NotNullWhen(true)] out ConfigFile? ConfigFile)
+		{
+			ConfigFile = LocationToConfigFile.GetOrAdd(Location, static Location => new Lazy<ConfigFile?>(() =>
+			{
+				if (FileReference.Exists(Location))
+				{
+					return new ConfigFile(Location);
+				}
+
+				// Note: it's important to keep the null value here so we don't keep checking for files that don't exist
+				return null;
+			})).Value;
+			return ConfigFile != null;
+		}
+
+		/// <summary>
+		/// Reads a config hierarchy (or retrieve it from the cache)
+		/// </summary>
+		/// <param name="Type">The type of hierarchy to read</param>
+		/// <param name="ProjectDir">The project directory to read the hierarchy for</param>
+		/// <param name="Platform">Which platform to read platform-specific config files for</param>
+		/// <param name="CustomConfig">Optional override config directory to search, for support of multiple target types</param>
+		/// <param name="CustomArgs">Optional list of command line arguments added to the existing command line arguments</param>
+		/// <param name="HotfixDir">A hotfix directory where modifications/additions to config files can be read from</param>
+		/// <param name="IncludePluginsForTargetType">If specified, this hierarchy will unclude plugins enabled for the given target type</param>
+		/// <returns>The requested config hierarchy</returns>
+		public static ConfigHierarchy ReadHierarchy(ConfigHierarchyType Type, DirectoryReference? ProjectDir, UnrealTargetPlatform Platform, string CustomConfig = "", string[]? CustomArgs = null, DirectoryReference? HotfixDir = null, TargetType? IncludePluginsForTargetType = null)
+		{
+			CommandLineArguments CombinedArgs;
+			if (CustomArgs != null)
+			{
+				List<string> CommandLineArgs = new List<string>(Environment.GetCommandLineArgs());
+				foreach (string CustomArg in CustomArgs)
+				{
+					if (CustomArg.StartsWith("-", StringComparison.InvariantCultureIgnoreCase))
+					{
+						CommandLineArgs.Add(CustomArg);
+					}
+					else
+					{
+						CommandLineArgs.Add("-" + CustomArg);
+					}
+				}
+				CombinedArgs = new CommandLineArguments(CommandLineArgs.ToArray());
+			}
+			else
+			{
+				CombinedArgs = new CommandLineArguments(Environment.GetCommandLineArgs());
+			}
+
+			return ReadHierarchy(Type, ProjectDir, Platform, CustomConfig, CombinedArgs, HotfixDir, IncludePluginsForTargetType);
+		}
+
+		/// <summary>
+		/// Reads a config hierarchy (or retrieve it from the cache)
+		/// </summary>
+		/// <param name="Type">The type of hierarchy to read</param>
+		/// <param name="ProjectDir">The project directory to read the hierarchy for</param>
+		/// <param name="Platform">Which platform to read platform-specific config files for</param>
+		/// <param name="CustomConfig">Optional override config directory to search, for support of multiple target types</param>
+		/// <param name="CmdLineArgs">The command line arguments to parse</param>
+		/// <param name="HotfixDir">A hotfix directory where modifications/additions to config files can be read from</param>
+		/// <param name="IncludePluginsForTargetType">If specified, this hierarchy will unclude plugins enabled for the given target type</param>
+		/// <returns>The requested config hierarchy</returns>
+		public static ConfigHierarchy ReadHierarchy(ConfigHierarchyType Type, DirectoryReference? ProjectDir, UnrealTargetPlatform Platform,
+			string CustomConfig, CommandLineArguments CmdLineArgs, DirectoryReference? HotfixDir = null, TargetType? IncludePluginsForTargetType = null)
+		{
+			// Handle command line overrides
+			List<string> OverrideStrings = new List<string>();
+
+			string IniConfigArgPrefix = "-ini:" + Enum.GetName(typeof(ConfigHierarchyType), Type) + ":";
+			string CustomConfigPrefix = "-CustomConfig=";
+			bool bDumpIniLoads = false;
+			foreach (string CmdLineArg in CmdLineArgs)
+			{
+				if (CmdLineArg.StartsWith(IniConfigArgPrefix, StringComparison.OrdinalIgnoreCase))
+				{
+					OverrideStrings.Add(CmdLineArg.Substring(IniConfigArgPrefix.Length));
+				}
+				if (CmdLineArg.StartsWith(CustomConfigPrefix, StringComparison.OrdinalIgnoreCase))
+				{
+					CustomConfig = CmdLineArg.Substring(CustomConfigPrefix.Length);
+				}
+				if (CmdLineArg.Equals("-dumpiniloads", StringComparison.OrdinalIgnoreCase))
+				{
+					bDumpIniLoads = true;
+				}
+			}
+
+			CustomConfig ??= String.Empty;
+
+			// Get the key to use for the cache. It cannot be null, so we use the engine directory if a project directory is not given.
+			ConfigHierarchyKey Key = new ConfigHierarchyKey(Type, ProjectDir, Platform, CustomConfig, OverrideStrings, HotfixDir, IncludePluginsForTargetType);
+
+			ILogger Logger = NullLogger.Instance;
+
+			if (bDumpIniLoads)
+			{
+				Logger = Log.Logger;
+			}
+
+			// Try to get the cached hierarchy with this key
+			Logger.LogInformation($"Requested Hierarchy: {Type},{ProjectDir},{Platform},{CustomConfig}");
+			return HierarchyKeyToHierarchy.GetOrAdd(Key, new Lazy<ConfigHierarchy>(() => InternalReadHierarchy(Logger, Type, ProjectDir, Platform, CustomConfig, HotfixDir, IncludePluginsForTargetType, OverrideStrings))).Value;
+		}
+
+		static ConfigHierarchy InternalReadHierarchy(ILogger Logger, ConfigHierarchyType Type, DirectoryReference? ProjectDir, UnrealTargetPlatform Platform, string CustomConfig
+			, DirectoryReference? HotfixDir, TargetType? IncludePluginsForTargetType, List<string> OverrideStrings)
+		{
+			// Find all the input files
+			List<ConfigFile> Files = new List<ConfigFile>();
+			IEnumerable<FileReference> IniFileNames = ConfigHierarchy.EnumerateConfigFileLocations(Type, ProjectDir, Platform, CustomConfig, IncludePluginsForTargetType);
+			List<FileReference> ReadFiles = [];
+			foreach (FileReference IniFileName in IniFileNames)
+			{
+				ConfigFile? File;
+				Logger.LogInformation($"Trying to load: {IniFileName}");
+				if (TryReadFile(IniFileName, out File))
+				{
+					Logger.LogInformation($"    Found!");
+					Files.Add(File);
+					ReadFiles.Add(IniFileName);
+				}
+
+				if (HotfixDir != null)
+				{
+					FileReference HotfixFile = FileReference.Combine(HotfixDir, IniFileName.GetFileName());
+
+					if (TryReadFile(HotfixFile, out File))
+					{
+						Logger.LogInformation($"    Hotfixed with '{HotfixFile}'");
+						Files.Add(File);
+						ReadFiles.Add(HotfixFile);
+					}
+				}
+			}
+
+			foreach (string OverrideString in OverrideStrings)
+			{
+				ConfigFile OverrideFile = new ConfigFile(OverrideString);
+				Files.Add(OverrideFile);
+				Logger.LogInformation($"Trying to load override: {OverrideFile}");
+			}
+
+			// Create the hierarchy
+			return new ConfigHierarchy(Files, ReadFiles);
+		}
+
+		/// <summary>
+		/// Gets a list of ConfigFields for the given type
+		/// </summary>
+		/// <param name="TargetObjectType">Type to get configurable fields for</param>
+		/// <returns>List of config fields for the given type</returns>
+		static List<ConfigMember> FindConfigMembersForType(Type TargetObjectType)
+		{
+			return TypeToConfigMembers.GetOrAdd(TargetObjectType, static TargetObjectType => new Lazy<List<ConfigMember>>(() =>
+			{
+				List<ConfigMember> Members = new List<ConfigMember>();
+				if (TargetObjectType.BaseType != null)
+				{
+					Members.AddRange(FindConfigMembersForType(TargetObjectType.BaseType));
+				}
+				foreach (FieldInfo FieldInfo in TargetObjectType.GetFields(BindingFlags.Instance | BindingFlags.GetField | BindingFlags.GetProperty | BindingFlags.Public | BindingFlags.NonPublic | BindingFlags.DeclaredOnly))
+				{
+					ProcessConfigTypeMember<FieldInfo>(FieldInfo, Members, (FieldInfo, Attribute) => new ConfigField(FieldInfo, Attribute));
+				}
+				foreach (PropertyInfo PropertyInfo in TargetObjectType.GetProperties(BindingFlags.Instance | BindingFlags.GetField | BindingFlags.GetProperty | BindingFlags.Public | BindingFlags.NonPublic | BindingFlags.DeclaredOnly))
+				{
+					ProcessConfigTypeMember<PropertyInfo>(PropertyInfo, Members, (PropertyInfo, Attribute) => new ConfigProperty(PropertyInfo, Attribute));
+				}
+				return Members;
+			})).Value;
+		}
+
+		static void ProcessConfigTypeMember<MEMBER>(MEMBER MemberInfo, List<ConfigMember> Members, Func<MEMBER, ConfigFileAttribute, ConfigMember> CreateConfigMember)
+			where MEMBER : System.Reflection.MemberInfo
+		{
+			IEnumerable<ConfigFileAttribute> Attributes = MemberInfo.GetCustomAttributes<ConfigFileAttribute>();
+			foreach (ConfigFileAttribute Attribute in Attributes)
+			{
+				// Copy the field 
+				ConfigMember Setter = CreateConfigMember(MemberInfo, Attribute);
+
+				// Check if the field type implements ICollection<>. If so, we can take multiple values.
+				foreach (Type InterfaceType in Setter.Type.GetInterfaces())
+				{
+					if (InterfaceType.IsGenericType && InterfaceType.GetGenericTypeDefinition() == typeof(ICollection<>))
+					{
+						MethodInfo MethodInfo = InterfaceType.GetRuntimeMethod("Add", new Type[] { InterfaceType.GenericTypeArguments[0] })!;
+						Setter.AddElement = (Target, Value) => { MethodInfo.Invoke(Setter.GetValue(Target), new object?[] { Value }); };
+						Setter.ElementType = InterfaceType.GenericTypeArguments[0];
+						break;
+					}
+				}
+
+				// Add it to the output list
+				Members.Add(Setter);
+			}
+		}
+
+		/// <summary>
+		/// Read config settings for the given object
+		/// </summary>
+		/// <param name="ProjectDir">Path to the project directory</param>
+		/// <param name="Platform">The platform being built</param>
+		/// <param name="TargetObject">Object to receive the settings</param>
+		/// <param name="CustomConfig">Optional override config directory to search, for support of multiple target types</param>
+		public static void ReadSettings(DirectoryReference? ProjectDir, UnrealTargetPlatform Platform, object TargetObject, string CustomConfig = "")
+		{
+			ReadSettings(ProjectDir, Platform, TargetObject, null, null, CustomConfig);
+		}
+
+		/// <summary>
+		/// Read config settings for the given object
+		/// </summary>
+		/// <param name="ProjectDir">Path to the project directory</param>
+		/// <param name="Platform">The platform being built</param>
+		/// <param name="TargetObject">Object to receive the settings</param>
+		/// <param name="ConfigValues">Will be populated with config values that were retrieved. May be null.</param>
+		/// <param name="CmdLineArgs">Command line arguments, if null the application's command line arguments will be used</param>
+		/// <param name="CustomConfig">Optional override config directory to search, for support of multiple target types</param>
+		internal static void ReadSettings(DirectoryReference? ProjectDir, UnrealTargetPlatform Platform, object TargetObject, Dictionary<ConfigDependencyKey, IReadOnlyList<string>?>? ConfigValues, CommandLineArguments? CmdLineArgs, string CustomConfig = "")
+		{
+			List<ConfigMember> Members = FindConfigMembersForType(TargetObject.GetType());
+
+			foreach (ConfigMember Member in Members)
+			{
+				// Read the hierarchy listed
+				ConfigHierarchy Hierarchy;
+
+				if (CmdLineArgs == null)
+				{
+					Hierarchy = ReadHierarchy(Member.Attribute.ConfigType, ProjectDir, Platform, CustomConfig);
+				}
+				else
+				{
+					Hierarchy = ReadHierarchy(Member.Attribute.ConfigType, ProjectDir, Platform, CustomConfig, CmdLineArgs);
+				}
+
+				// Get the key name
+				string KeyName = Member.Attribute.KeyName ?? Member.MemberInfo.Name;
+
+				// Get the value(s) associated with this key
+				IReadOnlyList<string>? Values;
+				Hierarchy.TryGetValues(Member.Attribute.SectionName, KeyName, out Values);
+
+				// Parse the values from the config files and update the target object
+				if (Member.AddElement == null)
+				{
+					if (Values != null && Values.Count == 1)
+					{
+						object? Value;
+						if (TryParseValue(Values[0], Member.Type, out Value))
+						{
+							Member.SetValue(TargetObject, Value);
+						}
+					}
+				}
+				else
+				{
+					if (Values != null)
+					{
+						foreach (string Item in Values)
+						{
+							object? Value;
+							if (TryParseValue(Item, Member.ElementType!, out Value))
+							{
+								Member.AddElement(TargetObject, Value);
+							}
+						}
+					}
+				}
+
+				// Save the dependency
+				if (ConfigValues != null)
+				{
+					ConfigDependencyKey Key = new ConfigDependencyKey(Member.Attribute.ConfigType, ProjectDir, Platform, Member.Attribute.SectionName, KeyName);
+					ConfigValues[Key] = Values;
+				}
+			}
+		}
+
+		/// <summary>
+		/// Attempts to parse the given text into an object which matches a specific field type
+		/// </summary>
+		/// <param name="Text">The text to parse</param>
+		/// <param name="FieldType">The type of field to parse</param>
+		/// <param name="Value">If successful, a value of type 'FieldType'</param>
+		/// <returns>True if the value could be parsed, false otherwise</returns>
+		public static bool TryParseValue(string Text, Type FieldType, out object? Value)
+		{
+			if (FieldType == typeof(string))
+			{
+				Value = Text;
+				return true;
+			}
+			else if (FieldType == typeof(bool))
+			{
+				bool BoolValue;
+				if (ConfigHierarchy.TryParse(Text, out BoolValue))
+				{
+					Value = BoolValue;
+					return true;
+				}
+				else
+				{
+					Value = null;
+					return false;
+				}
+			}
+			else if (FieldType == typeof(int))
+			{
+				int IntValue;
+				if (ConfigHierarchy.TryParse(Text, out IntValue))
+				{
+					Value = IntValue;
+					return true;
+				}
+				else
+				{
+					Value = null;
+					return false;
+				}
+			}
+			else if (FieldType == typeof(float))
+			{
+				float FloatValue;
+				if (ConfigHierarchy.TryParse(Text, out FloatValue))
+				{
+					Value = FloatValue;
+					return true;
+				}
+				else
+				{
+					Value = null;
+					return false;
+				}
+			}
+			else if (FieldType == typeof(double))
+			{
+				double DoubleValue;
+				if (ConfigHierarchy.TryParse(Text, out DoubleValue))
+				{
+					Value = DoubleValue;
+					return true;
+				}
+				else
+				{
+					Value = null;
+					return false;
+				}
+			}
+			else if (FieldType == typeof(Guid))
+			{
+				Guid GuidValue;
+				if (ConfigHierarchy.TryParse(Text, out GuidValue))
+				{
+					Value = GuidValue;
+					return true;
+				}
+				else
+				{
+					Value = null;
+					return false;
+				}
+			}
+			else if (FieldType.IsEnum)
+			{
+				try
+				{
+					Value = Enum.Parse(FieldType, Text);
+					return true;
+				}
+				catch
+				{
+					Value = null;
+					return false;
+				}
+			}
+			else if (FieldType.GetGenericTypeDefinition() == typeof(Nullable<>))
+			{
+				return TryParseValue(Text, FieldType.GetGenericArguments()[0], out Value);
+			}
+			else
+			{
+				throw new Exception("Unsupported type for [ConfigFile] attribute");
+			}
+		}
+
+		#region Updating Default Config file support
+
+		/// <summary>
+		/// Calculates the path to where the project's Default config of the type given (ie DefaultEngine.ini)
+		/// </summary>
+		/// <param name="ConfigType">Game, Engine, etc</param>
+		/// <param name="ProjectDir">Project directory, used to find Config/Default[Type].ini</param>
+		public static FileReference GetDefaultConfigFileReference(ConfigHierarchyType ConfigType, DirectoryReference ProjectDir)
+		{
+			return FileReference.Combine(ProjectDir, "Config", $"Default{ConfigType}.ini");
+		}
+
+		/// <summary>
+		/// Calculates the path to where the project's platform specific config of the type given (ie DefaultEngine.ini)
+		/// </summary>
+		/// <param name="ConfigType">Game, Engine, etc</param>
+		/// <param name="ProjectDir">Project directory, used to find Config/Default[Type].ini</param>
+		/// <param name="Platform">Platform name</param>
+		public static FileReference GetPlatformConfigFileReference(ConfigHierarchyType ConfigType, DirectoryReference ProjectDir, string Platform)
+		{
+			return FileReference.Combine(ProjectDir, "Platforms", Platform, "Config", $"{Platform}{ConfigType}.ini");
+		}
+
+		/// <summary>
+		/// Updates a section in a Default***.ini, and will write it out. If the file is not writable, p4 can attempt to check it out.
+		/// </summary>
+		/// <param name="ConfigType">Game, Engine, etc</param>
+		/// <param name="ProjectDir">Project directory, used to find Config/Default[Type].ini</param>
+		/// <param name="UpdateType">How to modify the secion</param>
+		/// <param name="Section">Name of the section with the Key in it</param>
+		/// <param name="Key">Key to update</param>
+		/// <param name="Value">Value to write for te Key</param>
+		/// <param name="Logger">Logger for output</param>
+		/// <returns></returns>		
+		public static bool WriteSettingToDefaultConfig(ConfigHierarchyType ConfigType, DirectoryReference ProjectDir, ConfigDefaultUpdateType UpdateType, string Section, string Key, string Value, ILogger Logger)
+		{
+			FileReference DefaultConfigFile = GetDefaultConfigFileReference(ConfigType, ProjectDir);
+			return WriteSettingToConfigFile(DefaultConfigFile, UpdateType, Section, Key, Value, Logger);
+		}
+
+		/// <summary>
+		/// Updates a section in a Default***.ini, and will write it out. If the file is not writable, p4 can attempt to check it out.
+		/// </summary>
+		/// <param name="ConfigType">Game, Engine, etc</param>
+		/// <param name="ProjectDir">Project directory, used to find Config/Default[Type].ini</param>
+		/// <param name="Platform">Platform name. Must have a directory under projects Platforms subdirectory</param>
+		/// <param name="UpdateType">How to modify the secion</param>
+		/// <param name="Section">Name of the section with the Key in it</param>
+		/// <param name="Key">Key to update</param>
+		/// <param name="Value">Value to write for te Key</param>
+		/// <param name="Logger">Logger for output</param>
+		/// <returns></returns>		
+		public static bool WriteSettingToPlatformConfigFile(ConfigHierarchyType ConfigType, DirectoryReference ProjectDir, string Platform, ConfigDefaultUpdateType UpdateType, string Section, string Key, string Value, ILogger Logger)
+		{
+			FileReference PlatformConfigFile = GetPlatformConfigFileReference(ConfigType, ProjectDir, Platform);
+			return WriteSettingToConfigFile(PlatformConfigFile, UpdateType, Section, Key, Value, Logger);
+		}
+
+		/// <summary>
+		/// Updates a section in config file, and will write it out. If the file is not writable, p4 can attempt to check it out.
+		/// </summary>
+		/// <param name="ConfigFile">Config file (.ini) name</param>
+		/// <param name="UpdateType">How to modify the secion</param>
+		/// <param name="Section">Name of the section with the Key in it</param>
+		/// <param name="Key">Key to update</param>
+		/// <param name="Value">Value to write for te Key</param>
+		/// <param name="Logger">Logger for output</param>
+		/// <returns></returns>		
+		public static bool WriteSettingToConfigFile(FileReference ConfigFile, ConfigDefaultUpdateType UpdateType, string Section, string Key, string Value, ILogger Logger)
+		{
+
+			if (!FileReference.Exists(ConfigFile))
+			{
+				Logger.LogWarning("Failed to find config file '{DefaultConfigFile}' to update", ConfigFile);
+				return false;
+			}
+
+			if (File.GetAttributes(ConfigFile.FullName).HasFlag(FileAttributes.ReadOnly))
+			{
+				Logger.LogWarning("Config file '{ConfigFile}' is read-only, unable to write setting {Key}", ConfigFile.FullName, Key);
+				return false;
+			}
+
+			// generate the section header
+			string SectionString = $"[{Section}]";
+
+			// genrate the line we are going to write to the config
+			string KeyWithEquals = $"{Key}=";
+			string LineToWrite = (UpdateType == ConfigDefaultUpdateType.AddArrayEntry ? "+" : "");
+			LineToWrite += KeyWithEquals + Value;
+
+			// read in all the lines so we can insert or replace one
+			List<string> Lines = File.ReadAllLines(ConfigFile.FullName).ToList();
+
+			// look for the section
+			int SectionIndex = -1;
+			for (int Index = 0; Index < Lines.Count; Index++)
+			{
+				if (Lines[Index].Trim().Equals(SectionString, StringComparison.OrdinalIgnoreCase))
+				{
+					SectionIndex = Index;
+					break;
+				}
+			}
+
+			// if section not found, just append to the end
+			if (SectionIndex == -1)
+			{
+				Lines.Add(SectionString);
+				Lines.Add(LineToWrite);
+
+				File.WriteAllLines(ConfigFile.FullName, Lines);
+				return true;
+			}
+
+			// find the last line in the section with the prefix
+			int LastIndexOfPrefix = -1;
+			int NextSectionIndex = -1;
+			for (int Index = SectionIndex + 1; Index < Lines.Count; Index++)
+			{
+				string Line = Lines[Index];
+				if (Line.StartsWith('+') || Line.StartsWith('-') || Line.StartsWith('.') || Line.StartsWith('!'))
+				{
+					Line = Line.Substring(1);
+				}
+
+				// look for last array entry in case of multiples (or the only line for non-array type)
+				if (Line.StartsWith(KeyWithEquals, StringComparison.InvariantCultureIgnoreCase))
+				{
+					LastIndexOfPrefix = Index;
+				}
+				else if (Lines[Index].StartsWith("["))
+				{
+					NextSectionIndex = Index;
+					// we found another section, so break out
+					break;
+				}
+			}
+
+			// now we know enough to either insert or replace a line
+
+			// if we never found the key, we will insert at the end of the section
+			if (LastIndexOfPrefix == -1)
+			{
+				// if we didn't find a next section, thjen we will insert at the end of the file
+				if (NextSectionIndex == -1)
+				{
+					NextSectionIndex = Lines.Count;
+				}
+
+				// move past blank lines between sections
+				while (String.IsNullOrWhiteSpace(Lines[NextSectionIndex - 1]))
+				{
+					NextSectionIndex--;
+				}
+				// insert before the next section (or end of file)
+				Lines.Insert(NextSectionIndex, LineToWrite);
+			}
+			// otherwise, insert after, or replace, a line, depending on type
+			else
+			{
+				if (UpdateType == ConfigDefaultUpdateType.AddArrayEntry)
+				{
+					Lines.Insert(LastIndexOfPrefix + 1, LineToWrite);
+				}
+				else
+				{
+					Lines[LastIndexOfPrefix] = LineToWrite;
+				}
+			}
+
+			// now the lines are updated, we can overwrite the file
+			File.WriteAllLines(ConfigFile.FullName, Lines);
+			return true;
+		}
+
+		/// <summary>
+		/// Invalidates the hierarchy internal caches so that next call to ReadHierarchy will re-read from disk
+		/// but existing ones will still be valid with old values
+		/// </summary>
+		public static void InvalidateCaches()
+		{
+			lock (LocationToConfigFile)
+			{
+				LocationToConfigFile.Clear();
+			}
+
+			lock (HierarchyKeyToHierarchy)
+			{
+				HierarchyKeyToHierarchy.Clear();
+			}
+		}
+
+		#endregion
+	}
+}
diff --git a/Engine/Source/Programs/UnrealBuildTool/Configuration/Ini/ConfigFile.cs b/Engine/Source/Programs/UnrealBuildTool/Configuration/Ini/ConfigFile.cs
new file mode 100644
--- /dev/null
+++ b/Engine/Source/Programs/UnrealBuildTool/Configuration/Ini/ConfigFile.cs
@@ -0,0 +1,531 @@
+// Copyright Epic Games, Inc. All Rights Reserved.
+
+using System;
+using System.Collections.Generic;
+using System.Diagnostics.CodeAnalysis;
+using System.IO;
+using EpicGames.Core;
+using Microsoft.Extensions.Logging;
+
+namespace UnrealBuildTool
+{
+	/// <summary>
+	/// Specifies the action to take for a config line, as denoted by its prefix.
+	/// </summary>
+	public enum ConfigLineAction
+	{
+		/// <summary>
+		/// Assign the value to the key
+		/// </summary>
+		Set,
+
+		/// <summary>
+		/// Add the value to the key (denoted with +X=Y in config files)
+		/// </summary>
+		Add,
+
+		/// <summary>
+		/// Remove the key without having to match value (denoted with !X in config files)
+		/// </summary>
+		RemoveKey,
+
+		/// <summary>
+		/// Remove the matching key and value (denoted with -X=Y in config files)
+		/// </summary>
+		RemoveKeyValue
+
+	}
+
+	/// <summary>
+	/// Contains a pre-parsed raw config line, consisting of action, key and value components.
+	/// </summary>
+	public class ConfigLine
+	{
+		/// <summary>
+		/// The action to take when merging this key/value pair with an existing value
+		/// </summary>
+		public ConfigLineAction Action;
+
+		/// <summary>
+		/// Name of the key to modify
+		/// </summary>
+		public string Key;
+
+		/// <summary>
+		/// Value to assign to the key
+		/// </summary>
+		public string Value;
+
+		/// <summary>
+		/// Whether to preserve enclosing quotations around <see cref="Value"/>
+		/// </summary>
+		internal bool ShouldPreserveValueQuotations = false;
+
+		/// <summary>
+		/// Constructor.
+		/// </summary>
+		/// <param name="Action">Action to take when merging this key/value pair with an existing value</param>
+		/// <param name="Key">Name of the key to modify</param>
+		/// <param name="Value">Value to assign</param>
+		public ConfigLine(ConfigLineAction Action, string Key, string Value)
+		{
+			this.Action = Action;
+			this.Key = Key;
+			this.Value = Value;
+		}
+
+		/// <summary>
+		/// Constructor.
+		/// </summary>
+		/// <param name="Action">Action to take when merging this key/value pair with an existing value</param>
+		/// <param name="Key">Name of the key to modify</param>
+		/// <param name="Value">Value to assign</param>
+		/// <param name="ShouldPreserveValueQuotations">Whether to preserve enclosing quotations around <see cref="Value"/></param>
+		public ConfigLine(ConfigLineAction Action, string Key, string Value, bool ShouldPreserveValueQuotations) : this(Action, Key, Value)
+		{
+			this.ShouldPreserveValueQuotations = ShouldPreserveValueQuotations;
+		}
+
+		/// <summary>
+		/// Formats this object for the debugger
+		/// </summary>
+		/// <returns>The original config line</returns>
+		public override string ToString()
+		{
+			string Prefix = (Action == ConfigLineAction.Add) ? "+" : (Action == ConfigLineAction.RemoveKey) ? "!" : (Action == ConfigLineAction.RemoveKeyValue) ? "-" : "";
+			return String.Format("{0}{1}={2}", Prefix, Key, Value);
+		}
+	}
+
+	/// <summary>
+	/// Contains the lines which appeared in a config section, with all comments and whitespace removed
+	/// </summary>
+	public class ConfigFileSection
+	{
+		/// <summary>
+		/// Name of the section
+		/// </summary>
+		public string Name;
+
+		/// <summary>
+		/// Lines which appeared in the config section
+		/// </summary>
+		public List<ConfigLine> Lines = new List<ConfigLine>();
+
+		/// <summary>
+		/// Construct an empty config section with the given name
+		/// </summary>
+		/// <param name="Name">Name of the config section</param>
+		public ConfigFileSection(string Name)
+		{
+			this.Name = Name;
+		}
+
+		/// <summary>
+		/// try to get a line using it's name, if the line doesn't exist returns false
+		/// </summary>
+		/// <param name="Name">Name of the line you want to get</param>
+		/// <param name="OutLine">The result of the operation</param>
+		/// <returns>return true if the line is retrieved return false and null OutLine if Name isn't found in this section</returns>
+		public bool TryGetLine(string Name, [NotNullWhen(true)] out ConfigLine? OutLine)
+		{
+			foreach (ConfigLine Line in Lines)
+			{
+				if (Line.Key.Equals(Name, StringComparison.Ordinal))
+				{
+					OutLine = Line;
+					return true;
+				}
+			}
+			OutLine = null;
+			return false;
+		}
+	}
+
+	/// <summary>
+	/// Represents a single config file as stored on disk. 
+	/// </summary>
+	public class ConfigFile
+	{
+		/// <summary>
+		/// Maps names to config sections
+		/// </summary>
+		readonly Dictionary<string, ConfigFileSection> Sections = new Dictionary<string, ConfigFileSection>(StringComparer.InvariantCultureIgnoreCase);
+
+		// Remap of config names/sections
+		private static readonly Dictionary<string, string> SectionNameRemap = new();
+		private static readonly Dictionary<string, Dictionary<string, string>> SectionKeyRemap = new();
+		private static readonly HashSet<string> WarnedKeys = new(StringComparer.InvariantCultureIgnoreCase);
+
+		private static string RemapSectionOrKey(Dictionary<string, string>? Remap, string Key, string Context)
+		{
+			if (Remap != null)
+			{
+				string? Remapped;
+				if (Remap.TryGetValue(Key, out Remapped))
+				{
+					if (!WarnedKeys.Contains(Key))
+					{
+						WarnedKeys.Add(Key);
+						Log.Logger.LogWarning($"DEPRECATION: '{Key}', {Context}, has been deprecated. Using '{Remapped}' instead. It is recommended you update your .ini files as soon as possible, and replace {Key} with {Remapped}");
+					}
+					return Remapped;
+				}
+			}
+			return Key;
+		}
+
+		/// <summary>
+		/// Static (class) constructor which is called before any ConfigFile instance is created
+		/// </summary>
+		static ConfigFile()
+		{
+			Dictionary<string, ConfigFileSection> Sections = new(StringComparer.InvariantCultureIgnoreCase);
+			try
+			{
+				// read the special ConfigRemap.ini file into sections
+				FileReference ConfigRemapFile = FileReference.Combine(UnrealBuildBase.Unreal.EngineDirectory, "Config", "ConfigRedirects.ini");
+				if (FileReference.Exists(ConfigRemapFile))
+				{
+					ReadIntoSections(ConfigRemapFile, Sections, ConfigLineAction.Set);
+				}
+			}
+			catch (Exception)
+			{
+				// Make ConfigFile when EngineDirectory is unknown a warning since ConfigRemapFile cannot be read in this case; e.g. Assemblies outside Engine that depend on ConfigFile
+				Log.Logger.LogWarning("Failed to read ConfigRemapFile into Sections");
+			}
+
+			// walk over the sections, where all but the special SectionNameRemap section is a section of keys to remap in that same section
+			foreach (KeyValuePair<string, ConfigFileSection> Pair in Sections)
+			{
+				// remember a remap for section names
+				if (Pair.Key.Equals("SectionNameRemap", StringComparison.OrdinalIgnoreCase))
+				{
+					foreach (ConfigLine Line in Pair.Value.Lines)
+					{
+						SectionNameRemap.Add(Line.Key, Line.Value);
+					}
+				}
+				else
+				{
+					// any other section is rmembered by the section name here, and each key/value pair is a remap for the given section
+					Dictionary<string, string> KeyRemap = new(StringComparer.OrdinalIgnoreCase);
+					SectionKeyRemap.Add(Pair.Key, KeyRemap);
+					foreach (ConfigLine Line in Pair.Value.Lines)
+					{
+						KeyRemap.Add(Line.Key, Line.Value);
+					}
+				}
+			}
+		}
+
+		/// <summary>
+		/// Constructs a new, empty config file
+		/// </summary>
+		/// <param name="DefaultAction">The default action to take when encountering arrays without a '+' prefix</param>
+		public ConfigFile(ConfigLineAction DefaultAction = ConfigLineAction.Set)
+		{
+		}
+
+		/// <summary>
+		/// Reads config data from the given file.
+		/// </summary>
+		/// <param name="Location">File to read from</param>
+		/// <param name="DefaultAction">The default action to take when encountering arrays without a '+' prefix</param>
+		public ConfigFile(FileReference Location, ConfigLineAction DefaultAction = ConfigLineAction.Set)
+		{
+			ReadIntoSections(Location, Sections, DefaultAction);
+		}
+
+		private static void ReadIntoSections(FileReference Location, Dictionary<string, ConfigFileSection> Sections, ConfigLineAction DefaultAction)
+		{
+			using (StreamReader Reader = new StreamReader(Location.FullName))
+			{
+				ConfigFileSection? CurrentSection = null;
+				Dictionary<string, string>? CurrentRemap = null;
+
+				for (; ; )
+				{
+					string? Line = Reader.ReadLine();
+					if (Line == null)
+					{
+						break;
+					}
+
+					// Find the first non-whitespace character
+					for (int StartIdx = 0; StartIdx < Line.Length; StartIdx++)
+					{
+						if (Line[StartIdx] != ' ' && Line[StartIdx] != '\t')
+						{
+							// Find the last non-whitespace character. If it's an escaped newline, merge the following line with it.
+							int EndIdx = Line.Length;
+							while (EndIdx > StartIdx)
+							{
+								if (Line[EndIdx - 1] == '\\')
+								{
+									string? NextLine = Reader.ReadLine();
+									if (NextLine == null)
+									{
+										break;
+									}
+									Line += NextLine;
+									EndIdx = Line.Length;
+									continue;
+								}
+								if (Line[EndIdx - 1] != ' ' && Line[EndIdx - 1] != '\t')
+								{
+									break;
+								}
+
+								EndIdx--;
+							}
+
+							// Break out if we've got a comment
+							if (Line[StartIdx] == ';')
+							{
+								break;
+							}
+							if (Line[StartIdx] == '/' && StartIdx + 1 < Line.Length && Line[StartIdx + 1] == '/')
+							{
+								break;
+							}
+
+							// Check if it's the start of a new section
+							if (Line[StartIdx] == '[')
+							{
+								CurrentSection = null;
+								if (Line[EndIdx - 1] == ']')
+								{
+									string SectionName = Line.Substring(StartIdx + 1, EndIdx - StartIdx - 2);
+
+									// lookup remaps
+									SectionName = RemapSectionOrKey(SectionNameRemap, SectionName, $"which is a config section in '{Location.FullName}'");
+									SectionKeyRemap.TryGetValue(SectionName, out CurrentRemap);
+
+									if (!Sections.TryGetValue(SectionName, out CurrentSection))
+									{
+										CurrentSection = new ConfigFileSection(SectionName);
+										Sections.Add(SectionName, CurrentSection);
+									}
+								}
+								break;
+							}
+
+							// Otherwise add it to the current section or add a new one
+							if (CurrentSection != null)
+							{
+								TryAddConfigLine(CurrentSection, CurrentRemap, Location.FullName, Line, StartIdx, EndIdx, DefaultAction, Sections);
+								break;
+							}
+
+							// Otherwise just ignore it
+							break;
+						}
+					}
+				}
+			}
+		}
+
+		/// <summary>
+		/// Reads config data from the given string.
+		/// </summary>
+		/// <param name="IniText">Single line string of config settings in the format [Section1]:Key1=Value1,[Section2]:Key2=Value2</param>
+		/// <param name="DefaultAction">The default action to take when encountering arrays without a '+' prefix</param>
+		public ConfigFile(string IniText, ConfigLineAction DefaultAction = ConfigLineAction.Set)
+		{
+
+			// Break into individual settings of the form [Section]:Key=Value
+			string[] SettingLines = IniText.Split(new char[] { ',' });
+			Dictionary<string, string>? CurrentRemap;
+			foreach (string Setting in SettingLines)
+			{
+				// Locate and break off the section name
+				string SectionName = Setting.Remove(Setting.IndexOf(':')).Trim(new char[] { '[', ']' });
+				// lookup remaps
+				RemapSectionOrKey(SectionNameRemap, SectionName, $"which is a config section (found in a string in code - search your .ini files for it, the source file is unknown)");
+				SectionKeyRemap.TryGetValue(SectionName, out CurrentRemap);
+
+				if (SectionName.Length > 0)
+				{
+					if (!Sections.TryGetValue(SectionName, out ConfigFileSection? CurrentSection))
+					{
+						CurrentSection = new ConfigFileSection(SectionName);
+						Sections.Add(SectionName, CurrentSection);
+					}
+
+					if (CurrentSection != null)
+					{
+						string IniKeyValue = Setting.Substring(Setting.IndexOf(':') + 1);
+						TryAddConfigLine(CurrentSection, CurrentRemap, "unknown source file", IniKeyValue, 0, IniKeyValue.Length, DefaultAction, Sections);
+					}
+				}
+			}
+		}
+
+		/// <summary>
+		/// Try to parse a key/value pair from the given line, and add it to a config section
+		/// </summary>
+		/// <param name="Section">The section to receive the parsed config line</param>
+		/// <param name="KeyRemap">Optional map to change names of keys on read</param> 
+		/// <param name="Filename">Optional map to change names of keys on read</param> 
+		/// <param name="Line">Text to parse</param>
+		/// <param name="StartIdx">Index of the first non-whitespace character in this line</param>
+		/// <param name="EndIdx">Index of the last (exclusive) non-whitespace character in this line</param>
+		/// <param name="DefaultAction">The default action to take if '+' or '-' is not specified on a config line</param>
+		/// <param name="Sections">The sections to find the redirected section for receiving the config line</param>
+		/// <returns>True if the line was parsed correctly, false otherwise</returns>
+		static bool TryAddConfigLine(ConfigFileSection Section, Dictionary<string, string>? KeyRemap, string Filename, string Line, int StartIdx, int EndIdx, ConfigLineAction DefaultAction, Dictionary<string, ConfigFileSection> Sections)
+		{
+			// Find the '=' character separating key and value
+			int EqualsIdx = Line.IndexOf('=', StartIdx, EndIdx - StartIdx);
+			if (EqualsIdx == -1 && Line[StartIdx] != '!')
+			{
+				return false;
+			}
+
+			// Keep track of the start of the key name
+			int KeyStartIdx = StartIdx;
+
+			// Remove the +/-/! prefix, if present
+			ConfigLineAction Action = DefaultAction;
+			if (Line[KeyStartIdx] == '+' || Line[KeyStartIdx] == '-' || Line[KeyStartIdx] == '!')
+			{
+				Action = (Line[KeyStartIdx] == '+') ? ConfigLineAction.Add : (Line[KeyStartIdx] == '!') ? ConfigLineAction.RemoveKey : ConfigLineAction.RemoveKeyValue;
+				KeyStartIdx++;
+				while (Line[KeyStartIdx] == ' ' || Line[KeyStartIdx] == '\t')
+				{
+					KeyStartIdx++;
+				}
+			}
+
+			// RemoveKey actions do not require a value
+			if (Action == ConfigLineAction.RemoveKey && EqualsIdx == -1)
+			{
+				Section.Lines.Add(new ConfigLine(Action, Line.Substring(KeyStartIdx).Trim(), ""));
+				return true;
+			}
+
+			// Remove trailing spaces after the name of the key
+			int KeyEndIdx = EqualsIdx;
+			for (; KeyEndIdx > KeyStartIdx; KeyEndIdx--)
+			{
+				if (Line[KeyEndIdx - 1] != ' ' && Line[KeyEndIdx - 1] != '\t')
+				{
+					break;
+				}
+			}
+
+			// Make sure there's a non-empty key name
+			if (KeyStartIdx == EqualsIdx)
+			{
+				return false;
+			}
+
+			// Skip whitespace between the '=' and the start of the value
+			int ValueStartIdx = EqualsIdx + 1;
+			for (; ValueStartIdx < EndIdx; ValueStartIdx++)
+			{
+				if (Line[ValueStartIdx] != ' ' && Line[ValueStartIdx] != '\t')
+				{
+					break;
+				}
+			}
+
+			bool HasStrippedValueQuotes = false;
+
+			// Strip quotes around the value if present
+			int ValueEndIdx = EndIdx;
+			if (ValueEndIdx >= ValueStartIdx + 2 && Line[ValueStartIdx] == '"' && Line[ValueEndIdx - 1] == '"')
+			{
+				HasStrippedValueQuotes = true;
+				ValueStartIdx++;
+				ValueEndIdx--;
+			}
+
+			// Add it to the config section
+			string Key = Line.Substring(KeyStartIdx, KeyEndIdx - KeyStartIdx);
+			string Value = Line.Substring(ValueStartIdx, ValueEndIdx - ValueStartIdx);
+
+			// remap the key if needed
+			string NewKey = RemapSectionOrKey(KeyRemap, Key, $"which is a config key in section [{Section.Name}], in '{Filename}'");
+
+			// look for a section:name remap
+			if (!NewKey.Equals(Key, StringComparison.Ordinal) && NewKey.Contains(':', StringComparison.Ordinal))
+			{
+				string SectionName = NewKey.Substring(0, NewKey.IndexOf(':'));
+				ConfigFileSection? CurrentSection;
+				if (!Sections.TryGetValue(SectionName, out CurrentSection))
+				{
+					CurrentSection = new ConfigFileSection(SectionName);
+					Sections.Add(SectionName, CurrentSection);
+				}
+
+				string KeyName = NewKey.Substring(NewKey.IndexOf(':') + 1);
+
+				CurrentSection.Lines.Add(new ConfigLine(Action, KeyName, Value, HasStrippedValueQuotes));
+
+				return true;
+			}
+
+			Section.Lines.Add(new ConfigLine(Action, NewKey, Value, HasStrippedValueQuotes));
+			return true;
+		}
+
+		/// <summary>
+		/// Names of sections in this file
+		/// </summary>
+		public IEnumerable<string> SectionNames => Sections.Keys;
+
+		/// <summary>
+		/// Tries to get a config section by name, or creates one if it doesn't exist
+		/// </summary>
+		/// <param name="SectionName">Name of the section to look for</param>
+		/// <returns>The config section</returns>
+		public ConfigFileSection FindOrAddSection(string SectionName)
+		{
+			ConfigFileSection? Section;
+			if (!Sections.TryGetValue(SectionName, out Section))
+			{
+				Section = new ConfigFileSection(SectionName);
+				Sections.Add(SectionName, Section);
+			}
+			return Section;
+		}
+
+		/// <summary>
+		/// Tries to get a config section by name
+		/// </summary>
+		/// <param name="SectionName">Name of the section to look for</param>
+		/// <param name="RawSection">On success, the config section that was found</param>
+		/// <returns>True if the section was found, false otherwise</returns>
+		public bool TryGetSection(string SectionName, [NotNullWhen(true)] out ConfigFileSection? RawSection)
+		{
+			return Sections.TryGetValue(SectionName, out RawSection);
+		}
+
+		/// <summary>
+		/// Write the config file out to the given location. Useful for debugging.
+		/// </summary>
+		/// <param name="Location">The file to write</param>
+		public void Write(FileReference Location)
+		{
+			using (StreamWriter Writer = new StreamWriter(Location.FullName))
+			{
+				foreach (ConfigFileSection Section in Sections.Values)
+				{
+					Writer.WriteLine("[{0}]", Section.Name);
+					foreach (ConfigLine Line in Section.Lines)
+					{
+						if (Line.ShouldPreserveValueQuotations)
+						{
+							Line.Value = $"\"{Line.Value}\"";
+						}
+						Writer.WriteLine("{0}", Line.ToString());
+					}
+					Writer.WriteLine();
+				}
+			}
+		}
+	}
+}
diff --git a/Engine/Source/Programs/UnrealBuildTool/Configuration/Ini/ConfigFileAttribute.cs b/Engine/Source/Programs/UnrealBuildTool/Configuration/Ini/ConfigFileAttribute.cs
new file mode 100644
--- /dev/null
+++ b/Engine/Source/Programs/UnrealBuildTool/Configuration/Ini/ConfigFileAttribute.cs
@@ -0,0 +1,41 @@
+// Copyright Epic Games, Inc. All Rights Reserved.
+
+using System;
+
+namespace UnrealBuildTool
+{
+	/// <summary>
+	/// Attribute indicating a value which should be populated from a UE .ini config file
+	/// </summary>
+	[AttributeUsage(AttributeTargets.Field | AttributeTargets.Property)]
+	public sealed class ConfigFileAttribute : Attribute
+	{
+		/// <summary>
+		/// Name of the config hierarchy to read from
+		/// </summary>
+		public ConfigHierarchyType ConfigType { get; }
+
+		/// <summary>
+		/// Section containing the setting
+		/// </summary>
+		public string SectionName { get; }
+
+		/// <summary>
+		/// Key name to search for
+		/// </summary>
+		public string? KeyName { get; }
+
+		/// <summary>
+		/// Constructor
+		/// </summary>
+		/// <param name="configType">Type of the config hierarchy to read from</param>
+		/// <param name="sectionName">Section containing the setting</param>
+		/// <param name="keyName">Key name to search for. Optional; uses the name of the field if not set.</param>
+		public ConfigFileAttribute(ConfigHierarchyType configType, string sectionName, string? keyName = null)
+		{
+			ConfigType = configType;
+			SectionName = sectionName;
+			KeyName = keyName;
+		}
+	}
+}
diff --git a/Engine/Source/Programs/UnrealBuildTool/Configuration/Ini/ConfigHierarchy.cs b/Engine/Source/Programs/UnrealBuildTool/Configuration/Ini/ConfigHierarchy.cs
new file mode 100644
--- /dev/null
+++ b/Engine/Source/Programs/UnrealBuildTool/Configuration/Ini/ConfigHierarchy.cs
@@ -0,0 +1,1514 @@
+// Copyright Epic Games, Inc. All Rights Reserved.
+
+using System;
+using System.Collections.Concurrent;
+using System.Collections.Generic;
+using System.Diagnostics.CodeAnalysis;
+using System.Linq;
+using System.Text;
+using System.Text.RegularExpressions;
+using EpicGames.Core;
+using UnrealBuildBase;
+
+namespace UnrealBuildTool
+{
+	/// <summary>
+	/// Types of config file hierarchy
+	/// </summary>
+	public enum ConfigHierarchyType
+	{
+		/// <summary>
+		/// BaseGame.ini, DefaultGame.ini, etc...
+		/// </summary>
+		Game,
+
+		/// <summary>
+		/// BaseEngine.ini, DefaultEngine.ini, etc...
+		/// </summary>
+		Engine,
+
+		/// <summary>
+		/// BaseEditor.ini, DefaultEditor.ini, etc...
+		/// </summary>
+		Editor,
+
+		/// <summary>
+		/// BaseEditorPerProjectUserSettings.ini, DefaultEditorPerProjectUserSettings.ini, etc..
+		/// </summary>
+		EditorPerProjectUserSettings,
+
+		/// <summary>
+		/// BaseEncryption.ini, DefaultEncryption.ini, etc..
+		/// </summary>
+		Encryption,
+
+		/// <summary>
+		/// BaseCrypto.ini, DefaultCrypto.ini, etc..
+		/// </summary>
+		Crypto,
+
+		/// <summary>
+		/// BaseEditorSettings.ini, DefaultEditorSettings.ini, etc...
+		/// </summary>
+		EditorSettings,
+
+		/// <summary>
+		/// BaseInstallBundle.ini, DefaultInstallBundle.ini, etc...
+		/// </summary>
+		InstallBundle,
+
+		/// <summary>
+		/// BasePakFileRules.ini, DefaultPakFileRules.ini, etc, etc....
+		/// </summary>
+		PakFileRules,
+	}
+
+	/// <summary>
+	/// Stores a set of merged key/value pairs for a config section
+	/// </summary>
+	public class ConfigHierarchySection
+	{
+		/// <summary>
+		/// Map of key names to their values
+		/// </summary>
+		readonly Dictionary<string, List<string>> KeyToValue = new Dictionary<string, List<string>>(StringComparer.InvariantCultureIgnoreCase);
+
+		/// <summary>
+		/// Construct a merged config section from the given per-file config sections
+		/// </summary>
+		/// <param name="FileSections">Config sections from individual files</param>
+		public ConfigHierarchySection(IEnumerable<ConfigFileSection> FileSections)
+		{
+			foreach (ConfigFileSection FileSection in FileSections)
+			{
+				foreach (ConfigLine Line in FileSection.Lines)
+				{
+					if (Line.Action == ConfigLineAction.RemoveKey)
+					{
+						KeyToValue.Remove(Line.Key);
+						continue;
+					}
+
+					// Find or create the values for this key
+					List<string>? Values;
+
+					if (KeyToValue.TryGetValue(Line.Key, out Values))
+					{
+						// Update the existing list
+						if (Line.Action == ConfigLineAction.Set)
+						{
+							Values.Clear();
+							Values.Add(Line.Value);
+						}
+						else if (Line.Action == ConfigLineAction.Add)
+						{
+							Values.Add(Line.Value);
+						}
+						else if (Line.Action == ConfigLineAction.RemoveKeyValue)
+						{
+							Values.RemoveAll(x => x.Equals(Line.Value, StringComparison.OrdinalIgnoreCase));
+						}
+					}
+					else
+					{
+						// If it's a set or add action, create and add a new list
+						if (Line.Action == ConfigLineAction.Set || Line.Action == ConfigLineAction.Add)
+						{
+							Values = new List<string>();
+							Values.Add(Line.Value);
+							KeyToValue.Add(Line.Key, Values);
+						}
+					}
+				}
+			}
+		}
+
+		/// <summary>
+		/// Returns a list of key names
+		/// </summary>
+		public IEnumerable<string> KeyNames => KeyToValue.Keys;
+
+		/// <summary>
+		/// Tries to find the value for a given key
+		/// </summary>
+		/// <param name="KeyName">The key name to search for</param>
+		/// <param name="Value">On success, receives the corresponding value</param>
+		/// <returns>True if the key was found, false otherwise</returns>
+		public bool TryGetValue(string KeyName, [NotNullWhen(true)] out string? Value)
+		{
+			List<string>? ValuesList;
+			if (KeyToValue.TryGetValue(KeyName, out ValuesList) && ValuesList.Count > 0)
+			{
+				Value = ValuesList[0];
+				return true;
+			}
+			else
+			{
+				Value = null;
+				return false;
+			}
+		}
+
+		/// <summary>
+		/// Tries to find the values for a given key
+		/// </summary>
+		/// <param name="KeyName">The key name to search for</param>
+		/// <param name="Values">On success, receives a list of the corresponding values</param>
+		/// <returns>True if the key was found, false otherwise</returns>
+		public bool TryGetValues(string KeyName, [NotNullWhen(true)] out IReadOnlyList<string>? Values)
+		{
+			List<string>? ValuesList;
+			if (KeyToValue.TryGetValue(KeyName, out ValuesList))
+			{
+				Values = ValuesList;
+				return true;
+			}
+			else
+			{
+				Values = null;
+				return false;
+			}
+		}
+	}
+
+	/// <summary>
+	/// Encapsulates a hierarchy of config files, merging sections from them together on request 
+	/// </summary>
+	public sealed class ConfigHierarchy : IDisposable
+	{
+		/// <summary>
+		/// Array of 
+		/// </summary>
+		readonly ConfigFile[] Files;
+
+		/// <summary>
+		/// Set of file references
+		/// </summary>
+		public IEnumerable<FileReference> FileReferences { get; }
+
+		/// <summary>
+		/// Cache of requested config sections
+		/// </summary>
+		readonly ConcurrentDictionary<string, Lazy<ConfigHierarchySection>> NameToSection = new(StringComparer.InvariantCultureIgnoreCase);
+
+		/// <summary>
+		/// Construct a config hierarchy from the given files
+		/// </summary>
+		/// <param name="Files">Set of files to include (in order)</param>
+		/// <param name="FileReferences">Set of file paths, for tracking</param>
+		public ConfigHierarchy(IEnumerable<ConfigFile> Files, IEnumerable<FileReference>? FileReferences = null)
+		{
+			this.Files = [.. Files];
+			this.FileReferences = FileReferences ?? [];
+		}
+
+		#region IDisposible
+		/// <inheritdoc/>
+		public void Dispose()
+		{
+		}
+		#endregion
+
+		/// <summary>
+		/// Names of all sections in all config files
+		/// </summary>
+		/// <returns></returns>
+		public HashSet<string> SectionNames => [.. Files.SelectMany(x => x.SectionNames)];
+
+		/// <summary>
+		/// Finds a config section with the given name
+		/// </summary>
+		/// <param name="SectionName">Name of the section to look for</param>
+		/// <returns>The merged config section</returns>
+		public ConfigHierarchySection FindSection(string SectionName)
+		{
+			return NameToSection.GetOrAdd(SectionName, new Lazy<ConfigHierarchySection>(() =>
+			{
+				// Find all the raw sections from the file hierarchy
+				List<ConfigFileSection> RawSections = new List<ConfigFileSection>();
+				foreach (ConfigFile File in Files)
+				{
+					ConfigFileSection? RawSection;
+					if (File.TryGetSection(SectionName, out RawSection))
+					{
+						RawSections.Add(RawSection);
+					}
+				}
+
+				// Merge them together and add it to the cache
+				return new ConfigHierarchySection(RawSections);
+			})).Value;
+		}
+
+		/// <summary>
+		/// Legacy function for ease of transition from ConfigCacheIni to ConfigHierarchy. Gets a bool with the given key name.
+		/// </summary>
+		/// <param name="SectionName">Section name</param>
+		/// <param name="KeyName">Key name</param>
+		/// <param name="Value">Value associated with the specified key. If the key has more than one value, only the first one is returned</param>
+		/// <returns>True if the key exists</returns>
+		public bool GetBool(string SectionName, string KeyName, out bool Value)
+		{
+			return TryGetValue(SectionName, KeyName, out Value);
+		}
+
+		/// <summary>
+		/// Legacy function for ease of transition from ConfigCacheIni to ConfigHierarchy. Gets an array with the given key name, returning null on failure.
+		/// </summary>
+		/// <param name="SectionName">Section name</param>
+		/// <param name="KeyName">Key name</param>
+		/// <param name="Values">Value associated with the specified key. If the key has more than one value, only the first one is returned</param>
+		/// <returns>True if the key exists</returns>
+		public bool GetArray(string SectionName, string KeyName, [NotNullWhen(true)] out List<string>? Values)
+		{
+			IReadOnlyList<string>? ValuesEnumerable;
+			if (TryGetValues(SectionName, KeyName, out ValuesEnumerable))
+			{
+				Values = ValuesEnumerable.ToList();
+				return true;
+			}
+			else
+			{
+				Values = null;
+				return false;
+			}
+		}
+
+		/// <summary>
+		/// Legacy function for ease of transition from ConfigCacheIni to ConfigHierarchy. Gets a string with the given key name, returning an empty string on failure.
+		/// </summary>
+		/// <param name="SectionName">Section name</param>
+		/// <param name="KeyName">Key name</param>
+		/// <param name="Value">Value associated with the specified key. If the key has more than one value, only the first one is returned</param>
+		/// <returns>True if the key exists</returns>
+		public bool GetString(string SectionName, string KeyName, out string Value)
+		{
+			string? RetrievedValue;
+			if (TryGetValue(SectionName, KeyName, out RetrievedValue))
+			{
+				Value = RetrievedValue;
+				return true;
+			}
+			else
+			{
+				Value = "";
+				return false;
+			}
+		}
+
+		/// <summary>
+		/// Legacy function for ease of transition from ConfigCacheIni to ConfigHierarchy. Gets an int with the given key name.
+		/// </summary>
+		/// <param name="SectionName">Section name</param>
+		/// <param name="KeyName">Key name</param>
+		/// <param name="Value">Value associated with the specified key. If the key has more than one value, only the first one is returned</param>
+		/// <returns>True if the key exists</returns>
+		public bool GetInt32(string SectionName, string KeyName, out int Value)
+		{
+			return TryGetValue(SectionName, KeyName, out Value);
+		}
+
+		/// <summary>
+		/// Gets a single string value associated with the specified key.
+		/// </summary>
+		/// <param name="SectionName">Section name</param>
+		/// <param name="KeyName">Key name</param>
+		/// <param name="Value">Value associated with the specified key. If the key has more than one value, only the first one is returned</param>
+		/// <returns>True if the key exists</returns>
+		public bool TryGetValue(string SectionName, string KeyName, [NotNullWhen(true)] out string? Value)
+		{
+			return FindSection(SectionName).TryGetValue(KeyName, out Value);
+		}
+
+		/// <summary>
+		/// Gets a single bool value associated with the specified key.
+		/// </summary>
+		/// <param name="SectionName">Section name</param>
+		/// <param name="KeyName">Key name</param>
+		/// <param name="Value">Value associated with the specified key. If the key has more than one value, only the first one is returned</param>
+		/// <returns>True if the key exists</returns>
+		public bool TryGetValue(string SectionName, string KeyName, out bool Value)
+		{
+			string? Text;
+			if (!TryGetValue(SectionName, KeyName, out Text))
+			{
+				Value = false;
+				return false;
+			}
+			return TryParse(Text, out Value);
+		}
+
+		/// <summary>
+		/// Gets a single Int32 value associated with the specified key.
+		/// </summary>
+		/// <param name="SectionName">Section name</param>
+		/// <param name="KeyName">Key name</param>
+		/// <param name="Value">Value associated with the specified key. If the key has more than one value, only the first one is returned</param>
+		/// <returns>True if the key exists</returns>
+		public bool TryGetValue(string SectionName, string KeyName, out int Value)
+		{
+			string? Text;
+			if (!TryGetValue(SectionName, KeyName, out Text))
+			{
+				Value = 0;
+				return false;
+			}
+			return TryParse(Text, out Value);
+		}
+
+		/// <summary>
+		/// Gets a single GUID value associated with the specified key.
+		/// </summary>
+		/// <param name="SectionName">Section name</param>
+		/// <param name="KeyName">Key name</param>
+		/// <param name="Value">Value associated with the specified key. If the key has more than one value, only the first one is returned</param>
+		/// <returns>True if the key exists</returns>
+		public bool TryGetValue(string SectionName, string KeyName, out Guid Value)
+		{
+			string? Text;
+			if (!TryGetValue(SectionName, KeyName, out Text))
+			{
+				Value = Guid.Empty;
+				return false;
+			}
+			return TryParse(Text, out Value);
+		}
+
+		/// <summary>
+		/// Gets a single-precision floating point value associated with the specified key.
+		/// </summary>
+		/// <param name="SectionName">Section name</param>
+		/// <param name="KeyName">Key name</param>
+		/// <param name="Value">Value associated with the specified key. If the key has more than one value, only the first one is returned</param>
+		/// <returns>True if the key exists</returns>
+		public bool TryGetValue(string SectionName, string KeyName, out float Value)
+		{
+			string? Text;
+			if (!TryGetValue(SectionName, KeyName, out Text))
+			{
+				Value = 0;
+				return false;
+			}
+			return TryParse(Text, out Value);
+		}
+
+		/// <summary>
+		/// Gets a double-precision floating point value associated with the specified key.
+		/// </summary>
+		/// <param name="SectionName">Section name</param>
+		/// <param name="KeyName">Key name</param>
+		/// <param name="Value">Value associated with the specified key. If the key has more than one value, only the first one is returned</param>
+		/// <returns>True if the key exists</returns>
+		public bool TryGetValue(string SectionName, string KeyName, out double Value)
+		{
+			string? Text;
+			if (!TryGetValue(SectionName, KeyName, out Text))
+			{
+				Value = 0;
+				return false;
+			}
+			return TryParse(Text, out Value);
+		}
+
+		/// <summary>
+		/// Gets an enumeration value associated with the specified key.
+		/// </summary>
+		/// <typeparam name="T"></typeparam>
+		/// <param name="SectionName">Section name</param>
+		/// <param name="KeyName">Key name</param>
+		/// <param name="Value">Value associated with the specified key. If the key has more than one value, only the first one is returned</param>
+		/// <returns>True if the key exists</returns>
+		public bool TryGetValue<T>(string SectionName, string KeyName, out T Value) where T : struct
+		{
+			string? Text;
+			if (!TryGetValue(SectionName, KeyName, out Text))
+			{
+				Value = default;
+				return false;
+			}
+			return Enum.TryParse<T>(Text, out Value);
+		}
+
+		/// <summary>
+		/// Gets all values associated with the specified key
+		/// </summary>
+		/// <param name="SectionName">Section where the key is located</param>
+		/// <param name="KeyName">Key name</param>
+		/// <param name="Values">Copy of the list containing all values associated with the specified key</param>
+		/// <returns>True if the key exists</returns>
+		public bool TryGetValues(string SectionName, string KeyName, [NotNullWhen(true)] out IReadOnlyList<string>? Values)
+		{
+			return FindSection(SectionName).TryGetValues(KeyName, out Values);
+		}
+
+		/// <summary>
+		/// Gets the value for the given type. Can return a full struct hierarchy.
+		/// </summary>
+		/// <param name="SectionName">Section where the key is located</param>
+		/// <param name="KeyName">Key name</param>
+		/// <param name="Value">Value hierarchy associated with the specified key. All field names must exist</param>
+		/// <returns>True if the key exists and could be parsed</returns>
+		public bool TryGetValueGeneric<T>(string SectionName, string KeyName, [NotNullWhen(true)] out T? Value) where T : new()
+		{
+			if (TryGetValue(SectionName, KeyName, out string? Line))
+			{
+				return ConfigValueParser.TryParseGeneric(Line, out Value);
+			}
+
+			Value = default;
+			return false;
+		}
+
+		/// <summary>
+		/// Gets the array of values for the given type. Can return a full struct hierarchy.
+		/// </summary>
+		/// <param name="SectionName">Section where the key is located</param>
+		/// <param name="KeyName">Key name</param>
+		/// <param name="Values">Array of values associated with the specified key. All field names must exist</param>
+		/// <returns>True if the key exists and could be parsed</returns>
+		public bool TryGetValuesGeneric<T>(string SectionName, string KeyName, [NotNullWhen(true)] out T[]? Values) where T : new()
+		{
+			if (TryGetValues(SectionName, KeyName, out IReadOnlyList<string>? Lines))
+			{
+				return ConfigValueParser.TryParseArrayGeneric(Lines.ToArray(), out Values);
+			}
+
+			Values = null;
+			return false;
+		}
+
+		/// <summary>
+		/// Parse a string as a boolean value
+		/// </summary>
+		/// <param name="Text">The text to parse</param>
+		/// <param name="Value">The parsed value, if successful</param>
+		/// <returns>True if the text was parsed, false otherwise</returns>
+		public static bool TryParse(string Text, out bool Value)
+		{
+			// C# Boolean type expects "False" or "True" but since we're not case sensitive, we need to suppor that manually
+			if (Text == "1" || Text.Equals("true", StringComparison.OrdinalIgnoreCase))
+			{
+				Value = true;
+				return true;
+			}
+			else if (Text == "0" || Text.Equals("false", StringComparison.OrdinalIgnoreCase))
+			{
+				Value = false;
+				return true;
+			}
+			else
+			{
+				Value = false;
+				return false;
+			}
+		}
+
+		/// <summary>
+		/// Parse a string as an integer value
+		/// </summary>
+		/// <param name="Text">The text to parse</param>
+		/// <param name="Value">The parsed value, if successful</param>
+		/// <returns>True if the text was parsed, false otherwise</returns>
+		public static bool TryParse(string Text, out int Value)
+		{
+			return Int32.TryParse(Text, out Value);
+		}
+
+		/// <summary>
+		/// Parse a string as a GUID value
+		/// </summary>
+		/// <param name="Text">The text to parse</param>
+		/// <param name="Value">The parsed value, if successful</param>
+		/// <returns>True if the text was parsed, false otherwise</returns>
+		public static bool TryParse(string Text, out Guid Value)
+		{
+			if (Text.Contains("A=") && Text.Contains("B=") && Text.Contains("C=") && Text.Contains("D="))
+			{
+				char[] Separators = new char[] { '(', ')', '=', ',', ' ', 'A', 'B', 'C', 'D' };
+				string[] ComponentValues = Text.Split(Separators, StringSplitOptions.RemoveEmptyEntries);
+				if (ComponentValues.Length == 4)
+				{
+					StringBuilder HexString = new StringBuilder();
+					for (int ComponentIndex = 0; ComponentIndex < 4; ComponentIndex++)
+					{
+						int IntegerValue;
+						if (!Int32.TryParse(ComponentValues[ComponentIndex], out IntegerValue))
+						{
+							Value = Guid.Empty;
+							return false;
+						}
+						HexString.Append(IntegerValue.ToString("X8"));
+					}
+					Text = HexString.ToString();
+				}
+			}
+			return Guid.TryParseExact(Text, "N", out Value);
+		}
+
+		/// <summary>
+		/// Parse a string as a single-precision floating point value
+		/// </summary>
+		/// <param name="Text">The text to parse</param>
+		/// <param name="Value">The parsed value, if successful</param>
+		/// <returns>True if the text was parsed, false otherwise</returns>
+		public static bool TryParse(string Text, out float Value)
+		{
+			if (Text.EndsWith("f"))
+			{
+				return Single.TryParse(Text.Substring(0, Text.Length - 1), out Value);
+			}
+			else
+			{
+				return Single.TryParse(Text, out Value);
+			}
+		}
+
+		/// <summary>
+		/// Parse a string as a double-precision floating point value
+		/// </summary>
+		/// <param name="Text">The text to parse</param>
+		/// <param name="Value">The parsed value, if successful</param>
+		/// <returns>True if the text was parsed, false otherwise</returns>
+		public static bool TryParse(string Text, out double Value)
+		{
+			if (Text.EndsWith("f"))
+			{
+				return Double.TryParse(Text.Substring(0, Text.Length - 1), out Value);
+			}
+			else
+			{
+				return Double.TryParse(Text, out Value);
+			}
+		}
+
+		/// <summary>
+		/// Attempts to parse the given line as a UE config object (eg. (Name="Foo",Number=1234)).
+		/// </summary>
+		/// <param name="Line">Line of text to parse</param>
+		/// <param name="Properties">Receives key/value pairs for the config object</param>
+		/// <returns>True if an object was parsed, false otherwise</returns>
+		public static bool TryParse(string Line, [NotNullWhen(true)] out Dictionary<string, string>? Properties)
+		{
+			// Convert the string to a zero-terminated array, to make parsing easier.
+			char[] Chars = new char[Line.Length + 1];
+			Line.CopyTo(0, Chars, 0, Line.Length);
+
+			// Get the opening paren
+			int Idx = 0;
+			while (Char.IsWhiteSpace(Chars[Idx]))
+			{
+				Idx++;
+			}
+			if (Chars[Idx] != '(')
+			{
+				Properties = null;
+				return false;
+			}
+
+			// Read to the next token
+			Idx++;
+			while (Char.IsWhiteSpace(Chars[Idx]))
+			{
+				Idx++;
+			}
+
+			// Create the dictionary to receive the new properties
+			Dictionary<string, string> NewProperties = new Dictionary<string, string>();
+
+			// Read a sequence of key/value pairs
+			StringBuilder Value = new StringBuilder();
+			if (Chars[Idx] != ')')
+			{
+				for (; ; )
+				{
+					// Find the end of the name
+					int NameIdx = Idx;
+					while (Char.IsLetterOrDigit(Chars[Idx]) || Chars[Idx] == '_')
+					{
+						Idx++;
+					}
+					if (Idx == NameIdx)
+					{
+						Properties = null;
+						return false;
+					}
+
+					// Extract the key string, and make sure it hasn't already been added
+					string Key = new string(Chars, NameIdx, Idx - NameIdx);
+					if (NewProperties.ContainsKey(Key))
+					{
+						Properties = null;
+						return false;
+					}
+
+					// Consume the equals character
+					while (Char.IsWhiteSpace(Chars[Idx]))
+					{
+						Idx++;
+					}
+					if (Chars[Idx] != '=')
+					{
+						Properties = null;
+						return false;
+					}
+
+					// Move to the value
+					Idx++;
+					while (Char.IsWhiteSpace(Chars[Idx]))
+					{
+						Idx++;
+					}
+
+					// Parse the value
+					Value.Clear();
+					if (Char.IsLetterOrDigit(Chars[Idx]) || Chars[Idx] == '_' || Chars[Idx] == '-')
+					{
+						while (Char.IsLetterOrDigit(Chars[Idx]) || Chars[Idx] == '_' || Chars[Idx] == '-' || Chars[Idx] == '.')
+						{
+							Value.Append(Chars[Idx]);
+							Idx++;
+						}
+					}
+					else if (Chars[Idx] == '\"')
+					{
+						Idx++;
+						for (; Chars[Idx] != '\"'; Idx++)
+						{
+							if (Chars[Idx] == '\0')
+							{
+								Properties = null;
+								return false;
+							}
+							else
+							{
+								Value.Append(Chars[Idx]);
+							}
+						}
+						Idx++;
+					}
+					else if (Chars[Idx] == '(')
+					{
+						Value.Append(Chars[Idx++]);
+
+						bool bInQuotes = false;
+						for (int Nesting = 1; Nesting > 0; Idx++)
+						{
+							if (Chars[Idx] == '\0')
+							{
+								Properties = null;
+								return false;
+							}
+							else if (Chars[Idx] == '(' && !bInQuotes)
+							{
+								Nesting++;
+							}
+							else if (Chars[Idx] == ')' && !bInQuotes)
+							{
+								Nesting--;
+							}
+							else if (Chars[Idx] == '\"' || Chars[Idx] == '\'')
+							{
+								bInQuotes ^= true;
+							}
+							Value.Append(Chars[Idx]);
+						}
+					}
+					else if (Chars[Idx] != ')' && Chars[Idx] != ',')
+					{
+						Properties = null;
+						return false;
+					}
+
+					// Extract the value string
+					NewProperties[Key] = Value.ToString();
+
+					// Move to the separator
+					while (Char.IsWhiteSpace(Chars[Idx]))
+					{
+						Idx++;
+					}
+					if (Chars[Idx] == ')')
+					{
+						break;
+					}
+					if (Chars[Idx] != ',')
+					{
+						Properties = null;
+						return false;
+					}
+
+					// Move to the next field
+					Idx++;
+					while (Char.IsWhiteSpace(Chars[Idx]))
+					{
+						Idx++;
+					}
+				}
+			}
+
+			// Make sure we're at the end of the string
+			Idx++;
+			while (Char.IsWhiteSpace(Chars[Idx]))
+			{
+				Idx++;
+			}
+			if (Chars[Idx] != '\0')
+			{
+				Properties = null;
+				return false;
+			}
+
+			Properties = NewProperties;
+			return true;
+		}
+
+		/// <summary>
+		/// Attempts to parse the given line as a UE config array (eg. ("one", "two", "three") ).
+		/// </summary>
+		/// <param name="Line">Line of text to parse</param>
+		/// <param name="Array">Receives array for the config array</param>
+		/// <returns>True if an array was parsed, false otherwise</returns>
+		public static bool TryParse(string Line, [NotNullWhen(true)] out string[]? Array)
+		{
+			// Convert the string to a zero-terminated array, to make parsing easier.
+			char[] Chars = new char[Line.Length + 1];
+			Line.CopyTo(0, Chars, 0, Line.Length);
+
+			// Get the opening paren
+			int Idx = 0;
+			while (Char.IsWhiteSpace(Chars[Idx]))
+			{
+				Idx++;
+			}
+			if (Chars[Idx] != '(')
+			{
+				Array = null;
+				return false;
+			}
+
+			// Read to the next token
+			Idx++;
+			while (Char.IsWhiteSpace(Chars[Idx]))
+			{
+				Idx++;
+			}
+
+			// Create the list to receive the new items
+			List<string> NewArray = new List<string>();
+
+			// Read a sequence items
+			StringBuilder Value = new StringBuilder();
+			if (Chars[Idx] != ')')
+			{
+				for (; ; )
+				{
+					// Skip whitespace
+					while (Char.IsWhiteSpace(Chars[Idx]))
+					{
+						Idx++;
+					}
+
+					// Parse the value
+					Value.Clear();
+					if (Char.IsLetterOrDigit(Chars[Idx]) || Chars[Idx] == '_' || Chars[Idx] == '-')
+					{
+						while (Char.IsLetterOrDigit(Chars[Idx]) || Chars[Idx] == '_' || Chars[Idx] == '-' || Chars[Idx] == '.')
+						{
+							Value.Append(Chars[Idx]);
+							Idx++;
+						}
+					}
+					else if (Chars[Idx] == '\"')
+					{
+						Idx++;
+						for (; Chars[Idx] != '\"'; Idx++)
+						{
+							if (Chars[Idx] == '\0')
+							{
+								Array = null;
+								return false;
+							}
+							else
+							{
+								Value.Append(Chars[Idx]);
+							}
+						}
+						Idx++;
+					}
+					else if (Chars[Idx] == '(')
+					{
+						Value.Append(Chars[Idx++]);
+
+						bool bInQuotes = false;
+						for (int Nesting = 1; Nesting > 0; Idx++)
+						{
+							if (Chars[Idx] == '\0')
+							{
+								Array = null;
+								return false;
+							}
+							else if (Chars[Idx] == '(' && !bInQuotes)
+							{
+								Nesting++;
+							}
+							else if (Chars[Idx] == ')' && !bInQuotes)
+							{
+								Nesting--;
+							}
+							else if (Chars[Idx] == '\"' || Chars[Idx] == '\'')
+							{
+								bInQuotes ^= true;
+							}
+							Value.Append(Chars[Idx]);
+						}
+					}
+					else
+					{
+						Array = null;
+						return false;
+					}
+
+					// Store the item
+					NewArray.Add(Value.ToString());
+
+					// Move to the separator
+					while (Char.IsWhiteSpace(Chars[Idx]))
+					{
+						Idx++;
+					}
+					if (Chars[Idx] == ')')
+					{
+						break;
+					}
+					if (Chars[Idx] != ',')
+					{
+						Array = null;
+						return false;
+					}
+
+					// Move to the next field
+					Idx++;
+					while (Char.IsWhiteSpace(Chars[Idx]))
+					{
+						Idx++;
+					}
+				}
+			}
+
+			// Make sure we're at the end of the string
+			Idx++;
+			while (Char.IsWhiteSpace(Chars[Idx]))
+			{
+				Idx++;
+			}
+			if (Chars[Idx] != '\0')
+			{
+				Array = null;
+				return false;
+			}
+
+			Array = NewArray.ToArray();
+			return true;
+		}
+
+		/// <summary>
+		/// Attempts to parse the given line as a UE map (eg. (("key1","value1"), ("key2","value2")).
+		/// </summary>
+		/// <param name="Line">Line of text to parse</param>
+		/// <param name="Map">Receives dictionary for the config map</param>
+		/// <returns>True if a map was parsed, false otherwise</returns>
+		public static bool TryParseAsMap(string Line, [NotNullWhen(true)] out Dictionary<string, string>? Map)
+		{
+			// read outer array
+			if (!TryParse(Line, out string[]? Array))
+			{
+				Map = null;
+				return false;
+			}
+
+			// read each pair - they're stored in the same way as an array of 2
+			Dictionary<string, string> NewMap = new Dictionary<string, string>();
+			foreach (string ArrayItem in Array)
+			{
+				if (!TryParse(ArrayItem, out string[]? Pairs) || Pairs.Length != 2)
+				{
+					Map = null;
+					return false;
+				}
+
+				NewMap[Pairs[0]] = Pairs[1];
+			}
+
+			Map = NewMap;
+			return true;
+		}
+
+		class ConfigLayerExpansion
+		{
+			// a set of replacements from the source file to possible other files
+			public string? Before1 = null;
+			public string? After1 = null;
+			public string? Before2 = null;
+			public string? After2 = null;
+		};
+
+		static readonly string[] ConfigLayers =
+		{
+			// Engine/Base.ini
+			"{ENGINE}/Config/Base.ini",
+			// Engine/Base*.ini
+ 			"{ENGINE}/Config/Base{TYPE}.ini",
+			// Engine/Platform/BasePlatform*.ini
+			"{ENGINE}/Config/{PLATFORM}/Base{PLATFORM}{TYPE}.ini",
+			// Project/Default*.ini
+			"{PROJECT}/Config/Default{TYPE}.ini",
+			// Project/Generated*.ini  this is reserved for files which are generated by buildmachine processes (i.e. should never be checked in)
+			"{PROJECT}/Config/Generated{TYPE}.ini",
+			// Project/Config/Custom/CustomConfig/Default*.ini only if CustomConfig is defined
+			"{PROJECT}/Config/Custom/{CUSTOMCONFIG}/Default{TYPE}.ini",
+			// Engine/Platform/Platform*.ini
+			"{ENGINE}/Config/{PLATFORM}/{PLATFORM}{TYPE}.ini",
+			// Project/Platform/Platform*.ini
+			"{PROJECT}/Config/{PLATFORM}/{PLATFORM}{TYPE}.ini",
+			// Project/Platform/GeneratedPlatform*.ini  this is reserved for files which are generated by buildmachine processes (i.e. should never be checked in)
+			"{PROJECT}/Config/{PLATFORM}/Generated{PLATFORM}{TYPE}.ini",
+			// Project/Platform/Custom/CustomConfig/Platform*.ini only if CustomConfig is defined
+			"{PROJECT}/Config/{PLATFORM}/Custom/{CUSTOMCONFIG}/{PLATFORM}{TYPE}.ini",
+			// AppSettings/.../System*.ini
+			"{APPSETTINGS}/Unreal Engine/Engine/Config/System{TYPE}.ini",
+			"{APPSETTINGS}/UnrealEngine/Engine/Config/System{TYPE}.ini",
+			// UserSettings/.../User*.ini
+			"{USERSETTINGS}/Unreal Engine/Engine/Config/User{TYPE}.ini",
+			"{USERSETTINGS}/UnrealEngine/Engine/Config/User{TYPE}.ini",
+			// UserDir/.../User*.ini
+			"{USER}/Unreal Engine/Engine/Config/User{TYPE}.ini",
+			"{USER}/UnrealEngine/Engine/Config/User{TYPE}.ini",
+			// Project/User*.ini
+			"{PROJECT}/Config/User{TYPE}.ini",
+		};
+
+		static readonly ConfigLayerExpansion[] ConfigLayerExpansions =
+		{
+			// The base expansion (ie, no expansion)
+			new ConfigLayerExpansion { }, 
+			// Restricted Locations
+			new ConfigLayerExpansion { Before1 = "{ENGINE}/", After1 = "{ENGINE}/Restricted/NotForLicensees/", Before2 = "{PROJECT}/Config/", After2 = "{RESTRICTEDPROJECT_NFL}/Config/" },
+			new ConfigLayerExpansion { Before1 = "{ENGINE}/", After1 = "{ENGINE}/Restricted/NoRedist/",         Before2 = "{PROJECT}/Config/", After2 = "{RESTRICTEDPROJECT_NR}/Config/" },
+			new ConfigLayerExpansion { Before1 = "{ENGINE}/", After1 = "{ENGINE}/Restricted/LimitedAccess/",    Before2 = "{PROJECT}/Config/", After2 = "{RESTRICTEDPROJECT_LA}/Config/" },
+			// Platform Extensions
+			new ConfigLayerExpansion { Before1 = "{ENGINE}/Config/{PLATFORM}/", After1 = "{EXTENGINE}/Config/",    Before2 = "{PROJECT}/Config/{PLATFORM}/", After2 = "{EXTPROJECT}/Config/" },
+			// Platform Extensions in Restricted Locations
+			new ConfigLayerExpansion { Before1 = "{ENGINE}/Config/{PLATFORM}/", After1 = "{ENGINE}/Restricted/NotForLicensees/Platforms/{PLATFORM}/Config/",   Before2 = "{PROJECT}/Config/{PLATFORM}/", After2 = "{RESTRICTEDPROJECT_NFL}/Platforms/{PLATFORM}/{OPT_SUBDIR}Config/" },
+			new ConfigLayerExpansion { Before1 = "{ENGINE}/Config/{PLATFORM}/", After1 = "{ENGINE}/Restricted/NoRedist/Platforms/{PLATFORM}/Config/",          Before2 = "{PROJECT}/Config/{PLATFORM}/", After2 = "{RESTRICTEDPROJECT_NR}/Platforms/{PLATFORM}/{OPT_SUBDIR}Config/" },
+			new ConfigLayerExpansion { Before1 = "{ENGINE}/Config/{PLATFORM}/", After1 = "{ENGINE}/Restricted/LimitedAccess/Platforms/{PLATFORM}/Config/",     Before2 = "{PROJECT}/Config/{PLATFORM}/", After2 = "{RESTRICTEDPROJECT_LA}/Platforms/{PLATFORM}/{OPT_SUBDIR}Config/" },
+		};
+
+		private static string PerformBasicReplacements(string InString, string BaseIniName, string CustomConfig)
+		{
+			string OutString = InString.Replace("{TYPE}", BaseIniName);
+
+			OutString = OutString.Replace("{APPSETTINGS}", Unreal.ApplicationSettingDirectory.FullName);
+			OutString = OutString.Replace("{USERSETTINGS}", Unreal.UserSettingDirectory.FullName);
+
+			DirectoryReference? UserDir = Unreal.UserDirectory;
+			if (UserDir != null)
+			{
+				OutString = OutString.Replace("{USER}", UserDir.FullName);
+			}
+
+			OutString = OutString.Replace("{CUSTOMCONFIG}", CustomConfig);
+
+			return OutString;
+		}
+
+		private static string? PerformExpansionReplacements(ConfigLayerExpansion Expansion, string InString)
+		{
+			// if there's replacement to do, the output is just the output
+			if (Expansion.Before1 == null)
+			{
+				return InString;
+			}
+
+			// if nothing to replace, then skip it entirely
+			if (!InString.Contains(Expansion.Before1) && (Expansion.Before2 == null || !InString.Contains(Expansion.Before2)))
+			{
+				return null;
+			}
+
+			// replace the directory bits
+			string OutString = InString.Replace(Expansion.Before1, Expansion.After1);
+			if (Expansion.Before2 != null)
+			{
+				OutString = OutString.Replace(Expansion.Before2, Expansion.After2);
+			}
+			return OutString;
+		}
+
+		private static string PerformFinalExpansions(string InString, string PlatformName, DirectoryReference? ProjectDir)
+		{
+			string PlatformExtensionEngineConfigDir = DirectoryReference.Combine(Unreal.EngineDirectory, "Platforms", PlatformName).FullName;
+
+			string OutString = InString.Replace("{ENGINE}", Unreal.EngineDirectory.FullName);
+			OutString = OutString.Replace("{EXTENGINE}", PlatformExtensionEngineConfigDir);
+			OutString = OutString.Replace("{PLATFORM}", PlatformName);
+
+			if (ProjectDir != null)
+			{
+				DirectoryReference NFLDir;
+				DirectoryReference NRDir;
+				DirectoryReference LADir;
+				string OptionalSubDir = "";
+
+				if (ProjectDir.IsUnderDirectory(Unreal.EngineDirectory))
+				{
+					OptionalSubDir = ProjectDir.MakeRelativeTo(Unreal.EngineDirectory) + "/";
+					NFLDir = DirectoryReference.Combine(Unreal.EngineDirectory, "Restricted/NotForLicensees");
+					NRDir = DirectoryReference.Combine(Unreal.EngineDirectory, "Restricted/NoRedist");
+					LADir = DirectoryReference.Combine(Unreal.EngineDirectory, "Restricted/LimitedAccess");
+				}
+				else
+				{
+					NFLDir = DirectoryReference.Combine(ProjectDir, "Restricted/NotForLicensees");
+					NRDir = DirectoryReference.Combine(ProjectDir, "Restricted/NoRedist");
+					LADir = DirectoryReference.Combine(ProjectDir, "Restricted/LimitedAccess");
+				}
+
+				if (ProjectDir.IsUnderDirectory(NFLDir))
+				{
+					OptionalSubDir = ProjectDir.MakeRelativeTo(NFLDir) + "/";
+				}
+				else if (ProjectDir.IsUnderDirectory(NRDir))
+				{
+					OptionalSubDir = ProjectDir.MakeRelativeTo(NRDir) + "/";
+				}
+				else if (ProjectDir.IsUnderDirectory(LADir))
+				{
+					OptionalSubDir = ProjectDir.MakeRelativeTo(LADir) + "/";
+				}
+
+				string PlatformExtensionProjectConfigDir = DirectoryReference.Combine(ProjectDir, "Platforms", PlatformName).FullName;
+
+				OutString = OutString.Replace("{PROJECT}", ProjectDir.FullName);
+				OutString = OutString.Replace("{EXTPROJECT}", PlatformExtensionProjectConfigDir);
+				OutString = OutString.Replace("{RESTRICTEDPROJECT_NFL}", NFLDir.FullName);
+				OutString = OutString.Replace("{RESTRICTEDPROJECT_NR}", NRDir.FullName);
+				OutString = OutString.Replace("{RESTRICTEDPROJECT_LA}", LADir.FullName);
+				OutString = OutString.Replace("{OPT_SUBDIR}", OptionalSubDir);
+
+			}
+
+			return OutString;
+		}
+
+		/// <summary>
+		/// Returns a list of INI filenames for the given project
+		/// </summary>
+		public static IEnumerable<FileReference> EnumerateConfigFileLocations(ConfigHierarchyType Type, DirectoryReference? ProjectDir, UnrealTargetPlatform Platform, string CustomConfig, TargetType? IncludePluginsForTargetType)
+		{
+			string BaseIniName = Enum.GetName(typeof(ConfigHierarchyType), Type) ?? String.Empty;
+			string PlatformName = GetIniPlatformName(Platform);
+			List<DirectoryReference> PluginDirs = new();
+			List<DirectoryReference> PluginPlatformExtDirs = new();
+
+			if (IncludePluginsForTargetType.HasValue)
+			{
+				ProjectDescriptor? Project = null;
+				if (ProjectDir != null)
+				{
+					Project = ProjectDescriptor.FromDirectory(ProjectDir);
+				}
+				List<PluginInfo> AllPLugins = Plugins.ReadAvailablePlugins(Unreal.EngineDirectory, ProjectDir, null);
+				foreach (PluginInfo Plugin in AllPLugins)
+				{
+					if (Plugins.IsPluginEnabledForTarget(Plugin, Project, Platform, UnrealTargetConfiguration.Development, IncludePluginsForTargetType.Value))
+					{
+						PluginDirs.Add(Plugin.Directory);
+						foreach (FileReference ChildPlugin in Plugin.ChildFiles)
+						{
+							// only look at a child plugin for a platform we are enumerating (ie, it is in a /PlatformName/ subdir), but always look in Restricted folders
+							if (ChildPlugin.ContainsName("Platforms", 0))
+							{
+								if (ChildPlugin.ContainsName(PlatformName, 0))
+								{
+									PluginPlatformExtDirs.Add(ChildPlugin.Directory);
+								}
+							}
+							else
+							{
+								PluginDirs.Add(ChildPlugin.Directory);
+							}
+						}
+					}
+				}
+			}
+
+			foreach (string Layer in ConfigLayers)
+			{
+				bool bHasPlatformTag = Layer.Contains("{PLATFORM}");
+				bool bHasProjectTag = Layer.Contains("{PROJECT}");
+				bool bHasUserTag = Layer.Contains("{USER}");
+				bool bHasCustomConfigTag = Layer.Contains("{CUSTOMCONFIG}");
+
+				// skip certain layers if we are platform-less, project-less, or userdir-less
+				if ((bHasPlatformTag && PlatformName == "None") ||
+					(bHasProjectTag && ProjectDir == null) ||
+					(bHasUserTag && Unreal.UserSettingDirectory == null) ||
+					(bHasCustomConfigTag && String.IsNullOrEmpty(CustomConfig)))
+				{
+					continue;
+				}
+
+				string LayerPath = PerformBasicReplacements(Layer, BaseIniName, CustomConfig);
+
+				// we only expand engine/project inis
+				if (Layer.Contains("{ENGINE}") || Layer.Contains("{PROJECT}"))
+				{
+					foreach (ConfigLayerExpansion Expansion in ConfigLayerExpansions)
+					{
+						// expansion replacements
+						string? ExpandedPath = PerformExpansionReplacements(Expansion, LayerPath);
+
+						// if nothing was replaced, then skip it, as it won't change anything
+						if (ExpandedPath == null)
+						{
+							continue;
+						}
+
+						// now go up the ini parent chain
+						if (bHasPlatformTag)
+						{
+							DataDrivenPlatformInfo.ConfigDataDrivenPlatformInfo? Info = DataDrivenPlatformInfo.GetDataDrivenInfoForPlatform(PlatformName);
+							if (Info != null && Info.IniParentChain != null)
+							{
+								// the IniParentChain
+								foreach (string ParentPlatform in Info.IniParentChain)
+								{
+									// @note: We are using the ParentPlatform as both PlatformExtensionName _and_ IniPlatformName. This is because the parent
+									// may not even exist as a UnrealTargetPlatform, and all we have is a string to look up, and it would just get the same
+									// string back, if we did look it up. This could become an issue if Win64 becomes a PlatformExtension, and wants to have 
+									// a parent Platform, of ... something. This is likely to never be an issue, but leaving this note here just in case.
+									yield return new FileReference(PerformFinalExpansions(ExpandedPath, ParentPlatform, ProjectDir));
+								}
+							}
+							// always yield the active platform last 
+							yield return new FileReference(PerformFinalExpansions(ExpandedPath, PlatformName, ProjectDir));
+						}
+						else
+						{
+							yield return new FileReference(PerformFinalExpansions(ExpandedPath, "", ProjectDir));
+						}
+					}
+				}
+				else
+				{
+					yield return new FileReference(LayerPath);
+				}
+			}
+
+			foreach (DirectoryReference PluginDir in PluginDirs)
+			{
+				yield return FileReference.Combine(PluginDir, "Config", $"{Type}.ini");
+				yield return FileReference.Combine(PluginDir, "Config", PlatformName, $"{PlatformName}{Type}.ini");
+			}
+			foreach (DirectoryReference PluginDir in PluginPlatformExtDirs)
+			{
+				yield return FileReference.Combine(PluginDir, "Config", $"{PlatformName}{Type}.ini");
+			}
+
+			// Find all the generated config files
+			foreach (FileReference GeneratedConfigFile in EnumerateGeneratedConfigFileLocations(Type, ProjectDir, Platform))
+			{
+				yield return GeneratedConfigFile;
+			}
+		}
+
+		/// <summary>
+		/// Returns a list of INI filenames for the given project
+		/// </summary>
+		public static IEnumerable<FileReference> EnumerateGeneratedConfigFileLocations(ConfigHierarchyType Type, DirectoryReference? ProjectDir, UnrealTargetPlatform Platform)
+		{
+			string BaseIniName = Enum.GetName(typeof(ConfigHierarchyType), Type)!;
+			string PlatformName = GetIniPlatformName(Platform);
+
+			// Get the generated config file too. EditorSettings overrides this from 
+			if (Type == ConfigHierarchyType.EditorSettings)
+			{
+				yield return FileReference.Combine(GetGameAgnosticSavedDir(), "Config", PlatformName + "Editor", BaseIniName + ".ini");
+			}
+			else if (Type == ConfigHierarchyType.EditorPerProjectUserSettings)
+			{
+				yield return FileReference.Combine(GetGeneratedConfigDir(ProjectDir), PlatformName + "Editor", BaseIniName + ".ini");
+			}
+			else
+			{
+				yield return FileReference.Combine(GetGeneratedConfigDir(ProjectDir), PlatformName, BaseIniName + ".ini");
+			}
+		}
+
+		/// <summary>
+		/// Determines the path to the generated config directory (same as FPaths::GeneratedConfigDir())
+		/// </summary>
+		/// <returns></returns>
+		public static DirectoryReference GetGeneratedConfigDir(DirectoryReference? ProjectDir)
+		{
+			if (ProjectDir == null)
+			{
+				return DirectoryReference.Combine(Unreal.EngineDirectory, "Saved", "Config");
+			}
+			else
+			{
+				return DirectoryReference.Combine(ProjectDir, "Saved", "Config");
+			}
+		}
+
+		/// <summary>
+		/// Determes the path to the game-agnostic saved directory (same as FPaths::GameAgnosticSavedDir())
+		/// </summary>
+		/// <returns></returns>
+		public static DirectoryReference GetGameAgnosticSavedDir()
+		{
+			if (Unreal.IsEngineInstalled())
+			{
+				DirectoryReference? UserSettingDir = Unreal.UserSettingDirectory;
+				if (UserSettingDir != null)
+				{
+					return DirectoryReference.Combine(UserSettingDir, "UnrealEngine", String.Format("{0}.{1}", ReadOnlyBuildVersion.Current.MajorVersion, ReadOnlyBuildVersion.Current.MinorVersion), "Saved");
+				}
+			}
+			return DirectoryReference.Combine(Unreal.EngineDirectory, "Saved");
+		}
+
+		/// <summary>
+		/// Returns the platform name to use as part of platform-specific config files
+		/// </summary>
+		public static string GetIniPlatformName(UnrealTargetPlatform TargetPlatform)
+		{
+			if (TargetPlatform == UnrealTargetPlatform.Win64)
+			{
+				return "Windows";
+			}
+			else
+			{
+				return TargetPlatform.ToString();
+			}
+		}
+
+		#region Unreal struct/map parsing helpers
+
+		/// <summary>
+		/// Gets an ini setting, and then pulls the value for a property out of a struct, in the format:
+		///   [SomeSection]
+		///   SomeStruct=(Foo=Bar,Prop="My Value")
+		/// </summary>
+		/// <param name="Section">Ini section ('SomeSection' in this example)</param>
+		/// <param name="Setting">Name of the struct setting ('SomeStruct' in this example)</param>
+		/// <param name="Property">Name of the property inside the struct ('Prop' in this example)</param>
+		/// <returns>The value retrieved from the struct ('My Value' in this example), or null if anything was not found</returns>
+		public string? GetStructEntryForSetting(string Section, string Setting, string Property)
+		{
+			string ConfigEntry;
+			if (GetString(Section, Setting, out ConfigEntry))
+			{
+				return GetStructEntry(ConfigEntry, Property, false);
+			}
+			return null;
+		}
+
+		/// <summary>
+		/// Pulls the value for a property out of a struct in the given input, in the format:
+		///   (Foo=Bar,Prop="My Value")
+		/// </summary>
+		/// <param name="Input">The entire struct as retrieved from the ini via GetString()</param>
+		/// <param name="Property">Name of the property inside the struct ('Prop' in this example)</param>
+		/// <param name="bIsArrayProperty">Pass true when the value pulled is an array, like (Foo=(X=1,Y=2)), this would return X=1,Y=2</param>
+		/// <returns>The value retrieved from the struct ('My Value' in this example), or null if anything was not found</returns>
+		public static string? GetStructEntry(string Input, string Property, bool bIsArrayProperty)
+		{
+			string PrimaryRegex;
+			string? AltRegex = null;
+			if (bIsArrayProperty)
+			{
+				PrimaryRegex = String.Format("{0}\\s*=\\s*\\((.*?)\\)", Property);
+			}
+			else
+			{
+				// handle quoted strings, allowing for escaped quotation marks (basically doing " followed by whatever, until we see a quote that was not proceeded by a \, and gather the whole mess in an outer group)
+				PrimaryRegex = String.Format("{0}\\s*=\\s*\"((.*?)[^\\\\])\"", Property);
+				// when no quotes, we skip over whitespace, and we end when we see whitespace, a comma or a ). This will handle (Ip = 192.168.0.1 , Name=....) , and return only '192.168.0.1'
+				AltRegex = String.Format("{0}\\s*=\\s*(.*?)[\\s,\\)]", Property);
+			}
+
+			// attempt to match it!
+			Match Result = Regex.Match(Input, PrimaryRegex);
+			if (!Result.Success && AltRegex != null)
+			{
+				Result = Regex.Match(Input, AltRegex);
+			}
+
+			// if we got a success, return the main match value
+			if (Result.Success)
+			{
+				return Result.Groups[1].Value.ToString();
+			}
+
+			return null;
+		}
+
+		/// <summary>
+		/// Gets an ini setting, and then pulls the value for a property out of a map, in the format:
+		///   [SomeSection]
+		///   SomeMap=((Foo=Bar),(SomeKey="My Value"))
+		/// </summary>
+		/// <param name="Section">Ini section ('SomeSection' in this example)</param>
+		/// <param name="Setting">Name of the struct setting ('SomeMap' in this example)</param>
+		/// <param name="Key">Name of the key inside the struct ('SomeKey' in this example)</param>
+		/// <returns>The value retrieved from the map ('My Value' in this example), or null if anything was not found</returns>
+		public string? GetMapValueForSetting(string Section, string Setting, string Key)
+		{
+			string ConfigEntry;
+			if (GetString(Section, Setting, out ConfigEntry))
+			{
+				return GetMapValue(ConfigEntry, Key);
+			}
+			return null;
+		}
+
+		/// <summary>
+		/// Pulls the value for a property out of a struct in the given input, in the format:
+		///   ((Foo=Bar),(SomeKey="My Value"))
+		/// </summary>
+		/// <param name="Input">The entire struct as retrieved from the ini via GetString()</param>
+		/// <param name="Key">Name of the key inside the struct ('SomeKey' in this example). Key cannot have escaped quotes or commas</param>
+		/// <returns>The value retrieved from the map ('My Value' in this example), or null if anything was not found</returns>
+		public static string? GetMapValue(string Input, string Key)
+		{
+			// handle quoted strings, allowing for escaped quotation marks (and possibly the key in quotes as well)
+			string PrimaryRegex = String.Format("{0}\"?\\s*,\\s*\"((.*?)[^\\\\])\"", Key);
+			string AltRegex = String.Format("{0}\"?\\s*,\\s*(.*?)[\\s,\\)]", Key);
+
+			// attempt to match it!
+			Match Result = Regex.Match(Input, PrimaryRegex);
+			if (!Result.Success && AltRegex != null)
+			{
+				Result = Regex.Match(Input, AltRegex);
+			}
+
+			// if we got a success, return the main match value
+			if (Result.Success)
+			{
+				return Result.Groups[1].Value.ToString();
+			}
+
+			return null;
+		}
+
+		/// <summary>
+		/// Load a ini value for the given key, then use GetStructKeyValuePairs to return the key/value pairs of the struct
+		/// </summary>
+		/// <param name="Section">Ini section to read from</param>
+		/// <param name="Setting">Name of the ini key to read from</param>
+		/// <returns>Dictionary of struct vars/values</returns>
+		public Dictionary<string, string>? GetStructKeyValuePairsForSetting(string Section, string Setting)
+		{
+			string ConfigEntry;
+			if (GetString(Section, Setting, out ConfigEntry))
+			{
+				return GetStructKeyValuePairs(ConfigEntry);
+			}
+			return null;
+		}
+
+		/// <summary>
+		/// Given a string input (a struct loaded from a .ini file usually), like (Foo=A, Bar="Hello world"), this will return a dictionary of all values, with quotes trimmed off
+		/// In this example, { { Foo, A } , { Bar, Hello world } }
+		/// </summary>
+		/// <param name="Input">String containing a struct representation</param>
+		/// <returns>Dictionary of struct vars/values</returns>
+		public static Dictionary<string, string> GetStructKeyValuePairs(string Input)
+		{
+			// we expect parens around a properly encoded struct
+			if (!Input.StartsWith("(") || !Input.EndsWith(")"))
+			{
+				return new Dictionary<string, string>();
+			}
+			// strip ()
+			Input = Input.Substring(1, Input.Length - 2);
+
+			List<string> Props = new List<string>();
+
+			int TokenStart = 0;
+			int StrLen = Input.Length;
+			while (TokenStart < StrLen)
+			{
+				// get the next location of each special character
+				int NextComma = Input.IndexOf(',', TokenStart);
+				int NextQuote = Input.IndexOf('\"', TokenStart);
+				// comma first? easy
+				if (NextComma != -1 && NextComma < NextQuote)
+				{
+					Props.Add(Input.Substring(TokenStart, NextComma - TokenStart));
+					TokenStart = NextComma + 1;
+				}
+				// comma but no quotes
+				else if (NextComma != -1 && NextQuote == -1)
+				{
+					Props.Add(Input.Substring(TokenStart, NextComma - TokenStart));
+					TokenStart = NextComma + 1;
+				}
+				// neither found, use the rest
+				else if (NextComma == -1 && NextQuote == -1)
+				{
+					Props.Add(Input.Substring(TokenStart));
+					break;
+				}
+				// quote first? look for quote after
+				else
+				{
+					NextQuote = Input.IndexOf('\"', NextQuote + 1);
+					// are we at the end?
+					if (NextQuote + 1 == StrLen)
+					{
+						// use the rest of the string
+						Props.Add(Input.Substring(TokenStart));
+						break;
+					}
+					// it's expected that the following character is a comma, if not, give up
+					if (Input[NextQuote + 1] != ',')
+					{
+						break;
+					}
+					// if next is comma, we are done this token
+					Props.Add(Input.Substring(TokenStart, (NextQuote - TokenStart) + 1));
+					// skip over the quote and following commma
+					TokenStart = NextQuote + 2;
+				}
+			}
+
+			// now make a dictionary from the properties
+			Dictionary<string, string> KeyValues = new Dictionary<string, string>();
+			foreach (string AProp in Props)
+			{
+				string Prop = AProp.Trim(" \t".ToCharArray());
+				// find the first = (UE properties can't have an equal sign, so it's valid to do)
+				int Equals = Prop.IndexOf('=');
+				// we must have one
+				if (Equals == -1)
+				{
+					continue;
+				}
+
+				string Key = Prop.Substring(0, Equals);
+				string Value = Prop.Substring(Equals + 1);
+				// trim off any quotes around the entire value
+				Value = Value.Trim(" \"".ToCharArray());
+				Key = Key.Trim(" ".ToCharArray());
+				KeyValues.Add(Key, Value);
+			}
+
+			// convert to array type
+			return KeyValues;
+		}
+
+		#endregion
+	}
+}
diff --git a/Engine/Source/Programs/UnrealBuildTool/Configuration/Ini/ConfigValueParser.cs b/Engine/Source/Programs/UnrealBuildTool/Configuration/Ini/ConfigValueParser.cs
new file mode 100644
--- /dev/null
+++ b/Engine/Source/Programs/UnrealBuildTool/Configuration/Ini/ConfigValueParser.cs
@@ -0,0 +1,300 @@
+// Copyright Epic Games, Inc. All Rights Reserved.
+
+using System;
+using System.Collections;
+using System.Collections.Generic;
+using System.Diagnostics.CodeAnalysis;
+using System.Linq;
+using EpicGames.Core;
+using UnrealBuildBase;
+using static UnrealBuildTool.ConfigHierarchy;
+
+namespace UnrealBuildTool
+{
+	/// <summary>
+	/// Functionality for parsing values from a config value
+	/// </summary>
+	public class ConfigValueParser
+	{
+		/// <summary>
+		/// Attempts to parse the given line as the given type. Struct member variables are expected to match
+		/// </summary>
+		/// <param name="T">Type to attempt to parse</param>
+		/// <param name="Line">Line of text to parse</param>
+		/// <param name="Result">Receives parse result</param>
+		/// <returns>True if the parse was successful, false otherwise</returns>
+		public static bool TryParseTyped(Type T, string Line, [NotNullWhen(true)] out object? Result)
+		{
+			try
+			{
+				if (T == typeof(string))
+				{
+					Result = Line;
+					return true;
+				}
+				else if (T.IsPrimitive)
+				{
+					if (T == typeof(bool) && TryParse(Line, out bool BoolValue))
+					{
+						Result = BoolValue;
+						return true;
+					}
+					else if (T == typeof(int) && TryParse(Line, out int IntValue))
+					{
+						Result = IntValue;
+						return true;
+					}
+					else if (T == typeof(float) && TryParse(Line, out float FloatValue))
+					{
+						Result = FloatValue;
+						return true;
+					}
+					else if (T == typeof(double) && TryParse(Line, out double DoubleValue))
+					{
+						Result = DoubleValue;
+						return true;
+					}
+
+					throw new Exception($"Cannot parse \"{Line}\" for primitive type {T}");
+				}
+				else if (T.IsEnum)
+				{
+					if (Enum.TryParse(T, Line, out Result) && Result != null)
+					{
+						return true;
+					}
+
+					throw new Exception($"Cannot parse \"{Line}\" for enum {T}");
+				}
+				else if (T.IsArray)
+				{
+					if (TryParse(Line, out string[]? ArrayValues))
+					{
+						// check type & early out if we can
+						Type ItemType = T.GetElementType()!;
+						if (ItemType == typeof(string))
+						{
+							Result = ArrayValues;
+							return true;
+						}
+
+						// make a generic array
+						Array ResultArray = Array.CreateInstance(ItemType, ArrayValues.Length);
+						for (int i = 0; i < ArrayValues.Length; i++)
+						{
+							if (TryParseTyped(ItemType, ArrayValues[i], out object? ItemValue))
+							{
+								ResultArray.SetValue(ItemValue, i);
+							}
+							else
+							{
+								Result = null;
+								return false;
+							}
+						}
+						Result = ResultArray;
+						return true;
+					}
+				}
+				else if (T == typeof(Guid))
+				{
+					if (TryParse(Line, out Guid GuidValue))
+					{
+						Result = GuidValue;
+						return true;
+					}
+
+					throw new Exception($"Cannot parse Guid from {Line}");
+				}
+				else if (T.IsGenericType)
+				{
+					if (T.GetGenericTypeDefinition() == typeof(Dictionary<,>) && TryParseAsMap(Line, out Dictionary<string, string>? DictValues))
+					{
+						// check types & early out if we can
+						Type KeyType = T.GetGenericArguments()[0];
+						Type ValType = T.GetGenericArguments()[1];
+						if (KeyType == typeof(string) && ValType == typeof(string))
+						{
+							Result = DictValues;
+							return true;
+						}
+
+						// make a generic dictionary
+						IDictionary ResultDict = (IDictionary)Activator.CreateInstance(T)!;
+						foreach (KeyValuePair<string, string> DictValue in DictValues)
+						{
+							if (TryParseTyped(KeyType, DictValue.Key, out object? Key) && Key != null &&
+								TryParseTyped(ValType, DictValue.Value, out object? Value) && Value != null)
+							{
+								ResultDict.Add(Key, Value);
+							}
+							else
+							{
+								Result = null;
+								return false;
+							}
+						}
+
+						Result = ResultDict;
+						return true;
+					}
+					else if (T.GetGenericTypeDefinition() == typeof(List<>) && TryParse(Line, out string[]? ListValues))
+					{
+						// check types & early out if we can
+						Type ItemType = T.GetGenericArguments()[0];
+						if (ItemType == typeof(string))
+						{
+							Result = ListValues.ToList();
+							return true;
+						}
+
+						// make a generic list
+						IList ResultList = (IList)Activator.CreateInstance(T)!;
+						foreach (string Value in ListValues)
+						{
+							if (TryParseTyped(ItemType, Value, out object? ItemValue))
+							{
+								ResultList.Add(ItemValue);
+							}
+							else
+							{
+								Result = null;
+								return false;
+							}
+						}
+						Result = ResultList;
+						return true;
+					}
+
+					throw new Exception($"Cannot parse \"{Line}\" for generic type {T}");
+				}
+
+				// try to extract fields
+				if (!TryParse(Line, out Dictionary<string, string>? Properties))
+				{
+					Properties = null;
+				}
+
+				// some special cases that may or may not use Properties
+				if (T == typeof(FileReference))
+				{
+					if (Properties != null && Properties.ContainsKey("FilePath")) // FFilePath
+					{
+						Result = FileReference.Combine(Unreal.EngineDirectory, "Binaries", BuildHostPlatform.Current.Platform.ToString(), Properties["FilePath"]);
+						return true;
+					}
+					else if (Properties == null)
+					{
+						Result = new FileReference(Line);
+						return true;
+					}
+					else
+					{
+						throw new Exception($"Cannot parse FileReference from {Line}");
+					}
+				}
+				else if (T == typeof(DirectoryReference))
+				{
+					if (Properties != null && Properties.ContainsKey("Path")) // FDirectoryPath
+					{
+						Result = DirectoryReference.Combine(Unreal.EngineDirectory, "Binaries", BuildHostPlatform.Current.Platform.ToString(), Properties["Path"]);
+						return true;
+					}
+					else if (Properties == null)
+					{
+						Result = new DirectoryReference(Line);
+						return true;
+					}
+					else
+					{
+						throw new Exception($"Cannot parse DirectoryReference from {Line}");
+					}
+				}
+
+				// must have Properties from this point onwards
+				if (Properties == null)
+				{
+					Result = null;
+					return false;
+				}
+
+				// parse out structure
+				Result = Activator.CreateInstance(T);
+				if (Result == null)
+				{
+					return false;
+				}
+				const System.Reflection.FieldAttributes IgnoreFlags = System.Reflection.FieldAttributes.Static | System.Reflection.FieldAttributes.InitOnly;
+				foreach (System.Reflection.FieldInfo Field in T.GetFields().Where(X => !X.Attributes.HasFlag(IgnoreFlags)))
+				{
+					if (Properties.ContainsKey(Field.Name))
+					{
+						if (TryParseTyped(Field.FieldType, Properties[Field.Name], out object? FieldValue))
+						{
+							Field.SetValue(Result, FieldValue);
+							continue;
+						}
+						else
+						{
+							Result = null;
+							return false;
+						}
+					}
+
+					// NB. We allow uninitialized values in structs/classes - it may be expected & ini saving omits empty strings etc.
+				}
+				return true;
+			}
+			catch (Exception)
+			{
+				//Log.TraceInformation(e.Message);
+				Result = null;
+				return false;
+			}
+		}
+
+		/// <summary>
+		/// Attempts to parse the given line as the given template type
+		/// </summary>
+		/// <param name="Line">Line of text to parse</param>
+		/// <param name="Result">Receives parse result</param>
+		/// <returns>True if the parse was successful, false otherwise</returns>
+		public static bool TryParseGeneric<T>(string Line, [NotNullWhen(true)] out T? Result) where T : new()
+		{
+			if (TryParseTyped(typeof(T), Line, out object? ObjResult) && ObjResult != null)
+			{
+				Result = (T)ObjResult!;
+				return (Result != null);
+			}
+
+			Result = default;
+			return false;
+		}
+
+		/// <summary>
+		/// Attempts to parse the given array of lines as array of the given template type
+		/// </summary>
+		/// <param name="Lines">Line of text to parse</param>
+		/// <param name="Result">Receives parse result</param>
+		/// <returns>True if the parse was successful, false otherwise</returns>
+		public static bool TryParseArrayGeneric<T>(string[] Lines, [NotNullWhen(true)] out T[]? Result) where T : new()
+		{
+			if (Lines.Length == 0)
+			{
+				Result = Array.Empty<T>();
+				return true;
+			}
+
+			Result = new T[Lines.Length];
+			for (int i = 0; i < Lines.Length; i++)
+			{
+				if (!TryParseGeneric(Lines[i], out Result[i]!))
+				{
+					Result = null;
+					return false;
+				}
+			}
+			return true;
+		}
+	}
+}
diff --git a/Engine/Source/Programs/UnrealBuildTool/Configuration/Ini/ConfigValueTracker.cs b/Engine/Source/Programs/UnrealBuildTool/Configuration/Ini/ConfigValueTracker.cs
new file mode 100644
--- /dev/null
+++ b/Engine/Source/Programs/UnrealBuildTool/Configuration/Ini/ConfigValueTracker.cs
@@ -0,0 +1,192 @@
+// Copyright Epic Games, Inc. All Rights Reserved.
+
+using System;
+using System.Collections.Generic;
+using System.Diagnostics;
+using System.Linq;
+using EpicGames.Core;
+
+namespace UnrealBuildTool
+{
+	/// <summary>
+	/// Identifier for a config file key, including information about the hierarchy used to read it
+	/// </summary>
+	[DebuggerDisplay("{Name}")]
+	class ConfigDependencyKey : IEquatable<ConfigDependencyKey>
+	{
+		/// <summary>
+		/// The config hierarchy type
+		/// </summary>
+		public ConfigHierarchyType Type;
+
+		/// <summary>
+		/// Project directory to read config files from
+		/// </summary>
+		public DirectoryReference? ProjectDir;
+
+		/// <summary>
+		/// The platform being built
+		/// </summary>
+		public UnrealTargetPlatform Platform;
+
+		/// <summary>
+		/// The section name
+		/// </summary>
+		public string SectionName;
+
+		/// <summary>
+		/// The key name
+		/// </summary>
+		public string KeyName;
+
+		/// <summary>
+		/// Constructor
+		/// </summary>
+		/// <param name="Type">The config hierarchy type</param>
+		/// <param name="ProjectDir">Project directory to read config files from</param>
+		/// <param name="Platform">The platform being built</param>
+		/// <param name="SectionName">The section name</param>
+		/// <param name="KeyName">The key name</param>
+		public ConfigDependencyKey(ConfigHierarchyType Type, DirectoryReference? ProjectDir, UnrealTargetPlatform Platform, string SectionName, string KeyName)
+		{
+			this.Type = Type;
+			this.ProjectDir = ProjectDir;
+			this.Platform = Platform;
+			this.SectionName = SectionName;
+			this.KeyName = KeyName;
+		}
+
+		/// <summary>
+		/// Construct a key from an archive
+		/// </summary>
+		/// <param name="Reader">Archive to read from</param>
+		public ConfigDependencyKey(BinaryArchiveReader Reader)
+		{
+			Type = (ConfigHierarchyType)Reader.ReadInt();
+			ProjectDir = Reader.ReadDirectoryReference();
+			Platform = Reader.ReadUnrealTargetPlatform();
+			SectionName = Reader.ReadString()!;
+			KeyName = Reader.ReadString()!;
+		}
+
+		/// <summary>
+		/// Writes this key to an archive
+		/// </summary>
+		/// <param name="Writer">Archive to write to</param>
+		public void Write(BinaryArchiveWriter Writer)
+		{
+			Writer.WriteInt((int)Type);
+			Writer.WriteDirectoryReference(ProjectDir);
+			Writer.WriteUnrealTargetPlatform(Platform);
+			Writer.WriteString(SectionName);
+			Writer.WriteString(KeyName);
+		}
+
+		/// <summary>
+		/// Tests whether this key is equal to another object
+		/// </summary>
+		/// <param name="Other">The object to compare to</param>
+		/// <returns>True if the keys are equal, false otherwise</returns>
+		public override bool Equals(object? Other)
+		{
+			return (Other is ConfigDependencyKey key) && Equals(key);
+		}
+
+		/// <summary>
+		/// Tests whether this key is equal to another key
+		/// </summary>
+		/// <param name="Other">The key to compare to</param>
+		/// <returns>True if the keys are equal, false otherwise</returns>
+		public bool Equals(ConfigDependencyKey? Other)
+		{
+			return Other is not null && Type == Other.Type && ProjectDir == Other.ProjectDir && Platform == Other.Platform && SectionName == Other.SectionName && KeyName == Other.KeyName;
+		}
+
+		/// <summary>
+		/// Gets a hash code for this object
+		/// </summary>
+		/// <returns>Hash code for the object</returns>
+		public override int GetHashCode()
+		{
+			int Hash = 17;
+			Hash = (Hash * 31) + Type.GetHashCode();
+			Hash = (Hash * 31) + ((ProjectDir == null) ? 0 : ProjectDir.GetHashCode());
+			Hash = (Hash * 31) + Platform.GetHashCode();
+			Hash = (Hash * 31) + SectionName.GetHashCode();
+			Hash = (Hash * 31) + KeyName.GetHashCode();
+			return Hash;
+		}
+	}
+
+	/// <summary>
+	/// Stores a list of config key/value pairs that have been read
+	/// </summary>
+	class ConfigValueTracker
+	{
+		/// <summary>
+		/// The dependencies list
+		/// </summary>
+		readonly IReadOnlyDictionary<ConfigDependencyKey, IReadOnlyList<string>?> Dependencies;
+
+		/// <summary>
+		/// Constructor
+		/// </summary>
+		public ConfigValueTracker(IReadOnlyDictionary<ConfigDependencyKey, IReadOnlyList<string>?> ConfigValues)
+		{
+			Dependencies = new Dictionary<ConfigDependencyKey, IReadOnlyList<string>?>(ConfigValues);
+		}
+
+		/// <summary>
+		/// Construct an object from an archive on disk
+		/// </summary>
+		/// <param name="Reader">Archive to read from</param>
+		public ConfigValueTracker(BinaryArchiveReader Reader)
+		{
+			Dependencies = Reader.ReadDictionary(() => new ConfigDependencyKey(Reader), () => (IReadOnlyList<string>?)Reader.ReadList(() => Reader.ReadString()))!;
+		}
+
+		/// <summary>
+		/// Write the dependencies object to disk
+		/// </summary>
+		/// <param name="Writer">Archive to write to</param>
+		public void Write(BinaryArchiveWriter Writer)
+		{
+			Writer.WriteDictionary(Dependencies, Key => Key.Write(Writer), Value => Writer.WriteList(Value, x => Writer.WriteString(x)));
+		}
+
+		/// <summary>
+		/// Checks whether the list of dependencies is still valid
+		/// </summary>
+		/// <returns></returns>
+		public bool IsValid()
+		{
+			foreach (KeyValuePair<ConfigDependencyKey, IReadOnlyList<string>?> Pair in Dependencies)
+			{
+				// Read the appropriate hierarchy
+				ConfigHierarchy Hierarchy = ConfigCache.ReadHierarchy(Pair.Key.Type, Pair.Key.ProjectDir, Pair.Key.Platform);
+
+				// Get the value(s) associated with this key
+				IReadOnlyList<string>? NewValues;
+				Hierarchy.TryGetValues(Pair.Key.SectionName, Pair.Key.KeyName, out NewValues);
+
+				// Check if they're different
+				if (Pair.Value == null)
+				{
+					if (NewValues != null)
+					{
+						return false;
+					}
+				}
+				else
+				{
+					if (NewValues == null || !Enumerable.SequenceEqual(Pair.Value, NewValues, StringComparer.Ordinal))
+					{
+						return false;
+					}
+				}
+			}
+			return true;
+		}
+	}
+}
+
diff --git a/Engine/Source/Programs/UnrealBuildTool/Configuration/Xml/XmlConfig.cs b/Engine/Source/Programs/UnrealBuildTool/Configuration/Xml/XmlConfig.cs
new file mode 100644
--- /dev/null
+++ b/Engine/Source/Programs/UnrealBuildTool/Configuration/Xml/XmlConfig.cs
@@ -0,0 +1,1133 @@
+// Copyright Epic Games, Inc. All Rights Reserved.
+
+using System;
+using System.Collections;
+using System.Collections.Generic;
+using System.Diagnostics.CodeAnalysis;
+using System.IO;
+using System.Linq;
+using System.Reflection;
+using System.Text;
+using System.Xml;
+using System.Xml.Linq;
+using System.Xml.Schema;
+using System.Xml.Serialization;
+using EpicGames.Core;
+using Microsoft.Extensions.Logging;
+using UnrealBuildBase;
+
+namespace UnrealBuildTool
+{
+	/// <summary>
+	/// Functions for manipulating the XML config cache
+	/// </summary>
+	static class XmlConfig
+	{
+		/// <summary>
+		/// An input config file
+		/// </summary>
+		public class InputFile
+		{
+			/// <summary>
+			/// Location of the file
+			/// </summary>
+			public FileReference Location { get; init; }
+
+			/// <summary>
+			/// Which folder to display the config file under in the generated project files
+			/// </summary>
+			public string FolderName { get; init; }
+
+			/// <summary>
+			/// Constructor
+			/// </summary>
+			/// <param name="location"></param>
+			/// <param name="folderName"></param>
+			public InputFile(FileReference location, string folderName)
+			{
+				Location = location;
+				FolderName = folderName;
+			}
+		}
+
+		/// <summary>
+		/// The cache file that is being used
+		/// </summary>
+		public static FileReference? CacheFile;
+
+		/// <summary>
+		/// Parsed config values
+		/// </summary>
+		static XmlConfigData? s_values;
+
+		/// <summary>
+		/// Cached serializer for the XML schema
+		/// </summary>
+		static XmlSerializer? s_cachedSchemaSerializer;
+
+		/// <summary>
+		/// Initialize the config system with the given types
+		/// </summary>
+		/// <param name="overrideCacheFile">Force use of the cached XML config without checking if it's valid (useful for remote builds)</param>
+		/// <param name="projectRootDirectory">Read XML configuration with a Project directory</param>
+		/// <param name="logger">Logger for output</param>
+		public static void ReadConfigFiles(FileReference? overrideCacheFile, DirectoryReference? projectRootDirectory, ILogger logger)
+		{
+			// Find all the configurable types
+			List<Type> configTypes = FindConfigurableTypes();
+
+			if (projectRootDirectory == null)
+			{
+				logger.LogDebug("No project directory provided for project scope BuildConfiguration.xml resolution.");
+			}
+
+			// Update the cache if necessary
+			if (overrideCacheFile != null)
+			{
+				// Set the cache file to the overriden value
+				CacheFile = overrideCacheFile;
+
+				// Never rebuild the cache; just try to load it.
+				if (!XmlConfigData.TryRead(CacheFile, configTypes, out s_values))
+				{
+					throw new BuildException("Unable to load XML config cache ({0})", CacheFile);
+				}
+			}
+			else
+			{
+				if (projectRootDirectory == null)
+				{
+					// Get the default cache file
+					CacheFile = FileReference.Combine(Unreal.EngineDirectory, "Intermediate", "Build", "XmlConfigCache.bin");
+					if (Unreal.IsEngineInstalled())
+					{
+						DirectoryReference? userSettingsDir = Unreal.UserSettingDirectory;
+						if (userSettingsDir != null)
+						{
+							CacheFile = FileReference.Combine(userSettingsDir, "UnrealEngine", $"XmlConfigCache-{Unreal.RootDirectory.FullName.Replace(":", "", StringComparison.OrdinalIgnoreCase).Replace(Path.DirectorySeparatorChar, '+')}.bin");
+						}
+					}
+				}
+				else
+				{
+					CacheFile = FileReference.Combine(projectRootDirectory, "Intermediate", "Build", "XmlConfigCache.bin");
+					s_values = null;
+				}
+
+				// Find all the input files
+				List<FileReference> temporaryInputFileLocations = InputFiles.Select(x => x.Location).ToList();
+
+				if (projectRootDirectory != null)
+				{
+					FileReference projectRootConfigLocation = FileReference.Combine(projectRootDirectory, "Saved", "UnrealBuildTool", "BuildConfiguration.xml");
+					if (!FileReference.Exists(projectRootConfigLocation))
+					{
+						logger.LogDebug("No project scope config file at {ProjectRootConfigLocation}. Creating default.", projectRootConfigLocation.FullName);
+						CreateDefaultConfigFile(projectRootConfigLocation);
+					}
+
+					temporaryInputFileLocations.Add(projectRootConfigLocation);
+				}
+
+				// Write any configuration environment variables to a file so they can be loaded
+				FileReference environmentXml = FileReference.Combine(Unreal.WritableEngineDirectory, "Intermediate", "Build", "UnrealBuildTool.Env.BuildConfiguration.xml");
+				WriteEnvironmentXml(environmentXml, logger);
+				temporaryInputFileLocations.Add(environmentXml);
+
+				FileReference[] inputFileLocations = temporaryInputFileLocations.ToArray();
+
+				// Get the path to the schema
+				FileReference schemaFile = GetSchemaLocation(projectRootDirectory);
+
+				// Try to read the existing cache from disk
+				XmlConfigData? cachedValues;
+				if (IsCacheUpToDate(CacheFile, inputFileLocations) && FileReference.Exists(schemaFile))
+				{
+					if (XmlConfigData.TryRead(CacheFile, configTypes, out cachedValues) && Enumerable.SequenceEqual(inputFileLocations, cachedValues.InputFiles))
+					{
+						s_values = cachedValues;
+					}
+				}
+
+				// If that failed, regenerate it
+				if (s_values == null)
+				{
+					// Find all the configurable fields from the given types
+					Dictionary<string, Dictionary<string, XmlConfigData.TargetMember>> categoryToFields = new Dictionary<string, Dictionary<string, XmlConfigData.TargetMember>>();
+					FindConfigurableFields(configTypes, categoryToFields);
+
+					// Create a schema for the config files
+					XmlSchema schema = CreateSchema(categoryToFields);
+					if (!Unreal.IsEngineInstalled())
+					{
+						WriteSchema(schema, schemaFile);
+					}
+
+					// Read all the XML files and validate them against the schema
+					Dictionary<Type, Dictionary<XmlConfigData.TargetMember, XmlConfigData.ValueInfo>> typeToValues =
+						new Dictionary<Type, Dictionary<XmlConfigData.TargetMember, XmlConfigData.ValueInfo>>();
+					foreach (FileReference inputFile in inputFileLocations)
+					{
+						logger.LogDebug("Reading configuration file from: {Location}", inputFile.FullName);
+						if (!TryReadFile(inputFile, categoryToFields, typeToValues, schema, logger))
+						{
+							throw new BuildException("Failed to properly read XML file : {0}", inputFile.FullName);
+						}
+					}
+
+					// Make sure the cache directory exists
+					DirectoryReference.CreateDirectory(CacheFile.Directory);
+
+					// Create the new cache
+					s_values = new XmlConfigData(inputFileLocations, typeToValues.ToDictionary(
+							x => x.Key,
+							x => x.Value.Select(x => x.Value).ToArray()));
+					s_values.Write(CacheFile);
+				}
+				else
+				{
+					foreach (FileReference inputFile in s_values.InputFiles)
+					{
+						logger.LogDebug("Reading configuration file from cache ({CacheLocaltion}): {Location}", CacheFile.FullName, inputFile.FullName);
+					}
+				}
+			}
+
+			// Apply all the static field values
+			foreach (KeyValuePair<Type, XmlConfigData.ValueInfo[]> typeValuesPair in s_values.TypeToValues)
+			{
+				foreach (XmlConfigData.ValueInfo memberValue in typeValuesPair.Value)
+				{
+					if (memberValue.Target.IsStatic)
+					{
+						object value = InstanceValue(memberValue.Value, memberValue.Target.Type);
+						memberValue.Target.SetValue(null, value);
+					}
+				}
+			}
+		}
+
+		/// <summary>
+		/// Find all the configurable types in the current assembly
+		/// </summary>
+		/// <returns>List of configurable types</returns>
+		static List<Type> FindConfigurableTypes()
+		{
+			List<Type> configTypes = new List<Type>();
+			try
+			{
+				foreach (Type configType in Assembly.GetExecutingAssembly().GetTypes())
+				{
+					if (HasXmlConfigFileAttribute(configType))
+					{
+						configTypes.Add(configType);
+					}
+				}
+			}
+			catch (ReflectionTypeLoadException ex)
+			{
+				Console.WriteLine("TypeLoadException: {0}", String.Join("\n", ex.LoaderExceptions.Select(x => x?.Message)));
+				throw;
+			}
+			return configTypes;
+		}
+
+		/// <summary>
+		/// Determines whether the given type has a field with an XmlConfigFile attribute
+		/// </summary>
+		/// <param name="type">The type to check</param>
+		/// <returns>True if the type has a field with the XmlConfigFile attribute</returns>
+		static bool HasXmlConfigFileAttribute(Type type)
+		{
+			foreach (FieldInfo field in type.GetFields(BindingFlags.Instance | BindingFlags.Static | BindingFlags.Public | BindingFlags.NonPublic))
+			{
+				foreach (CustomAttributeData customAttribute in field.CustomAttributes)
+				{
+					if (customAttribute.AttributeType == typeof(XmlConfigFileAttribute))
+					{
+						return true;
+					}
+				}
+			}
+			foreach (PropertyInfo property in type.GetProperties(BindingFlags.Instance | BindingFlags.Static | BindingFlags.Public | BindingFlags.NonPublic))
+			{
+				foreach (CustomAttributeData customAttribute in property.CustomAttributes)
+				{
+					if (customAttribute.AttributeType == typeof(XmlConfigFileAttribute))
+					{
+						return true;
+					}
+				}
+			}
+			return false;
+		}
+
+		/// <summary>
+		/// Find the location of the XML config schema
+		/// </summary>
+		/// <param name="projectRootDirecory">Optional project root directory</param>
+		/// <returns>The location of the schema file</returns>
+		public static FileReference GetSchemaLocation(DirectoryReference? projectRootDirecory = null)
+		{
+			if (projectRootDirecory != null)
+			{
+				FileReference projectSchema = FileReference.Combine(projectRootDirecory, "Saved", "UnrealBuildTool", "BuildConfiguration.Schema.xsd");
+				if (FileReference.Exists(projectSchema))
+				{
+					return projectSchema;
+				}
+			}
+			return FileReference.Combine(Unreal.EngineDirectory, "Saved", "UnrealBuildTool", "BuildConfiguration.Schema.xsd");
+		}
+
+		static InputFile[]? s_cachedInputFiles;
+
+		/// <summary>
+		/// Initialize the list of input files
+		/// </summary>
+		public static InputFile[] InputFiles
+		{
+			get
+			{
+				if (s_cachedInputFiles != null)
+				{
+					return s_cachedInputFiles;
+				}
+
+				ILogger logger = Log.Logger;
+
+				// Find all the input file locations
+				List<InputFile> inputFilesFound = new List<InputFile>(5);
+
+				// InputFile info and if a default file should be created if missing
+				List<KeyValuePair<InputFile, bool>> configs = new();
+
+				// Skip all the config files under the Engine folder if it's an installed build
+				if (!Unreal.IsEngineInstalled())
+				{
+					// Check for the engine config file under /Engine/Programs/NotForLicensees/UnrealBuildTool
+					configs.Add(new(new InputFile(FileReference.Combine(Unreal.EngineDirectory, "Restricted", "NotForLicensees", "Programs", "UnrealBuildTool", "BuildConfiguration.xml"), "Engine (NotForLicensees)"), false));
+
+					// Check for the engine user config file under /Engine/Saved/UnrealBuildTool
+					configs.Add(new(new InputFile(FileReference.Combine(Unreal.EngineDirectory, "Saved", "UnrealBuildTool", "BuildConfiguration.xml"), "Engine (Saved)"), true));
+				}
+
+				// Check for the global config file under ProgramData/Unreal Engine/UnrealBuildTool
+				DirectoryReference? commonProgramsFolder = DirectoryReference.FromString(Environment.GetFolderPath(Environment.SpecialFolder.CommonApplicationData));
+				if (commonProgramsFolder != null)
+				{
+					configs.Add(new(new InputFile(FileReference.Combine(commonProgramsFolder, "Unreal Engine", "UnrealBuildTool", "BuildConfiguration.xml"), "Global (ProgramData)"), false));
+				}
+
+				// Check for the global config file under AppData/Unreal Engine/UnrealBuildTool (Roaming)
+				DirectoryReference? appDataFolder = DirectoryReference.FromString(Environment.GetFolderPath(Environment.SpecialFolder.ApplicationData));
+				if (appDataFolder != null)
+				{
+					configs.Add(new(new InputFile(FileReference.Combine(appDataFolder, "Unreal Engine", "UnrealBuildTool", "BuildConfiguration.xml"), "Global (AppData)"), true));
+				}
+
+				// Check for the global config file under LocalAppData/Unreal Engine/UnrealBuildTool
+				DirectoryReference? localAppDataFolder = DirectoryReference.FromString(Environment.GetFolderPath(Environment.SpecialFolder.LocalApplicationData));
+				if (localAppDataFolder != null)
+				{
+					configs.Add(new(new InputFile(FileReference.Combine(localAppDataFolder, "Unreal Engine", "UnrealBuildTool", "BuildConfiguration.xml"), "Global (LocalAppData)"), false));
+				}
+
+				// Check for the global config file under My Documents/Unreal Engine/UnrealBuildTool
+				DirectoryReference? personalFolder = DirectoryReference.FromString(Environment.GetFolderPath(Environment.SpecialFolder.Personal));
+				if (personalFolder != null)
+				{
+					configs.Add(new(new InputFile(FileReference.Combine(personalFolder, "Unreal Engine", "UnrealBuildTool", "BuildConfiguration.xml"), "Global (Documents)"), false));
+				}
+
+				foreach (KeyValuePair<InputFile, bool> config in configs)
+				{
+					if (config.Value && !FileReference.Exists(config.Key.Location))
+					{
+						logger.LogDebug("Creating default config file at {ConfigLocation}", config.Key.Location);
+						CreateDefaultConfigFile(config.Key.Location);
+					}
+					if (FileReference.Exists(config.Key.Location))
+					{
+						inputFilesFound.Add(config.Key);
+					}
+					else
+					{
+						logger.LogDebug("No config file at {ConfigLocation}", config.Key.Location);
+					}
+				}
+
+				s_cachedInputFiles = inputFilesFound.ToArray();
+
+				logger.LogDebug("Configuration will be read from:");
+				foreach (InputFile inputFile in InputFiles)
+				{
+					logger.LogDebug("  {File}", inputFile.Location.FullName);
+				}
+
+				return s_cachedInputFiles;
+			}
+		}
+
+		/// <summary>
+		/// Create a default config file at the given location
+		/// </summary>
+		/// <param name="location">Location to read from</param>
+		static void CreateDefaultConfigFile(FileReference location)
+		{
+			DirectoryReference.CreateDirectory(location.Directory);
+			using (StreamWriter writer = new StreamWriter(location.FullName))
+			{
+				writer.WriteLine("<?xml version=\"1.0\" encoding=\"utf-8\" ?>");
+				writer.WriteLine("<Configuration xmlns=\"{0}\">", XmlConfigFile.SchemaNamespaceURI);
+				writer.WriteLine("</Configuration>");
+			}
+		}
+
+		/// <summary>
+		/// Applies config values to the given object
+		/// </summary>
+		/// <param name="targetObject">The object instance to be configured</param>
+		public static void ApplyTo(object targetObject)
+		{
+			ILogger logger = Log.Logger;
+			for (Type? targetType = targetObject.GetType(); targetType != null; targetType = targetType.BaseType)
+			{
+				XmlConfigData.ValueInfo[]? fieldValues;
+				if (s_values!.TypeToValues.TryGetValue(targetType, out fieldValues))
+				{
+					foreach (XmlConfigData.ValueInfo fieldValue in fieldValues)
+					{
+						if (!fieldValue.Target.IsStatic)
+						{
+							XmlConfigData.TargetMember targetToWrite = fieldValue.Target;
+
+							// Check if setting has been deprecated
+							if (fieldValue.XmlConfigAttribute.Deprecated)
+							{
+								string currentSettingName = fieldValue.XmlConfigAttribute.Name ?? fieldValue.Target.MemberInfo.Name;
+
+								logger.LogWarning("Deprecated setting found in \"{SourceFile}\":", fieldValue.SourceFile);
+								logger.LogWarning("The setting \"{Setting}\" is deprecated. Support for this setting will be removed in a future version of Unreal Engine.", currentSettingName);
+
+								if (fieldValue.XmlConfigAttribute.NewAttributeName != null)
+								{
+									// NewAttributeName is the name of a member in code. However, the log messages below are written from the XML's perspective,
+									// so we need to check if the new target member is not exposed under a custom name in the config.
+									string newSettingName = GetMemberConfigAttributeName(targetType, fieldValue.XmlConfigAttribute.NewAttributeName);
+
+									logger.LogWarning("Use \"{NewAttributeName}\" in place of \"{OldAttributeName}\"", newSettingName, currentSettingName);
+									logger.LogInformation("The value provided for deprecated setting \"{OldName}\" will be applied to \"{NewName}\"", currentSettingName, newSettingName);
+
+									targetToWrite = GetTargetMember(targetType, fieldValue.XmlConfigAttribute.NewAttributeName) ?? targetToWrite;
+								}
+							}
+
+							object valueInstance = InstanceValue(fieldValue.Value, fieldValue.Target.Type);
+							targetToWrite.SetValue(targetObject, valueInstance);
+						}
+					}
+				}
+			}
+		}
+
+		private static string GetMemberConfigAttributeName(Type targetType, string memberName)
+		{
+			MemberInfo? memberInfo = targetType.GetRuntimeFields().FirstOrDefault(x => x.Name == memberName) as MemberInfo
+				?? targetType.GetRuntimeProperties().FirstOrDefault(x => x.Name == memberName) as MemberInfo;
+
+			XmlConfigFileAttribute? attribute = memberInfo?.GetCustomAttributes<XmlConfigFileAttribute>().FirstOrDefault();
+
+			return attribute?.Name ?? memberName;
+		}
+
+		private static XmlConfigData.TargetMember? GetTargetMember(Type targetType, string memberName)
+		{
+			// First, try to find the new field to which the setting should be actually applied.
+			FieldInfo? fieldInfo = targetType.GetRuntimeFields()
+				.FirstOrDefault(x => x.Name == memberName);
+
+			if (fieldInfo != null)
+			{
+				return new XmlConfigData.TargetField(fieldInfo);
+			}
+
+			// If not found, try to find the new property to which the setting should be actually applied.
+			PropertyInfo? propertyInfo = targetType.GetRuntimeProperties()
+				.FirstOrDefault(x => x.Name == memberName);
+
+			if (propertyInfo != null)
+			{
+				return new XmlConfigData.TargetProperty(propertyInfo);
+			}
+
+			return null;
+		}
+
+		/// <summary>
+		/// Instances a value for assignment to a target object
+		/// </summary>
+		/// <param name="value">The value to instance</param>
+		/// <param name="valueType">The type of value</param>
+		/// <returns>New instance of the given value, if necessary</returns>
+		static object InstanceValue(object value, Type valueType)
+		{
+			if (valueType == typeof(string[]))
+			{
+				return ((string[])value).Clone();
+			}
+			else
+			{
+				return value;
+			}
+		}
+
+		/// <summary>
+		/// Gets a config value for a single value, without writing it to an instance of that class
+		/// </summary>
+		/// <param name="targetType">Type to find config values for</param>
+		/// <param name="name">Name of the field to receive</param>
+		/// <param name="value">On success, receives the value of the field</param>
+		/// <returns>True if the value was read, false otherwise</returns>
+		public static bool TryGetValue(Type targetType, string name, [NotNullWhen(true)] out object? value)
+		{
+			// Find all the config values for this type
+			XmlConfigData.ValueInfo[]? fieldValues;
+			if (!s_values!.TypeToValues.TryGetValue(targetType, out fieldValues))
+			{
+				value = null;
+				return false;
+			}
+
+			// Find the value with the matching name
+			foreach (XmlConfigData.ValueInfo fieldValue in fieldValues)
+			{
+				if (fieldValue.Target.MemberInfo.Name == name)
+				{
+					value = fieldValue.Value;
+					return true;
+				}
+			}
+
+			// Not found
+			value = null;
+			return false;
+		}
+
+		/// <summary>
+		/// Find all the configurable fields in the given types by searching for XmlConfigFile attributes.
+		/// </summary>
+		/// <param name="configTypes">Array of types to search</param>
+		/// <param name="categoryToFields">Dictionaries populated with category -> name -> field mappings on return</param>
+		static void FindConfigurableFields(IEnumerable<Type> configTypes, Dictionary<string, Dictionary<string, XmlConfigData.TargetMember>> categoryToFields)
+		{
+			foreach (Type configType in configTypes)
+			{
+				foreach (FieldInfo fieldInfo in configType.GetFields(BindingFlags.Instance | BindingFlags.Static | BindingFlags.GetField | BindingFlags.Public | BindingFlags.NonPublic))
+				{
+					ProcessConfigurableMember<FieldInfo>(configType, fieldInfo, categoryToFields, fieldInfo => new XmlConfigData.TargetField(fieldInfo));
+				}
+				foreach (PropertyInfo propertyInfo in configType.GetProperties(BindingFlags.Instance | BindingFlags.Static | BindingFlags.GetField | BindingFlags.Public | BindingFlags.NonPublic))
+				{
+					ProcessConfigurableMember<PropertyInfo>(configType, propertyInfo, categoryToFields, propertyInfo => new XmlConfigData.TargetProperty(propertyInfo));
+				}
+			}
+		}
+
+		private static void ProcessConfigurableMember<MEMBER>(Type type, MEMBER memberInfo, Dictionary<string, Dictionary<string, XmlConfigData.TargetMember>> categoryToFields, Func<MEMBER, XmlConfigData.TargetMember> createTarget)
+			where MEMBER : System.Reflection.MemberInfo
+		{
+			IEnumerable<XmlConfigFileAttribute> attributes = memberInfo.GetCustomAttributes<XmlConfigFileAttribute>();
+			foreach (XmlConfigFileAttribute attribute in attributes)
+			{
+				string categoryName = attribute.Category ?? type.Name;
+
+				Dictionary<string, XmlConfigData.TargetMember>? nameToTarget;
+				if (!categoryToFields.TryGetValue(categoryName, out nameToTarget))
+				{
+					nameToTarget = new Dictionary<string, XmlConfigData.TargetMember>();
+					categoryToFields.Add(categoryName, nameToTarget);
+				}
+
+				nameToTarget[attribute.Name ?? memberInfo.Name] = createTarget(memberInfo);
+			}
+		}
+
+		/// <summary>
+		/// Creates a schema from attributes in the given types
+		/// </summary>
+		/// <param name="categoryToFields">Lookup for all field settings</param>
+		/// <returns>New schema instance</returns>
+		static XmlSchema CreateSchema(Dictionary<string, Dictionary<string, XmlConfigData.TargetMember>> categoryToFields)
+		{
+			// Create elements for all the categories
+			XmlSchemaAll rootAll = new XmlSchemaAll();
+			foreach (KeyValuePair<string, Dictionary<string, XmlConfigData.TargetMember>> categoryPair in categoryToFields)
+			{
+				string categoryName = categoryPair.Key;
+
+				XmlSchemaAll categoryAll = new XmlSchemaAll();
+				foreach (KeyValuePair<string, XmlConfigData.TargetMember> fieldPair in categoryPair.Value)
+				{
+					XmlSchemaElement element = CreateSchemaFieldElement(fieldPair.Key, fieldPair.Value.Type);
+					categoryAll.Items.Add(element);
+				}
+
+				XmlSchemaComplexType categoryType = new XmlSchemaComplexType();
+				categoryType.Particle = categoryAll;
+
+				XmlSchemaElement categoryElement = new XmlSchemaElement();
+				categoryElement.Name = categoryName;
+				categoryElement.SchemaType = categoryType;
+				categoryElement.MinOccurs = 0;
+				categoryElement.MaxOccurs = 1;
+
+				rootAll.Items.Add(categoryElement);
+			}
+
+			// Create the root element and schema object
+			XmlSchemaComplexType rootType = new XmlSchemaComplexType();
+			rootType.Particle = rootAll;
+
+			XmlSchemaElement rootElement = new XmlSchemaElement();
+			rootElement.Name = XmlConfigFile.RootElementName;
+			rootElement.SchemaType = rootType;
+
+			XmlSchema schema = new XmlSchema();
+			schema.TargetNamespace = XmlConfigFile.SchemaNamespaceURI;
+			schema.ElementFormDefault = XmlSchemaForm.Qualified;
+			schema.Items.Add(rootElement);
+
+			// Finally compile it
+			XmlSchemaSet schemaSet = new XmlSchemaSet();
+			schemaSet.Add(schema);
+			schemaSet.Compile();
+			return schemaSet.Schemas().OfType<XmlSchema>().First();
+		}
+
+		/// <summary>
+		/// Creates an XML schema element for reading a value of the given type
+		/// </summary>
+		/// <param name="name">Name of the field</param>
+		/// <param name="type">Type of the field</param>
+		/// <returns>New schema element representing the field</returns>
+		static XmlSchemaElement CreateSchemaFieldElement(string name, Type type)
+		{
+			XmlSchemaElement element = new XmlSchemaElement();
+			element.Name = name;
+			element.MinOccurs = 0;
+			element.MaxOccurs = 1;
+
+			if (TryGetNullableStructType(type, out Type? innerType))
+			{
+				type = innerType;
+			}
+
+			if (type == typeof(string))
+			{
+				element.SchemaTypeName = XmlSchemaType.GetBuiltInSimpleType(XmlTypeCode.String).QualifiedName;
+			}
+			else if (type == typeof(bool))
+			{
+				element.SchemaTypeName = XmlSchemaType.GetBuiltInSimpleType(XmlTypeCode.Boolean).QualifiedName;
+			}
+			else if (type == typeof(int))
+			{
+				element.SchemaTypeName = XmlSchemaType.GetBuiltInSimpleType(XmlTypeCode.Int).QualifiedName;
+			}
+			else if (type == typeof(float))
+			{
+				element.SchemaTypeName = XmlSchemaType.GetBuiltInSimpleType(XmlTypeCode.Float).QualifiedName;
+			}
+			else if (type == typeof(double))
+			{
+				element.SchemaTypeName = XmlSchemaType.GetBuiltInSimpleType(XmlTypeCode.Double).QualifiedName;
+			}
+			else if (type == typeof(FileReference))
+			{
+				element.SchemaTypeName = XmlSchemaType.GetBuiltInSimpleType(XmlTypeCode.String).QualifiedName;
+			}
+			else if (type.IsEnum)
+			{
+				XmlSchemaSimpleTypeRestriction restriction = new XmlSchemaSimpleTypeRestriction();
+				restriction.BaseTypeName = XmlSchemaType.GetBuiltInSimpleType(XmlTypeCode.String).QualifiedName;
+
+				foreach (string enumName in Enum.GetNames(type))
+				{
+					XmlSchemaEnumerationFacet facet = new XmlSchemaEnumerationFacet();
+					facet.Value = enumName;
+					restriction.Facets.Add(facet);
+				}
+
+				XmlSchemaSimpleType enumType = new XmlSchemaSimpleType();
+				enumType.Content = restriction;
+				element.SchemaType = enumType;
+			}
+			else if (type == typeof(string[]))
+			{
+				XmlSchemaElement itemElement = new XmlSchemaElement();
+				itemElement.Name = "Item";
+				itemElement.SchemaTypeName = XmlSchemaType.GetBuiltInSimpleType(XmlTypeCode.String).QualifiedName;
+				itemElement.MinOccurs = 0;
+				itemElement.MaxOccursString = "unbounded";
+
+				XmlSchemaSequence sequence = new XmlSchemaSequence();
+				sequence.Items.Add(itemElement);
+
+				XmlSchemaComplexType arrayType = new XmlSchemaComplexType();
+				arrayType.Particle = sequence;
+				element.SchemaType = arrayType;
+			}
+			else
+			{
+				throw new Exception("Unsupported field type for XmlConfigFile attribute");
+			}
+			return element;
+		}
+
+		/// <summary>
+		/// Writes a schema to the given location. Avoids writing it if the file is identical.
+		/// </summary>
+		/// <param name="schema">The schema to be written</param>
+		/// <param name="location">Location to write to</param>
+		static void WriteSchema(XmlSchema schema, FileReference location)
+		{
+			XmlWriterSettings settings = new XmlWriterSettings();
+			settings.Indent = true;
+			settings.IndentChars = "\t";
+			settings.NewLineChars = Environment.NewLine;
+			settings.OmitXmlDeclaration = true;
+
+			s_cachedSchemaSerializer ??= XmlSerializer.FromTypes(new Type[] { typeof(XmlSchema) })[0]!;
+
+			StringBuilder output = new StringBuilder();
+			output.AppendLine("<?xml version=\"1.0\" encoding=\"utf-8\"?>");
+			using (XmlWriter writer = XmlWriter.Create(output, settings))
+			{
+				XmlSerializerNamespaces namespaces = new XmlSerializerNamespaces();
+				namespaces.Add("", "http://www.w3.org/2001/XMLSchema");
+				s_cachedSchemaSerializer.Serialize(writer, schema, namespaces);
+			}
+
+			string outputText = output.ToString();
+			if (!FileReference.Exists(location) || File.ReadAllText(location.FullName) != outputText)
+			{
+				DirectoryReference.CreateDirectory(location.Directory);
+				File.WriteAllText(location.FullName, outputText);
+			}
+		}
+
+		/// <summary>
+		/// Tests whether a type is a nullable struct, and extracts the inner type if it is
+		/// </summary>
+		static bool TryGetNullableStructType(Type type, [NotNullWhen(true)] out Type? innerType)
+		{
+			if (type.IsGenericType && type.GetGenericTypeDefinition() == typeof(Nullable<>))
+			{
+				innerType = type.GetGenericArguments()[0];
+				return true;
+			}
+			else
+			{
+				innerType = null;
+				return false;
+			}
+		}
+
+		/// <summary>
+		/// Reads an XML config file and merges it to the given cache
+		/// </summary>
+		/// <param name="location">Location to read from</param>
+		/// <param name="categoryToFields">Lookup for configurable fields by category</param>
+		/// <param name="typeToValues">Map of types to fields and their associated values</param>
+		/// <param name="schema">Schema to validate against</param>
+		/// <param name="logger">Logger for output</param>
+		/// <returns>True if the file was read successfully</returns>
+		static bool TryReadFile(FileReference location, Dictionary<string, Dictionary<string, XmlConfigData.TargetMember>> categoryToFields,
+			Dictionary<Type, Dictionary<XmlConfigData.TargetMember, XmlConfigData.ValueInfo>> typeToValues, XmlSchema schema, ILogger logger)
+		{
+			// Read the XML file, and validate it against the schema
+			XmlConfigFile? configFile;
+			if (!XmlConfigFile.TryRead(location, schema, logger, out configFile))
+			{
+				return false;
+			}
+
+			// Parse the document
+			foreach (XmlElement categoryElement in configFile.DocumentElement!.ChildNodes.OfType<XmlElement>())
+			{
+				Dictionary<string, XmlConfigData.TargetMember>? nameToField;
+				if (categoryToFields.TryGetValue(categoryElement.Name, out nameToField))
+				{
+					foreach (XmlElement keyElement in categoryElement.ChildNodes.OfType<XmlElement>())
+					{
+						if (nameToField.TryGetValue(keyElement.Name, out XmlConfigData.TargetMember? field))
+						{
+							object value;
+							if (field.Type == typeof(string[]))
+							{
+								value = keyElement.ChildNodes.OfType<XmlElement>().Where(x => x.Name == "Item").Select(x => x.InnerText).ToArray();
+							}
+							else if (TryGetNullableStructType(field.Type, out Type? structType))
+							{
+								value = ParseValue(structType, keyElement.InnerText);
+							}
+							else
+							{
+								value = ParseValue(field.Type, keyElement.InnerText);
+							}
+
+							// Add it to the set of values for the type containing this field
+							Dictionary<XmlConfigData.TargetMember, XmlConfigData.ValueInfo>? fieldToValue;
+							if (!typeToValues.TryGetValue(field.MemberInfo.DeclaringType!, out fieldToValue))
+							{
+								fieldToValue = new Dictionary<XmlConfigData.TargetMember, XmlConfigData.ValueInfo>();
+								typeToValues.Add(field.MemberInfo.DeclaringType!, fieldToValue);
+							}
+
+							// Parse the corresponding value
+							XmlConfigData.ValueInfo fieldValue = new XmlConfigData.ValueInfo(field, value, location,
+								field.MemberInfo.GetCustomAttribute<XmlConfigFileAttribute>()!);
+
+							fieldToValue[field] = fieldValue;
+						}
+					}
+				}
+			}
+			return true;
+		}
+
+		/// <summary>
+		/// Writes environment variables as a BuildConfiguration xml
+		/// </summary>
+		/// <param name="location">Path to write to</param>
+		/// <param name="logger">Logger for output</param>
+		/// <returns>True if the environment was written successfully</returns>
+		static bool WriteEnvironmentXml(FileReference location, ILogger logger)
+		{
+			// Read the environment variables
+			Dictionary<string, Dictionary<string, string>> contents = [];
+			foreach (DictionaryEntry EnvironmentVariable in Environment.GetEnvironmentVariables())
+			{
+				string key = EnvironmentVariable.Key?.ToString() ?? String.Empty;
+				string value = EnvironmentVariable.Value?.ToString() ?? String.Empty;
+				if (!key.StartsWith("UnrealBuildTool_") || !key.Contains("__"))
+				{
+					continue;
+				}
+				string[] split = key.Substring("UnrealBuildTool_".Length).Split("__");
+				string category = split.First();
+				string field = split.Last();
+				if (!contents.TryGetValue(category, out Dictionary<string, string>? fields))
+				{
+					fields = new();
+					contents.Add(category, fields);
+				}
+
+				fields.Add(field, value);
+			}
+			StringBuilder output = new StringBuilder();
+			output.AppendLine("<?xml version=\"1.0\" encoding=\"UTF-8\"?>");
+			XmlWriterSettings xmlSettings = new XmlWriterSettings();
+			xmlSettings.Encoding = new UTF8Encoding(false);
+			xmlSettings.Indent = true;
+			xmlSettings.OmitXmlDeclaration = true;
+			using (XmlWriter writer = XmlWriter.Create(output, xmlSettings))
+			{
+				XNamespace ns = XNamespace.Get("https://www.unrealengine.com/BuildConfiguration");
+				XElement root = new XElement(ns + "Configuration");
+				foreach (KeyValuePair<string, Dictionary<string, string>> CategoryItem in contents)
+				{
+					XElement categoryElement = new XElement(ns + CategoryItem.Key);
+					foreach (KeyValuePair<string, string> FieldItem in CategoryItem.Value)
+					{
+						XElement fieldElement = new XElement(ns + FieldItem.Key);
+						if (FieldItem.Value.Contains(';'))
+						{
+							foreach (string item in FieldItem.Value.Split(';', StringSplitOptions.RemoveEmptyEntries | StringSplitOptions.TrimEntries))
+							{
+								fieldElement.Add(new XElement(ns + "Item", item));
+							}
+						}
+						else
+						{
+							fieldElement.Add(FieldItem.Value);
+						}
+						categoryElement.Add(fieldElement);
+					}
+					root.Add(categoryElement);
+				}
+				new XDocument(root).Save(writer);
+			}
+			Utils.WriteFileIfChanged(location, output.ToString(), logger);
+			return true;
+		}
+
+		/// <summary>
+		/// Parse the value for a field from its text based representation in an XML file
+		/// </summary>
+		/// <param name="fieldType">The type of field being read</param>
+		/// <param name="text">Text to parse</param>
+		/// <returns>The object that was parsed</returns>
+		static object ParseValue(Type fieldType, string text)
+		{
+			// ignore whitespace in all fields except for Strings which we leave unprocessed
+			string trimmedText = text.Trim();
+			if (fieldType == typeof(string))
+			{
+				return text;
+			}
+			else if (fieldType == typeof(bool) || fieldType == typeof(bool?))
+			{
+				if (trimmedText == "1" || trimmedText.Equals("true", StringComparison.OrdinalIgnoreCase))
+				{
+					return true;
+				}
+				else if (trimmedText == "0" || trimmedText.Equals("false", StringComparison.OrdinalIgnoreCase))
+				{
+					return false;
+				}
+				else
+				{
+					throw new Exception(String.Format("Unable to convert '{0}' to boolean. 'true/false/0/1' are the supported formats.", text));
+				}
+			}
+			else if (fieldType == typeof(int))
+			{
+				return Int32.Parse(trimmedText);
+			}
+			else if (fieldType == typeof(float))
+			{
+				return Single.Parse(trimmedText, System.Globalization.CultureInfo.InvariantCulture);
+			}
+			else if (fieldType == typeof(double))
+			{
+				return Double.Parse(trimmedText, System.Globalization.CultureInfo.InvariantCulture);
+			}
+			else if (fieldType.IsEnum)
+			{
+				return Enum.Parse(fieldType, trimmedText);
+			}
+			else if (fieldType == typeof(FileReference))
+			{
+				return FileReference.FromString(text);
+			}
+			else
+			{
+				throw new Exception(String.Format("Unsupported config type '{0}'", fieldType.Name));
+			}
+		}
+
+		/// <summary>
+		/// Checks that the given cache file exists and is newer than the given input files, and attempts to read it. Verifies that the resulting cache was created
+		/// from the same input files in the same order.
+		/// </summary>
+		/// <param name="cacheFile">Path to the config cache file</param>
+		/// <param name="inputFiles">The expected set of input files in the cache</param>
+		/// <returns>True if the cache was valid and could be read, false otherwise.</returns>
+		static bool IsCacheUpToDate(FileReference cacheFile, FileReference[] inputFiles)
+		{
+			// Always rebuild if the cache doesn't exist
+			if (!FileReference.Exists(cacheFile))
+			{
+				return false;
+			}
+
+			// Get the timestamp for the cache
+			DateTime cacheWriteTime = File.GetLastWriteTimeUtc(cacheFile.FullName);
+
+			// Always rebuild if this executable is newer
+			if (File.GetLastWriteTimeUtc(Assembly.GetExecutingAssembly().Location) > cacheWriteTime)
+			{
+				return false;
+			}
+
+			// Check if any of the input files are newer than the cache
+			foreach (FileReference inputFile in inputFiles)
+			{
+				if (File.GetLastWriteTimeUtc(inputFile.FullName) > cacheWriteTime)
+				{
+					return false;
+				}
+			}
+
+			// Otherwise, it's up to date
+			return true;
+		}
+
+		/// <summary>
+		/// Generates documentation files for the available settings, by merging the XML documentation from the compiler.
+		/// </summary>
+		/// <param name="outputFile">The documentation file to write</param>
+		/// <param name="logger">Logger for output</param>
+		public static void WriteDocumentation(FileReference outputFile, ILogger logger)
+		{
+			// Find all the configurable types
+			List<Type> configTypes = FindConfigurableTypes();
+
+			// Find all the configurable fields from the given types
+			Dictionary<string, Dictionary<string, XmlConfigData.TargetMember>> categoryToFields = new Dictionary<string, Dictionary<string, XmlConfigData.TargetMember>>();
+			FindConfigurableFields(configTypes, categoryToFields);
+			categoryToFields = categoryToFields.Where(x => x.Value.Count > 0).ToDictionary(x => x.Key, x => x.Value);
+
+			// Get the path to the XML documentation
+			FileReference inputDocumentationFile = new FileReference(Assembly.GetExecutingAssembly().Location).ChangeExtension(".xml");
+			if (!FileReference.Exists(inputDocumentationFile))
+			{
+				throw new BuildException("Generated assembly documentation not found at {0}.", inputDocumentationFile);
+			}
+
+			// Read the documentation
+			XmlDocument inputDocumentation = new XmlDocument();
+			inputDocumentation.Load(inputDocumentationFile.FullName);
+
+			// Make sure we can write to the output file
+			if (FileReference.Exists(outputFile))
+			{
+				FileReference.MakeWriteable(outputFile);
+			}
+			else
+			{
+				DirectoryReference.CreateDirectory(outputFile.Directory);
+			}
+
+			// Generate the documentation file
+			if (outputFile.HasExtension(".udn"))
+			{
+				WriteDocumentationUDN(outputFile, inputDocumentation, categoryToFields, logger);
+			}
+			else if (outputFile.HasExtension(".html"))
+			{
+				WriteDocumentationHTML(outputFile, inputDocumentation, categoryToFields, logger);
+			}
+			else
+			{
+				throw new BuildException("Unable to detect format from extension of output file ({0})", outputFile);
+			}
+
+			// Success!
+			logger.LogInformation("Written documentation to {OutputFile}.", outputFile);
+		}
+
+		/// <summary>
+		/// Writes out documentation in UDN format
+		/// </summary>
+		/// <param name="outputFile">The output file</param>
+		/// <param name="inputDocumentation">The XML documentation for this assembly</param>
+		/// <param name="categoryToFields">Map of string to types to fields</param>
+		/// <param name="logger">Logger for output</param>
+		private static void WriteDocumentationUDN(FileReference outputFile, XmlDocument inputDocumentation, Dictionary<string, Dictionary<string, XmlConfigData.TargetMember>> categoryToFields, ILogger logger)
+		{
+			using (StreamWriter writer = new StreamWriter(outputFile.FullName))
+			{
+				writer.WriteLine("Availability: NoPublish");
+				writer.WriteLine("Title: Build Configuration Properties Page");
+				writer.WriteLine("Crumbs:");
+				writer.WriteLine("Description: This is a procedurally generated markdown page.");
+				writer.WriteLine("Version: {0}.{1}", ReadOnlyBuildVersion.Current.MajorVersion, ReadOnlyBuildVersion.Current.MinorVersion);
+				writer.WriteLine("");
+
+				foreach (KeyValuePair<string, Dictionary<string, XmlConfigData.TargetMember>> categoryPair in categoryToFields)
+				{
+					string categoryName = categoryPair.Key;
+					writer.WriteLine("### {0}", categoryName);
+					writer.WriteLine();
+
+					Dictionary<string, XmlConfigData.TargetMember> fields = categoryPair.Value;
+					foreach (KeyValuePair<string, XmlConfigData.TargetMember> fieldPair in fields)
+					{
+						// Get the XML comment for this field
+						List<string>? lines;
+						if (!RulesDocumentation.TryGetXmlComment(inputDocumentation, fieldPair.Value.MemberInfo, logger, out lines) || lines.Count == 0)
+						{
+							continue;
+						}
+
+						// Write the result to the .udn file
+						writer.WriteLine("$ {0} : {1}", fieldPair.Key, lines[0]);
+						for (int idx = 1; idx < lines.Count; idx++)
+						{
+							if (lines[idx].StartsWith("*", StringComparison.OrdinalIgnoreCase) || lines[idx].StartsWith("-", StringComparison.OrdinalIgnoreCase))
+							{
+								writer.WriteLine("        * {0}", lines[idx].Substring(1).TrimStart());
+							}
+							else
+							{
+								writer.WriteLine("    * {0}", lines[idx]);
+							}
+						}
+						writer.WriteLine();
+					}
+				}
+			}
+		}
+
+		/// <summary>
+		/// Writes out documentation in HTML format
+		/// </summary>
+		/// <param name="outputFile">The output file</param>
+		/// <param name="inputDocumentation">The XML documentation for this assembly</param>
+		/// <param name="categoryToFields">Map of string to types to fields</param>
+		/// <param name="logger">Logger for output</param>
+		private static void WriteDocumentationHTML(FileReference outputFile, XmlDocument inputDocumentation, Dictionary<string, Dictionary<string, XmlConfigData.TargetMember>> categoryToFields, ILogger logger)
+		{
+			using (StreamWriter writer = new StreamWriter(outputFile.FullName))
+			{
+				writer.WriteLine("<html>");
+				writer.WriteLine("  <body>");
+				writer.WriteLine("  <h2>BuildConfiguration Properties</h2>");
+				foreach (KeyValuePair<string, Dictionary<string, XmlConfigData.TargetMember>> categoryPair in categoryToFields)
+				{
+					string categoryName = categoryPair.Key;
+					writer.WriteLine("    <h3>{0}</h3>", categoryName);
+					writer.WriteLine("    <dl>");
+
+					Dictionary<string, XmlConfigData.TargetMember> fields = categoryPair.Value;
+					foreach (KeyValuePair<string, XmlConfigData.TargetMember> fieldPair in fields)
+					{
+						// Get the XML comment for this field
+						List<string>? lines;
+						if (!RulesDocumentation.TryGetXmlComment(inputDocumentation, fieldPair.Value.MemberInfo, logger, out lines) || lines.Count == 0)
+						{
+							continue;
+						}
+
+						// Write the result to the .udn file
+						writer.WriteLine("      <dt>{0}</dt>", fieldPair.Key);
+
+						if (lines.Count == 1)
+						{
+							writer.WriteLine("      <dd>{0}</dd>", lines[0]);
+						}
+						else
+						{
+							writer.WriteLine("      <dd>");
+							for (int idx = 0; idx < lines.Count; idx++)
+							{
+								if (lines[idx].StartsWith("*", StringComparison.OrdinalIgnoreCase) || lines[idx].StartsWith("-", StringComparison.OrdinalIgnoreCase))
+								{
+									writer.WriteLine("        <ul>");
+									for (; idx < lines.Count && (lines[idx].StartsWith("*", StringComparison.OrdinalIgnoreCase) || lines[idx].StartsWith("-", StringComparison.OrdinalIgnoreCase)); idx++)
+									{
+										writer.WriteLine("          <li>{0}</li>", lines[idx].Substring(1).TrimStart());
+									}
+									writer.WriteLine("        </ul>");
+								}
+								else
+								{
+									writer.WriteLine("        {0}", lines[idx]);
+								}
+							}
+							writer.WriteLine("      </dd>");
+						}
+					}
+
+					writer.WriteLine("    </dl>");
+				}
+				writer.WriteLine("  </body>");
+				writer.WriteLine("</html>");
+			}
+		}
+	}
+}
+
diff --git a/Engine/Source/Programs/UnrealBuildTool/Configuration/Xml/XmlConfigData.cs b/Engine/Source/Programs/UnrealBuildTool/Configuration/Xml/XmlConfigData.cs
new file mode 100644
--- /dev/null
+++ b/Engine/Source/Programs/UnrealBuildTool/Configuration/Xml/XmlConfigData.cs
@@ -0,0 +1,280 @@
+// Copyright Epic Games, Inc. All Rights Reserved.
+
+using System;
+using System.Collections.Generic;
+using System.Diagnostics.CodeAnalysis;
+using System.IO;
+using System.Linq;
+using System.Reflection;
+using EpicGames.Core;
+
+namespace UnrealBuildTool
+{
+	/// <summary>
+	/// Stores parsed values from XML config files which can be applied to a configurable type. Can be serialized to disk in binary form as a cache.
+	/// </summary>
+	class XmlConfigData
+	{
+		/// <summary>
+		/// The current cache serialization version
+		/// </summary>
+		const int SerializationVersion = 2;
+
+		/// <summary>
+		/// List of input files. Stored to allow checking cache validity.
+		/// </summary>
+		public FileReference[] InputFiles { get; set; }
+
+		/// <summary>
+		/// Abstract description of a target data member.
+		/// </summary>
+		public abstract class TargetMember
+		{
+			/// <summary>
+			/// Returns Reflection.MemberInfo describing the target class member.
+			/// </summary>
+			public abstract MemberInfo MemberInfo { get; }
+
+			/// <summary>
+			/// Returns Reflection.Type of the target class member.
+			/// </summary>
+			public abstract Type Type { get; }
+
+			/// <summary>
+			/// Indicates whether the target class member is static or not.
+			/// </summary>
+			public abstract bool IsStatic { get; }
+
+			/// <summary>
+			/// Returns the value setter of the target class member.
+			/// </summary>
+			public abstract Action<object?, object?> SetValue { get; }
+
+			/// <summary>
+			/// Returns the value getter of the target class member.
+			/// </summary>
+			public abstract Func<object?, object?> GetValue { get; }
+		}
+
+		/// <summary>
+		/// Description of a field member.
+		/// </summary>
+		public class TargetField : TargetMember
+		{
+			public override MemberInfo MemberInfo => _fieldInfo;
+			public override Type Type => _fieldInfo.FieldType;
+			public override bool IsStatic => _fieldInfo.IsStatic;
+			public override Action<object?, object?> SetValue => _fieldInfo.SetValue;
+			public override Func<object?, object?> GetValue => _fieldInfo.GetValue;
+
+			private readonly FieldInfo _fieldInfo;
+
+			public TargetField(FieldInfo fieldInfo)
+			{
+				_fieldInfo = fieldInfo;
+			}
+		}
+
+		/// <summary>
+		/// Description of a property member.
+		/// </summary>
+		public class TargetProperty : TargetMember
+		{
+			public override MemberInfo MemberInfo => _propertyInfo;
+			public override Type Type => _propertyInfo.PropertyType;
+			public override bool IsStatic => _propertyInfo.GetGetMethod()!.IsStatic;
+			public override Action<object?, object?> SetValue => _propertyInfo.SetValue;
+			public override Func<object?, object?> GetValue => _propertyInfo.GetValue;
+
+			private readonly PropertyInfo _propertyInfo;
+
+			public TargetProperty(PropertyInfo propertyInfo)
+			{
+				_propertyInfo = propertyInfo;
+			}
+		}
+
+		public class ValueInfo
+		{
+			public TargetMember Target { get; init; }
+			public object Value { get; init; }
+			public FileReference SourceFile { get; init; }
+			public XmlConfigFileAttribute XmlConfigAttribute { get; init; }
+
+			public ValueInfo(FieldInfo fieldInfo, object value, FileReference sourceFile, XmlConfigFileAttribute xmlConfigAttribute)
+				: this(new TargetField(fieldInfo), value, sourceFile, xmlConfigAttribute)
+			{
+			}
+
+			public ValueInfo(PropertyInfo propertyInfo, object value, FileReference sourceFile, XmlConfigFileAttribute xmlConfigAttribute)
+				: this(new TargetProperty(propertyInfo), value, sourceFile, xmlConfigAttribute)
+			{
+			}
+
+			public ValueInfo(TargetMember target, object value, FileReference sourceFile, XmlConfigFileAttribute xmlConfigAttribute)
+			{
+				Target = target;
+				Value = value;
+				SourceFile = sourceFile;
+				XmlConfigAttribute = xmlConfigAttribute;
+			}
+		}
+
+		/// <summary>
+		/// Stores a mapping from type -> member -> value, with all the config values for configurable fields.
+		/// </summary>
+		public Dictionary<Type, ValueInfo[]> TypeToValues { get; init; }
+
+		/// <summary>
+		/// Constructor
+		/// </summary>
+		/// <param name="inputFiles"></param>
+		/// <param name="typeToValues"></param>
+		public XmlConfigData(FileReference[] inputFiles, Dictionary<Type, ValueInfo[]> typeToValues)
+		{
+			InputFiles = inputFiles;
+			TypeToValues = typeToValues;
+		}
+
+		/// <summary>
+		/// Attempts to read a previous block of config values from disk
+		/// </summary>
+		/// <param name="location">The file to read from</param>
+		/// <param name="types">Array of valid types. Used to resolve serialized type names to concrete types.</param>
+		/// <param name="data">On success, receives the parsed data</param>
+		/// <returns>True if the data was read and is valid</returns>
+		public static bool TryRead(FileReference location, IEnumerable<Type> types, [NotNullWhen(true)] out XmlConfigData? data)
+		{
+			// Check the file exists first
+			if (!FileReference.Exists(location))
+			{
+				data = null;
+				return false;
+			}
+
+			// Read the cache from disk
+			using (BinaryReader reader = new BinaryReader(File.Open(location.FullName, FileMode.Open, FileAccess.Read, FileShare.Read)))
+			{
+				// Check the serialization version matches
+				if (reader.ReadInt32() != SerializationVersion)
+				{
+					data = null;
+					return false;
+				}
+
+				// Read the input files
+				FileReference[] inputFiles = reader.ReadArray(() => reader.ReadFileReference())!;
+
+				// Read the types
+				int numTypes = reader.ReadInt32();
+				Dictionary<Type, ValueInfo[]> typeToValues = new Dictionary<Type, ValueInfo[]>(numTypes);
+				for (int typeIdx = 0; typeIdx < numTypes; typeIdx++)
+				{
+					// Read the type name
+					string typeName = reader.ReadString();
+
+					// Try to find it in the list of configurable types
+					Type? type = types.FirstOrDefault(x => x.Name == typeName);
+					if (type == null)
+					{
+						data = null;
+						return false;
+					}
+
+					// Read all the values
+					ValueInfo[] values = new ValueInfo[reader.ReadInt32()];
+					for (int valueIdx = 0; valueIdx < values.Length; valueIdx++)
+					{
+						string memberName = reader.ReadString();
+
+						TargetMember? targetMember = GetTargetMemberWithAttribute<XmlConfigFileAttribute>(type, memberName);
+
+						if (targetMember != null)
+						{
+							// If TargetMember is not null, we know it has our attribute.
+							XmlConfigFileAttribute xmlConfigAttribute = targetMember!.MemberInfo.GetCustomAttribute<XmlConfigFileAttribute>()!;
+
+							// Try to parse the value and add it to the output array
+							object value = reader.ReadObject(targetMember.Type)!;
+
+							// Read the path of the config file that provided this setting
+							FileReference sourceFile = reader.ReadFileReference();
+
+							values[valueIdx] = new ValueInfo(targetMember, value, sourceFile, xmlConfigAttribute);
+						}
+						else
+						{
+							data = null;
+							return false;
+						}
+					}
+
+					// Add it to the type map
+					typeToValues.Add(type, values);
+				}
+
+				// Return the parsed data
+				data = new XmlConfigData(inputFiles.ToArray(), typeToValues);
+				return true;
+			}
+		}
+
+		/// <summary>
+		/// Find a data member (field or property) with the given name and attribute and returns TargetMember wrapper created for it.
+		/// </summary>
+		/// <typeparam name="T">Attribute a member has to have to be considered.</typeparam>
+		/// <param name="type">Type which members are to be searched</param>
+		/// <param name="memberName">Name of a member (field or property) to find.</param>
+		/// <returns>TargetMember wrapper or null if no member has been found.</returns>
+		private static TargetMember? GetTargetMemberWithAttribute<T>(Type type, string memberName)
+			where T : Attribute
+		{
+			FieldInfo? field = type.GetField(memberName, BindingFlags.Public | BindingFlags.NonPublic | BindingFlags.Instance | BindingFlags.Static);
+			T? xmlConfigAttribute = field?.GetCustomAttribute<T>();
+			if (field != null && xmlConfigAttribute != null)
+			{
+				return new TargetField(field);
+			}
+
+			PropertyInfo? property = type.GetProperty(memberName, BindingFlags.Public | BindingFlags.NonPublic | BindingFlags.Instance | BindingFlags.Static);
+			xmlConfigAttribute = property?.GetCustomAttribute<T>();
+
+			if (property != null && xmlConfigAttribute != null)
+			{
+				return new TargetProperty(property);
+			}
+
+			return null;
+		}
+
+		/// <summary>
+		/// Writes the coalesced config hierarchy to disk
+		/// </summary>
+		/// <param name="location">File to write to</param>
+		public void Write(FileReference location)
+		{
+			DirectoryReference.CreateDirectory(location.Directory);
+			using (BinaryWriter writer = new BinaryWriter(File.Open(location.FullName, FileMode.Create, FileAccess.Write, FileShare.Read)))
+			{
+				writer.Write(SerializationVersion);
+
+				// Save all the input files. The cache will not be valid if these change.
+				writer.Write(InputFiles, item => writer.Write(item));
+
+				// Write all the categories
+				writer.Write(TypeToValues.Count);
+				foreach (KeyValuePair<Type, ValueInfo[]> typePair in TypeToValues)
+				{
+					writer.Write(typePair.Key.Name);
+					writer.Write(typePair.Value.Length);
+					foreach (ValueInfo memberPair in typePair.Value)
+					{
+						writer.Write(memberPair.Target.MemberInfo.Name);
+						writer.Write(memberPair.Target.Type, memberPair.Value);
+						writer.Write(memberPair.SourceFile);
+					}
+				}
+			}
+		}
+	}
+}
diff --git a/Engine/Source/Programs/UnrealBuildTool/Configuration/Xml/XmlConfigFile.cs b/Engine/Source/Programs/UnrealBuildTool/Configuration/Xml/XmlConfigFile.cs
new file mode 100644
--- /dev/null
+++ b/Engine/Source/Programs/UnrealBuildTool/Configuration/Xml/XmlConfigFile.cs
@@ -0,0 +1,147 @@
+// Copyright Epic Games, Inc. All Rights Reserved.
+
+using System.Diagnostics.CodeAnalysis;
+using System.Xml;
+using System.Xml.Schema;
+using EpicGames.Core;
+using Microsoft.Extensions.Logging;
+
+namespace UnrealBuildTool
+{
+	/// <summary>
+	/// Implementation of XmlDocument which preserves line numbers for its elements
+	/// </summary>
+	class XmlConfigFile : XmlDocument
+	{
+		/// <summary>
+		/// Root element for the XML document
+		/// </summary>
+		public const string RootElementName = "Configuration";
+
+		/// <summary>
+		/// Namespace for the XML schema
+		/// </summary>
+		public const string SchemaNamespaceURI = "https://www.unrealengine.com/BuildConfiguration";
+
+		/// <summary>
+		/// The file being read
+		/// </summary>
+		readonly FileReference _file;
+
+		/// <summary>
+		/// Interface to the LineInfo on the active XmlReader
+		/// </summary>
+		IXmlLineInfo _lineInfo = null!;
+
+		/// <summary>
+		/// Set to true if the reader encounters an error
+		/// </summary>
+		bool _bHasErrors;
+
+		/// <summary>
+		/// Private constructor. Use XmlConfigFile.TryRead to read an XML config file.
+		/// </summary>
+		private XmlConfigFile(FileReference inFile)
+		{
+			_file = inFile;
+		}
+
+		/// <summary>
+		/// Overrides XmlDocument.CreateElement() to construct ScriptElements rather than XmlElements
+		/// </summary>
+		public override XmlElement CreateElement(string? prefix, string localName, string? namespaceUri) => new XmlConfigFileElement(_file, _lineInfo.LineNumber, prefix!, localName, namespaceUri, this);
+
+		/// <summary>
+		/// Loads a script document from the given file
+		/// </summary>
+		/// <param name="file">The file to load</param>
+		/// <param name="schema">The schema to validate against</param>
+		/// <param name="logger">Logger for output</param>
+		/// <param name="outConfigFile">If successful, the document that was read</param>
+		/// <returns>True if the document could be read, false otherwise</returns>
+		public static bool TryRead(FileReference file, XmlSchema schema, ILogger logger, [NotNullWhen(true)] out XmlConfigFile? outConfigFile)
+		{
+			XmlConfigFile configFile = new XmlConfigFile(file);
+
+			XmlReaderSettings settings = new XmlReaderSettings();
+			settings.Schemas.Add(schema);
+			settings.ValidationType = ValidationType.Schema;
+			settings.ValidationEventHandler += configFile.ValidationEvent;
+
+			using (XmlReader reader = XmlReader.Create(file.FullName, settings))
+			{
+				// Read the document
+				configFile._lineInfo = (IXmlLineInfo)reader;
+				try
+				{
+					configFile.Load(reader);
+				}
+				catch (XmlException ex)
+				{
+					if (!configFile._bHasErrors)
+					{
+						Log.TraceErrorTask(file, ex.LineNumber, "{0}", ex.Message);
+						configFile._bHasErrors = true;
+					}
+				}
+
+				// If we hit any errors while parsing
+				if (configFile._bHasErrors)
+				{
+					outConfigFile = null;
+					return false;
+				}
+
+				// Check that the root element is valid. If not, we didn't actually validate against the schema.
+				if (configFile.DocumentElement?.Name != RootElementName)
+				{
+					logger.LogError("Script does not have a root element called '{RootElementName}'", RootElementName);
+					outConfigFile = null;
+					return false;
+				}
+				if (configFile.DocumentElement.NamespaceURI != SchemaNamespaceURI)
+				{
+					logger.LogError("Script root element is not in the '{NamespaceUri}' namespace (add the xmlns=\"{NamespaceUri2}\" attribute)", SchemaNamespaceURI, SchemaNamespaceURI);
+					outConfigFile = null;
+					return false;
+				}
+			}
+
+			outConfigFile = configFile;
+			return true;
+		}
+
+		/// <summary>
+		/// Callback for validation errors in the document
+		/// </summary>
+		/// <param name="sender">Standard argument for ValidationEventHandler</param>
+		/// <param name="args">Standard argument for ValidationEventHandler</param>
+		void ValidationEvent(object? sender, ValidationEventArgs args) => Log.TraceWarningTask(_file, args.Exception.LineNumber, "{0}", args.Message);
+	}
+
+	/// <summary>
+	/// Implementation of XmlElement which preserves line numbers
+	/// </summary>
+	class XmlConfigFileElement : XmlElement
+	{
+		/// <summary>
+		/// The file containing this element
+		/// </summary>
+		public FileReference File { get; init; }
+
+		/// <summary>
+		/// The line number containing this element
+		/// </summary>
+		public int LineNumber { get; init; }
+
+		/// <summary>
+		/// Constructor
+		/// </summary>
+		public XmlConfigFileElement(FileReference inFile, int inLineNumber, string prefix, string localName, string? namespaceUri, XmlConfigFile configFile)
+			: base(prefix, localName, namespaceUri, configFile)
+		{
+			File = inFile;
+			LineNumber = inLineNumber;
+		}
+	}
+}
diff --git a/Engine/Source/Programs/UnrealBuildTool/Configuration/Xml/XmlConfigFileAttribute.cs b/Engine/Source/Programs/UnrealBuildTool/Configuration/Xml/XmlConfigFileAttribute.cs
new file mode 100644
--- /dev/null
+++ b/Engine/Source/Programs/UnrealBuildTool/Configuration/Xml/XmlConfigFileAttribute.cs
@@ -0,0 +1,36 @@
+// Copyright Epic Games, Inc. All Rights Reserved.
+
+using System;
+
+namespace UnrealBuildTool
+{
+	/// <summary>
+	/// Marks a field as being serializable from a config file
+	/// </summary>
+	[AttributeUsage(AttributeTargets.Field | AttributeTargets.Property, AllowMultiple = true)]
+	sealed class XmlConfigFileAttribute : Attribute
+	{
+		/// <summary>
+		/// The category for this config value. Optional; defaults to the declaring type name.
+		/// </summary>
+		public string? Category { get; set; } = null;
+
+		/// <summary>
+		/// Name of the key to read. Optional; defaults to the field name.
+		/// </summary>
+		public string? Name { get; set; } = null;
+
+		/// <summary>
+		/// Use this field to indicate that the XML attribute has been deprecated, and that a warning should
+		/// be shown to the user if it is used.
+		/// A deprecated field should also be marked with the [Obsolete] attribute.
+		/// </summary>
+		public bool Deprecated { get; set; } = false;
+
+		/// <summary>
+		/// If the attribute has been deprecated because it has been renamed, this field can be used to apply the
+		/// value used for this field to another.
+		/// </summary>
+		public string? NewAttributeName { get; set; } = null;
+	}
+}
diff --git a/Engine/Source/Programs/UnrealBuildTool/System/ConfigCache.cs b/Engine/Source/Programs/UnrealBuildTool/System/ConfigCache.cs
deleted file mode 100644
diff --git a/Engine/Source/Programs/UnrealBuildTool/System/ConfigFile.cs b/Engine/Source/Programs/UnrealBuildTool/System/ConfigFile.cs
deleted file mode 100644
diff --git a/Engine/Source/Programs/UnrealBuildTool/System/ConfigFileAttribute.cs b/Engine/Source/Programs/UnrealBuildTool/System/ConfigFileAttribute.cs
deleted file mode 100644
diff --git a/Engine/Source/Programs/UnrealBuildTool/System/ConfigHierarchy.cs b/Engine/Source/Programs/UnrealBuildTool/System/ConfigHierarchy.cs
deleted file mode 100644
diff --git a/Engine/Source/Programs/UnrealBuildTool/System/ConfigValueParser.cs b/Engine/Source/Programs/UnrealBuildTool/System/ConfigValueParser.cs
deleted file mode 100644
diff --git a/Engine/Source/Programs/UnrealBuildTool/System/ConfigValueTracker.cs b/Engine/Source/Programs/UnrealBuildTool/System/ConfigValueTracker.cs
deleted file mode 100644
diff --git a/Engine/Source/Programs/UnrealBuildTool/System/XmlConfig.cs b/Engine/Source/Programs/UnrealBuildTool/System/XmlConfig.cs
deleted file mode 100644
diff --git a/Engine/Source/Programs/UnrealBuildTool/System/XmlConfigData.cs b/Engine/Source/Programs/UnrealBuildTool/System/XmlConfigData.cs
deleted file mode 100644
diff --git a/Engine/Source/Programs/UnrealBuildTool/System/XmlConfigFile.cs b/Engine/Source/Programs/UnrealBuildTool/System/XmlConfigFile.cs
deleted file mode 100644
diff --git a/Engine/Source/Programs/UnrealBuildTool/System/XmlConfigFileAttribute.cs b/Engine/Source/Programs/UnrealBuildTool/System/XmlConfigFileAttribute.cs
deleted file mode 100644
