Optimize UBT's JunkDeleter, which runs unconditionally on every single build job.

#Synced-CL 48842273

diff --git a/Engine/Source/Programs/UnrealBuildTool/System/JunkDeleter.cs b/Engine/Source/Programs/UnrealBuildTool/System/JunkDeleter.cs
--- a/Engine/Source/Programs/UnrealBuildTool/System/JunkDeleter.cs
+++ b/Engine/Source/Programs/UnrealBuildTool/System/JunkDeleter.cs
@@ -3,20 +3,24 @@
 using System;
 using System.Collections.Generic;
 using System.IO;
+using EpicGames.Core;
 using Microsoft.Extensions.Logging;
 using UnrealBuildBase;
 
 namespace UnrealBuildTool
 {
-	class JunkDeleter
+	internal static class JunkDeleter
 	{
 		/// <summary>
 		/// Loads JunkManifest.txt file and removes all junk files/folders defined in it.
 		/// </summary>
-		public static void DeleteJunk(ILogger Logger)
+		/// <remarks>
+		/// JunkManifest.txt is a manually edited file, changed when new things should be deleted.
+		/// </remarks>
+		public static void DeleteJunk(ILogger logger)
 		{
-			List<string> JunkManifest = LoadJunkManifest();
-			DeleteAllJunk(JunkManifest, Logger);
+			List<string> junkManifest = LoadJunkManifest();
+			DeleteAllJunk(junkManifest, logger);
 		}
 
 		/// <summary>
@@ -25,130 +29,136 @@ namespace UnrealBuildTool
 		/// <returns>Junk manifest file contents.</returns>
 		private static List<string> LoadJunkManifest()
 		{
-			string ManifestPath = ".." + Path.DirectorySeparatorChar + "Build" + Path.DirectorySeparatorChar + "JunkManifest.txt";
-			List<string> JunkManifest = new List<string>();
-			if (File.Exists(ManifestPath))
+			string manifestPath = Path.Combine(Unreal.EngineDirectory.FullName, "Build", "JunkManifest.txt");
+			if (!File.Exists(manifestPath) || new FileInfo(manifestPath).Length == 0)
 			{
-				string MachineName = Unreal.MachineName;
-				using (StreamReader reader = new StreamReader(ManifestPath))
+				return [];
+			}
+
+			List<string> junkManifest = [];
+			string machineName = Unreal.MachineName;
+			using (StreamReader reader = new StreamReader(manifestPath))
+			{
+				string? lineRead;
+				// Only three ranges, as we support 'Machine=', 'Platform=' and the path.
+				Span<Range> tokenRanges = stackalloc Range[3];
+				Span<Range> innerTokenRanges = stackalloc Range[2];
+
+				while ((lineRead = reader.ReadLine()) != null)
 				{
-					string CurrentToRootDir = ".." + Path.DirectorySeparatorChar + "..";
-					string? LineRead;
-					while ((LineRead = reader.ReadLine()) != null)
+					ReadOnlySpan<char> junkEntry = lineRead.AsSpan().Trim();
+					if (junkEntry.Length > 0)
 					{
-						string JunkEntry = LineRead.Trim();
-						if (String.IsNullOrEmpty(JunkEntry) == false)
+						int numSplits = junkEntry.Split(tokenRanges, ':');
+						if (numSplits > 3)
+						{
+							continue;
+						}
+
+						bool isValidJunkLine = true;
+
+						foreach (Range range in tokenRanges[..numSplits])
 						{
-							string[] Tokens = JunkEntry.Split(":".ToCharArray());
-							bool bIsValidJunkLine = true;
-							foreach (string Token in Tokens)
+							ReadOnlySpan<char> token = junkEntry[range];
+							if (token.StartsWith("Machine=", StringComparison.OrdinalIgnoreCase))
 							{
-								if (Token.StartsWith("Machine=", StringComparison.InvariantCultureIgnoreCase) == true)
+								int numInnerSplits = token.Split(innerTokenRanges, '=');
+								// check if the machine name on the line matches the current machine name, if not, we don't apply this junk
+								if (numInnerSplits == 2 && !machineName.AsSpan().StartsWith(token[innerTokenRanges[1]]))
 								{
-									string[] InnerTokens = Token.Split("=".ToCharArray());
-									// check if the machine name on the line matches the current machine name, if not, we don't apply this junk
-									if (InnerTokens.Length == 2 && MachineName.StartsWith(InnerTokens[1]) == false)
-									{
-										// Not meant for this machine
-										bIsValidJunkLine = false;
-									}
+									// Not meant for this machine
+									isValidJunkLine = false;
+									break;
 								}
-								else if (Token.StartsWith("Platform=", StringComparison.InvariantCultureIgnoreCase) == true)
+							}
+							else if (token.StartsWith("Platform=", StringComparison.OrdinalIgnoreCase))
+							{
+								int numInnerSplits = token.Split(innerTokenRanges, '=');
+								// if the platform is valid, then we want to keep the files, which means that we don't want to apply the junk line
+								if (numInnerSplits == 2
+									&& UnrealTargetPlatform.TryParse(token[innerTokenRanges[1]].ToString(), out UnrealTargetPlatform parsedPlatform)
+									&& UEBuildPlatform.TryGetBuildPlatform(parsedPlatform, out _))
 								{
-									string[] InnerTokens = Token.Split("=".ToCharArray());
-									// check if the machine name on the line matches the current machine name, if not, we don't apply this junk
-									if (InnerTokens.Length == 2)
-									{
-										UnrealTargetPlatform ParsedPlatform;
-										// if the platform is valid, then we want to keep the files, which means that we don't want to apply the junk line
-										if (UnrealTargetPlatform.TryParse(InnerTokens[1], out ParsedPlatform))
-										{
-											if (UEBuildPlatform.TryGetBuildPlatform(ParsedPlatform, out _))
-											{
-												// this is a good platform, so don't delete any files!
-												bIsValidJunkLine = false;
-											}
-										}
-									}
+									// this is a good platform, so don't delete any files!
+									isValidJunkLine = false;
+									break;
 								}
 							}
+						}
 
-							// All paths within the manifest are Unreal root directory relative.
-							// UBT's working directory is Engine\Source so add "..\..\" to each of the entires.
-							if (bIsValidJunkLine)
-							{
-								// the entry is always the last element in the token array (after the final :)
-								string FixedPath = Path.Combine(CurrentToRootDir, Tokens[^1]);
-								FixedPath = FixedPath.Replace('\\', Path.DirectorySeparatorChar);
-								JunkManifest.Add(FixedPath);
-							}
+						if (isValidJunkLine)
+						{
+							ReadOnlySpan<char> path = junkEntry[tokenRanges[numSplits - 1]];
+							string pathString = path.Length == lineRead.Length ? lineRead : path.ToString();
+							// the entry is always the last element in the token array (after the final :)
+							string fixedPath = Path.Combine(Unreal.RootDirectory.FullName, pathString);
+							fixedPath = fixedPath.Replace('\\', Path.DirectorySeparatorChar);
+							junkManifest.Add(fixedPath);
 						}
 					}
 				}
 			}
-			return JunkManifest;
+			return junkManifest;
 		}
 
 		/// <summary>
 		/// Goes through each entry from the junk manifest and deletes it.
 		/// </summary>
-		/// <param name="JunkManifest">JunkManifest.txt entries.</param>
-		/// <param name="Logger">Logger for output</param>
-		private static void DeleteAllJunk(List<string> JunkManifest, ILogger Logger)
+		/// <param name="junkManifest">JunkManifest.txt entries.</param>
+		/// <param name="logger">Logger for output</param>
+		private static void DeleteAllJunk(List<string> junkManifest, ILogger logger)
 		{
-			foreach (string Junk in JunkManifest)
+			foreach (string junk in junkManifest)
 			{
-				if (IsFile(Junk))
+				if (IsFile(junk))
 				{
-					string FileName = Path.GetFileName(Junk);
-					if (FileName.Contains('*'))
+					string fileName = Path.GetFileName(junk);
+					if (fileName.Contains('*'))
 					{
 						// Wildcard search and delete
-						string DirectoryToLookIn = Path.GetDirectoryName(Junk)!;
-						if (Directory.Exists(DirectoryToLookIn))
+						string directoryToLookIn = Path.GetDirectoryName(junk)!;
+						if (Directory.Exists(directoryToLookIn))
 						{
 							// Delete all files within the specified folder
-							string[] FilesToDelete = Directory.GetFiles(DirectoryToLookIn, FileName, SearchOption.TopDirectoryOnly);
-							foreach (string JunkFile in FilesToDelete)
+							string[] filesToDelete = Directory.GetFiles(directoryToLookIn, fileName, SearchOption.TopDirectoryOnly);
+							foreach (string junkFile in filesToDelete)
 							{
-								DeleteFile(JunkFile, Logger);
+								DeleteFile(junkFile, logger);
 							}
 
 							// Delete all subdirectories with the specified folder
-							string[] DirectoriesToDelete = Directory.GetDirectories(DirectoryToLookIn, FileName, SearchOption.TopDirectoryOnly);
-							foreach (string JunkFolder in DirectoriesToDelete)
+							string[] directoriesToDelete = Directory.GetDirectories(directoryToLookIn, fileName, SearchOption.TopDirectoryOnly);
+							foreach (string junkFolder in directoriesToDelete)
 							{
-								DeleteDirectory(JunkFolder, Logger);
+								DeleteDirectory(junkFolder, logger);
 							}
 						}
 					}
 					else
 					{
-						// Delete single file
-						DeleteFile(Junk, Logger);
+						DeleteFile(junk, logger);
 					}
 				}
-				else if (Directory.Exists(Junk))
+				else if (Directory.Exists(junk))
 				{
-					// Delete the selected folder and all its contents
-					DeleteDirectory(Junk, Logger);
+					DeleteDirectory(junk, logger);
 				}
 			}
 		}
 
-		private static bool IsFile(string PathToCheck)
+		private static bool IsFile(string pathToCheck)
 		{
-			string FileName = Path.GetFileName(PathToCheck);
-			if (String.IsNullOrEmpty(FileName) == false)
+			string fileName = Path.GetFileName(pathToCheck);
+			if (String.IsNullOrEmpty(fileName) == false)
 			{
-				if (FileName.Contains('*'))
+				if (fileName.Contains('*'))
 				{
 					// Assume wildcards are file because the path will be searched for files and directories anyway.
 					return true;
 				}
 				else
 				{
-					return File.Exists(PathToCheck);
+					return File.Exists(pathToCheck);
 				}
 			}
 			else
@@ -160,38 +170,36 @@ namespace UnrealBuildTool
 		/// <summary>
 		/// Deletes a directory recursively gracefully handling all exceptions.
 		/// </summary>
-		/// <param name="DirectoryPath">Path.</param>
-		/// <param name="Logger">Logger for output</param>
-		private static void DeleteDirectory(string DirectoryPath, ILogger Logger)
+		/// <param name="directoryPath">Path.</param>
+		/// <param name="logger">Logger for output</param>
+		private static void DeleteDirectory(string directoryPath, ILogger logger)
 		{
 			try
 			{
-				Logger.LogInformation("Deleting junk directory: \"{Dir}\".", DirectoryPath);
-				Directory.Delete(DirectoryPath, true);
+				logger.LogInformation("Deleting junk directory: \"{Dir}\".", directoryPath);
+				Directory.Delete(directoryPath, true);
 			}
-			catch (Exception Ex)
+			catch (Exception ex)
 			{
-				// Ignore all exceptions
-				Logger.LogInformation("Unable to delete junk directory: \"{Dir}\". Error: {Ex}", DirectoryPath, Ex.Message.TrimEnd());
+				logger.LogInformation("Unable to delete junk directory: \"{Dir}\". Error: {Ex}", directoryPath, ex.Message.TrimEnd());
 			}
 		}
 
 		/// <summary>
 		/// Deletes a file gracefully handling all exceptions.
 		/// </summary>
-		/// <param name="Filename">Filename.</param>
-		/// <param name="Logger">Logger for output</param>
-		private static void DeleteFile(string Filename, ILogger Logger)
+		/// <param name="filename">Filename.</param>
+		/// <param name="logger">Logger for output</param>
+		private static void DeleteFile(string filename, ILogger logger)
 		{
 			try
 			{
-				Logger.LogInformation("Deleting junk file: \"{File}\".", Filename);
-				File.Delete(Filename);
+				logger.LogInformation("Deleting junk file: \"{File}\".", filename);
+				File.Delete(filename);
 			}
-			catch (Exception Ex)
+			catch (Exception ex)
 			{
-				// Ingore all exceptions
-				Logger.LogInformation("Unable to delete junk file: \"{File}\". Error: {Ex}", Filename, Ex.Message.TrimEnd());
+				logger.LogInformation("Unable to delete junk file: \"{File}\". Error: {Ex}", filename, ex.Message.TrimEnd());
 			}
 		}
 	}
