Check for cancellation inside UBT's UBA executor using the cancellation token source rather than a denormalized boolean flag, so there's only one source of truth regarding cancellation.

#Synced-CL 48842273

diff --git a/Engine/Source/Programs/UnrealBuildTool/Executors/UnrealBuildAccelerator/UBAExecutor.cs b/Engine/Source/Programs/UnrealBuildTool/Executors/UnrealBuildAccelerator/UBAExecutor.cs
--- a/Engine/Source/Programs/UnrealBuildTool/Executors/UnrealBuildAccelerator/UBAExecutor.cs
+++ b/Engine/Source/Programs/UnrealBuildTool/Executors/UnrealBuildAccelerator/UBAExecutor.cs
@@ -41,7 +41,6 @@ namespace UnrealBuildTool
 		readonly Dictionary<UnrealBuildAcceleratorCacheConfig, ICacheClient> _cacheClients = [];
 		readonly List<IUBAAgentCoordinator> _agentCoordinators = new();
 		DirectoryReference? _rootDirRef;
-		bool _bIsCancelled;
 
 		DateTime _ubaStartTimeUtc = DateTime.UtcNow;
 		TimeSpan _ubaDurationWaitingForRemote = TimeSpan.Zero;
@@ -579,7 +578,7 @@ namespace UnrealBuildTool
 					_session!.PrintSummary();
 				}
 
-				return success && !_bIsCancelled;
+				return success;
 			}
 			finally
 			{
@@ -771,7 +770,7 @@ namespace UnrealBuildTool
 			{
 				int exitCode = info.ExitCode;
 
-				if (_bIsCancelled || exitCode == 99999) // Process was cancelled by executor
+				if (_cancellationSource.IsCancellationRequested || exitCode == 99999) // Process was cancelled by executor
 				{
 					return ProcessFinishedResponse.None;
 				}
@@ -1018,9 +1017,8 @@ namespace UnrealBuildTool
 
 				while (!scheduler.IsEmpty)
 				{
-					if (shouldCancel && !_bIsCancelled)
+					if (shouldCancel && !_cancellationSource.IsCancellationRequested)
 					{
-						_bIsCancelled = true;
 						_cancellationSource.Cancel();
 
 						scheduler.Cancel();
@@ -1045,17 +1043,18 @@ namespace UnrealBuildTool
 					Thread.Sleep(100);
 				}
 
-				bool res = failedActions == 0;
+				bool noFailedActions = failedActions == 0;
+				bool success = noFailedActions && !shouldCancel;
 
-				actionLogger.TraceSummary(res);
+				actionLogger.TraceSummary(noFailedActions);
 
 				coordinatorScheduler.Reset();
 
 				int cacheMissActions = cacheTotalActions - cacheHitActions;
 
-				if (!_bIsCancelled)
+				if (!_cancellationSource.IsCancellationRequested)
 				{
-					telemetryEvent = new TelemetryExecutorUBAEvent(Name, _ubaStartTimeUtc, res, totalActions, succeededActions, failedActions, cacheHitActions, cacheMissActions,
+					telemetryEvent = new TelemetryExecutorUBAEvent(Name, _ubaStartTimeUtc, success, totalActions, succeededActions, failedActions, cacheHitActions, cacheMissActions,
 						localProcessedActions, remoteProcessedActions,
 						localRetryActions, forcedRetryActions,
 						!UBAConfig.bDisableRemote ? _agentCoordinators.Count : 0, !UBAConfig.bDisableRemote ? _successfulCoordinatorConnections : 0, !UBAConfig.bDisableRemote ? _failedCoordinatorConnections : 0,
@@ -1068,7 +1067,7 @@ namespace UnrealBuildTool
 				{
 					logger.LogInformation("Cancelled execution after {Actions}/{TotalActions} actions ({Reason})", succeededActions + failedActions + cacheHitActions, totalActions, cancelReason);
 				}
-				return res;
+				return success;
 			}
 			finally
 			{
