Split the logic to choose an action executor into its own class.

#Synced-CL 48842273

diff --git a/Engine/Source/Programs/UnrealBuildTool/Actions/ActionGraph.cs b/Engine/Source/Programs/UnrealBuildTool/Actions/ActionGraph.cs
--- a/Engine/Source/Programs/UnrealBuildTool/Actions/ActionGraph.cs
+++ b/Engine/Source/Programs/UnrealBuildTool/Actions/ActionGraph.cs
@@ -17,6 +17,7 @@ using EpicGames.Core;
 using Microsoft.Extensions.Logging;
 using OpenTracing.Util;
 using UnrealBuildBase;
+using UnrealBuildTool.Executors;
 
 namespace UnrealBuildTool
 {
@@ -388,129 +389,6 @@ namespace UnrealBuildTool
 			}
 		}
 
-		internal static UBAExecutor GetUBAExecutor(BuildConfiguration BuildConfiguration, ILogger Logger, List<TargetDescriptor> TargetDescriptors)
-		{
-			if (!UBAExecutor.IsAvailable())
-			{
-				throw new BuildLogEventException("UBA is not available - please ensure the UBA binaries exist for your host platform");
-			}
-
-			UnrealBuildAcceleratorConfig Config = UBAExecutor.GetAppliedConfig(TargetDescriptors);
-			// We always use the UBA executor, but we disable detouring to mirror legacy behaviour if the config disables it.
-			if (!BuildConfiguration.bAllowUBAExecutor)
-			{
-				Config.bAllowDetour = false;
-			}
-
-			return new UBAExecutor(
-				Config,
-				BuildConfiguration.MaxParallelActions,
-				BuildConfiguration.bAllCores,
-				BuildConfiguration.bCompactOutput,
-				Logger,
-				TargetDescriptors);
-		}
-
-		public static ActionExecutor? PrecreateExecutor(BuildConfiguration BuildConfiguration, ILogger Logger, List<TargetDescriptor> TargetDescriptors)
-		{
-			foreach (string Name in BuildConfiguration.RemoteExecutorPriority)
-			{
-				switch (Name)
-				{
-					case XGE.ExecutorName:
-						if (BuildConfiguration.bAllowXGE && XGE.IsAvailable(Logger))
-						{
-							return null;
-						}
-						break;
-					case SNDBS.ExecutorName:
-						if (BuildConfiguration.bAllowSNDBS && SNDBS.IsAvailable(Logger))
-						{
-							return null;
-						}
-						break;
-					case FASTBuild.ExecutorName:
-						if (BuildConfiguration.bAllowFASTBuild && FASTBuild.IsAvailable(Logger))
-						{
-							return null;
-						}
-						break;
-					case UBAExecutor.ExecutorName:
-						if (BuildConfiguration.bAllowUBAExecutor)
-						{
-							return GetUBAExecutor(BuildConfiguration, Logger, TargetDescriptors);
-						}
-						break;
-				}
-			}
-
-			return GetUBAExecutor(BuildConfiguration, Logger, TargetDescriptors);
-		}
-
-		private static ActionExecutor? GetRemoteExecutorByName(string Name, BuildConfiguration BuildConfiguration, int ActionCount, int MinActionsForRemote, List<TargetDescriptor> TargetDescriptors, ILogger Logger)
-		{
-			switch (Name)
-			{
-				case XGE.ExecutorName:
-					{
-						if (ActionCount >= MinActionsForRemote && BuildConfiguration.bAllowXGE && XGE.IsAvailable(Logger) && ActionCount >= XGE.MinActions)
-						{
-							return new XGE(Logger);
-						}
-						return null;
-					}
-				case SNDBS.ExecutorName:
-					{
-						if (ActionCount >= MinActionsForRemote && BuildConfiguration.bAllowSNDBS && SNDBS.IsAvailable(Logger))
-						{
-							return new SNDBS(TargetDescriptors, Logger);
-						}
-						return null;
-					}
-				case FASTBuild.ExecutorName:
-					{
-						if (ActionCount >= MinActionsForRemote && BuildConfiguration.bAllowFASTBuild && FASTBuild.IsAvailable(Logger))
-						{
-							return new FASTBuild(BuildConfiguration, Logger, TargetDescriptors);
-						}
-						return null;
-					}
-				case UBAExecutor.ExecutorName:
-					{
-						// Intentionally not checking MinActionsForRemote
-						if (BuildConfiguration.bAllowUBAExecutor)
-						{
-							return GetUBAExecutor(BuildConfiguration, Logger, TargetDescriptors);
-						}
-						return null;
-					}
-				default:
-					{
-						Logger.LogWarning("Unknown remote executor {Name}", Name);
-						return null;
-					}
-			}
-		}
-
-		/// <summary>
-		/// Selects an ActionExecutor
-		/// </summary>
-		private static ActionExecutor SelectExecutor(BuildConfiguration BuildConfiguration, int ActionCount, List<TargetDescriptor> TargetDescriptors, ILogger Logger)
-		{
-			UnrealBuildAcceleratorConfig Config = UBAExecutor.GetAppliedConfig(TargetDescriptors);
-			int MinActionsForRemote = UBAExecutor.GetDefaultNumParallelProcesses(Config, BuildConfiguration.MaxParallelActions, BuildConfiguration.bAllCores, Logger);
-			foreach (string Name in BuildConfiguration.RemoteExecutorPriority)
-			{
-				ActionExecutor? Executor = GetRemoteExecutorByName(Name, BuildConfiguration, ActionCount, MinActionsForRemote, TargetDescriptors, Logger);
-				if (Executor != null)
-				{
-					return Executor;
-				}
-			}
-
-			return GetUBAExecutor(BuildConfiguration, Logger, TargetDescriptors);
-		}
-
 		/// <summary>
 		/// Executes a list of actions.
 		/// </summary>
@@ -531,7 +409,7 @@ namespace UnrealBuildTool
 				// Figure out which executor to use
 				if (executor == null)
 				{
-					using ActionExecutor SelectedExecutor = SelectExecutor(BuildConfiguration, ActionsToExecute.Count, TargetDescriptors, Logger);
+					using ActionExecutor SelectedExecutor = ExecutorFactory.Instance.GetExecutor(BuildConfiguration, ActionsToExecute.Count, TargetDescriptors, Logger);
 					await InternalExecuteActions(SelectedExecutor, ActionsToExecute, Logger);
 				}
 				else
diff --git a/Engine/Source/Programs/UnrealBuildTool/Executors/ExecutorFactory.cs b/Engine/Source/Programs/UnrealBuildTool/Executors/ExecutorFactory.cs
new file mode 100644
--- /dev/null
+++ b/Engine/Source/Programs/UnrealBuildTool/Executors/ExecutorFactory.cs
@@ -0,0 +1,148 @@
+// Copyright Epic Games, Inc. All Rights Reserved.
+
+using System.Collections.Generic;
+using Microsoft.Extensions.Logging;
+using UnrealBuildBase;
+
+namespace UnrealBuildTool.Executors;
+
+internal interface IExecutorFactory
+{
+	ActionExecutor GetExecutor(BuildConfiguration buildConfiguration, int actionCount, IEnumerable<TargetDescriptor> targetDescriptors, ILogger Logger);
+	UBAExecutor GetUBAExecutor(BuildConfiguration buildConfiguration, IEnumerable<TargetDescriptor> targetDescriptors, ILogger Logger);
+	ActionExecutor? PrecreateExecutor(BuildConfiguration buildConfiguration, IEnumerable<TargetDescriptor> targetDescriptors, ILogger Logger);
+}
+
+internal class ExecutorFactory : IExecutorFactory
+{
+	public static ExecutorFactory Instance { get; } = new ExecutorFactory();
+
+	/// <summary>
+	/// Selects an ActionExecutor
+	/// </summary>
+	public ActionExecutor GetExecutor(BuildConfiguration buildConfiguration, int actionCount, IEnumerable<TargetDescriptor> targetDescriptors, ILogger logger)
+	{
+		UnrealBuildAcceleratorConfig config = UBAExecutor.GetAppliedConfig(targetDescriptors);
+		int minActionsForRemote = UBAExecutor.GetDefaultNumParallelProcesses(config, buildConfiguration.MaxParallelActions, buildConfiguration.bAllCores, logger);
+		foreach (string name in buildConfiguration.RemoteExecutorPriority)
+		{
+			ActionExecutor? executor = GetRemoteExecutorByName(name, buildConfiguration, actionCount, minActionsForRemote, targetDescriptors, logger);
+			if (executor != null)
+			{
+				return executor;
+			}
+		}
+
+		return GetUBAExecutor(buildConfiguration, targetDescriptors, logger);
+	}
+
+	public UBAExecutor GetUBAExecutor(BuildConfiguration buildConfiguration, IEnumerable<TargetDescriptor> targetDescriptors, ILogger logger)
+	{
+		if (!UBAExecutor.IsAvailable())
+		{
+			throw new BuildLogEventException("UBA is not available - please ensure the UBA binaries exist for your host platform");
+		}
+
+		UnrealBuildAcceleratorConfig Config = UBAExecutor.GetAppliedConfig(targetDescriptors);
+		// We always use the UBA executor, but we disable detouring to mirror legacy behaviour if the config disables it.
+		if (!buildConfiguration.bAllowUBAExecutor)
+		{
+			Config.bAllowDetour = false;
+		}
+
+		return new UBAExecutor(
+			Config,
+			buildConfiguration.MaxParallelActions,
+			buildConfiguration.bAllCores,
+			buildConfiguration.bCompactOutput,
+			logger,
+			targetDescriptors);
+	}
+
+	public ActionExecutor? PrecreateExecutor(BuildConfiguration buildConfiguration, IEnumerable<TargetDescriptor> targetDescriptors, ILogger logger)
+	{
+		foreach (string name in buildConfiguration.RemoteExecutorPriority)
+		{
+			switch (name)
+			{
+				case XGE.ExecutorName:
+					if (buildConfiguration.bAllowXGE && XGE.IsAvailable(logger))
+					{
+						return null;
+					}
+					break;
+				case SNDBS.ExecutorName:
+					if (buildConfiguration.bAllowSNDBS && SNDBS.IsAvailable(logger))
+					{
+						return null;
+					}
+					break;
+				case FASTBuild.ExecutorName:
+					if (buildConfiguration.bAllowFASTBuild && FASTBuild.IsAvailable(logger))
+					{
+						return null;
+					}
+					break;
+				case UBAExecutor.ExecutorName:
+					if (buildConfiguration.bAllowUBAExecutor)
+					{
+						return GetUBAExecutor(buildConfiguration, targetDescriptors, logger);
+					}
+					break;
+			}
+		}
+
+		return GetUBAExecutor(buildConfiguration, targetDescriptors, logger);
+	}
+
+	private ActionExecutor? GetRemoteExecutorByName(
+		string name,
+		BuildConfiguration buildConfiguration,
+		int actionCount,
+		int minActionsForRemote,
+		IEnumerable<TargetDescriptor> targetDescriptors,
+		ILogger logger)
+	{
+		switch (name)
+		{
+			case XGE.ExecutorName:
+				{
+					if (actionCount >= minActionsForRemote && buildConfiguration.bAllowXGE && XGE.IsAvailable(logger) && actionCount >= XGE.MinActions)
+					{
+						return new XGE(logger);
+					}
+					return null;
+				}
+			case SNDBS.ExecutorName:
+				{
+					if (actionCount >= minActionsForRemote && buildConfiguration.bAllowSNDBS && SNDBS.IsAvailable(logger))
+					{
+						return new SNDBS(targetDescriptors, logger);
+					}
+					return null;
+				}
+			case FASTBuild.ExecutorName:
+				{
+					if (actionCount >= minActionsForRemote && buildConfiguration.bAllowFASTBuild && FASTBuild.IsAvailable(logger))
+					{
+						return new FASTBuild(buildConfiguration, logger, targetDescriptors);
+					}
+					return null;
+				}
+			case UBAExecutor.ExecutorName:
+				{
+					// Intentionally not checking MinActionsForRemote
+					if (buildConfiguration.bAllowUBAExecutor)
+					{
+						return GetUBAExecutor(buildConfiguration, targetDescriptors, logger);
+					}
+					return null;
+				}
+			default:
+				{
+					logger.LogWarning("Unknown remote executor {Name}", name);
+					return null;
+				}
+		}
+	}
+}
diff --git a/Engine/Source/Programs/UnrealBuildTool/Executors/SNDBS.cs b/Engine/Source/Programs/UnrealBuildTool/Executors/SNDBS.cs
--- a/Engine/Source/Programs/UnrealBuildTool/Executors/SNDBS.cs
+++ b/Engine/Source/Programs/UnrealBuildTool/Executors/SNDBS.cs
@@ -80,13 +80,13 @@ namespace UnrealBuildTool
 
 		private const string ProgressMarkupPrefix = "@action:";
 
-		private readonly List<TargetDescriptor> TargetDescriptors;
+		private readonly IReadOnlyList<TargetDescriptor> TargetDescriptors;
 
-		public SNDBS(List<TargetDescriptor> InTargetDescriptors, ILogger Logger)
+		public SNDBS(IEnumerable<TargetDescriptor> InTargetDescriptors, ILogger Logger)
 			: base(Logger)
 		{
 			XmlConfig.ApplyTo(this);
-			TargetDescriptors = InTargetDescriptors;
+			TargetDescriptors = [.. InTargetDescriptors];
 		}
 
 		public SNDBS AddTemplate(string ExeName, string TemplateContents)
diff --git a/Engine/Source/Programs/UnrealBuildTool/Modes/BuildMode.cs b/Engine/Source/Programs/UnrealBuildTool/Modes/BuildMode.cs
--- a/Engine/Source/Programs/UnrealBuildTool/Modes/BuildMode.cs
+++ b/Engine/Source/Programs/UnrealBuildTool/Modes/BuildMode.cs
@@ -11,6 +11,7 @@ using EpicGames.Core;
 using Microsoft.Extensions.Logging;
 using OpenTracing.Util;
 using UnrealBuildBase;
+using UnrealBuildTool.Executors;
 
 namespace UnrealBuildTool
 {
@@ -203,7 +204,7 @@ namespace UnrealBuildTool
 					CleanMode.Clean(TargetDescriptors.Where(D => D.bRebuild).ToList(), BuildConfiguration, Logger);
 				}
 
-				using ActionExecutor? executor = ActionGraph.PrecreateExecutor(BuildConfiguration, Logger, TargetDescriptors);
+				using ActionExecutor? executor = ExecutorFactory.Instance.PrecreateExecutor(BuildConfiguration, TargetDescriptors, Logger);
 
 				// Handle remote builds
 				for (int Idx = 0; Idx < TargetDescriptors.Count; ++Idx)
