[AutoRTFM] Run the AutoRTFM presubmit tests with the AutoRTFM sanitizer enabled

Add UBT support to allow combining AutoRTFMSan with other sanitizers.

#Synced-CL 48740692

diff --git a/Engine/Source/Programs/UnrealBuildTool/Platform/Linux/UEBuildLinux.cs b/Engine/Source/Programs/UnrealBuildTool/Platform/Linux/UEBuildLinux.cs
--- a/Engine/Source/Programs/UnrealBuildTool/Platform/Linux/UEBuildLinux.cs
+++ b/Engine/Source/Programs/UnrealBuildTool/Platform/Linux/UEBuildLinux.cs
@@ -294,9 +294,9 @@ namespace UnrealBuildTool
 				{
 					SanitizerSuffix = "MSan";
 				}
-				else if (Target.LinuxPlatform.bEnableAutoRTFMSanitizer)
+				if (Target.LinuxPlatform.bEnableAutoRTFMSanitizer)
 				{
-					SanitizerSuffix = "AutoRTFMSan";
+					SanitizerSuffix += "AutoRTFMSan";
 				}
 				if (Target.LinuxPlatform.bEnableLibFuzzer)
 				{
