[UBT] Don't pass -fsanitize=autortfm if using the regular clang compiler

Regular clang will error as it does not recognise this.

Dropping the sanitizer flag is preferable to erroring, as we have a number of scripts that throw a bunch of sanitizer flags at UBT, and filtering those based on whether the target is compiled with AutoRTFM or not will be tedious and error prone.

#Synced-CL 48856386

diff --git a/Engine/Source/Programs/UnrealBuildTool/Platform/Clang/ClangToolChain.cs b/Engine/Source/Programs/UnrealBuildTool/Platform/Clang/ClangToolChain.cs
--- a/Engine/Source/Programs/UnrealBuildTool/Platform/Clang/ClangToolChain.cs
+++ b/Engine/Source/Programs/UnrealBuildTool/Platform/Clang/ClangToolChain.cs
@@ -964,7 +964,8 @@ namespace UnrealBuildTool
 			}
 
 			// AutoRTFMSan
-			if (Options.HasFlag(ClangToolChainOptions.EnableAutoRTFMSanitizer))
+			if (Options.HasFlag(ClangToolChainOptions.EnableAutoRTFMSanitizer) &&
+				CompileEnvironment.bUseAutoRTFMCompiler)
 			{
 				Arguments.Add("-fsanitize=autortfm");
 			}
