PR #13146: Expose a unique session identifier for each UBT invocation

#Synced-CL 48862428

diff --git a/Engine/Source/Programs/UnrealBuildTool/UnrealBuildTool.cs b/Engine/Source/Programs/UnrealBuildTool/UnrealBuildTool.cs
--- a/Engine/Source/Programs/UnrealBuildTool/UnrealBuildTool.cs
+++ b/Engine/Source/Programs/UnrealBuildTool/UnrealBuildTool.cs
@@ -390,6 +390,12 @@ namespace UnrealBuildTool
 				{
 				}
 
+				// Set an environment variable for the current SessionIdentifier so scripts can access
+				if (String.IsNullOrEmpty(Environment.GetEnvironmentVariable("UnrealBuildTool_SessionId")))
+				{
+					Environment.SetEnvironmentVariable("UnrealBuildTool_SessionId", $"{SessionIdentifier}");
+				}
+
 				// Override the temp directory
 				try
 				{
