UnrealBuildTool: Fix bEnableBufferSecurityChecks not being set in copy constructor

#rnx
#Synced-CL 48862428

diff --git a/Engine/Source/Programs/UnrealBuildTool/System/CppCompileEnvironment.cs b/Engine/Source/Programs/UnrealBuildTool/System/CppCompileEnvironment.cs
--- a/Engine/Source/Programs/UnrealBuildTool/System/CppCompileEnvironment.cs
+++ b/Engine/Source/Programs/UnrealBuildTool/System/CppCompileEnvironment.cs
@@ -864,6 +864,7 @@ namespace UnrealBuildTool
 			bUseStackProtection = Other.bUseStackProtection;
 			bUseInlining = Other.bUseInlining;
 			bCompileISPC = Other.bCompileISPC;
+			bEnableBufferSecurityChecks = Other.bEnableBufferSecurityChecks;
 			bUseUnity = Other.bUseUnity;
 			MinSourceFilesForUnityBuildOverride = Other.MinSourceFilesForUnityBuildOverride;
 			MinFilesUsingPrecompiledHeaderOverride = Other.MinFilesUsingPrecompiledHeaderOverride;
