PR #12383: [UBT] Fix sln icon in the shell not displaying vs version, by removing empty first line in sln file

#Synced-CL 48862428

diff --git a/Engine/Source/Programs/UnrealBuildTool/ProjectFiles/VisualStudio/VCProjectFileGenerator.cs b/Engine/Source/Programs/UnrealBuildTool/ProjectFiles/VisualStudio/VCProjectFileGenerator.cs
--- a/Engine/Source/Programs/UnrealBuildTool/ProjectFiles/VisualStudio/VCProjectFileGenerator.cs
+++ b/Engine/Source/Programs/UnrealBuildTool/ProjectFiles/VisualStudio/VCProjectFileGenerator.cs
@@ -724,7 +724,6 @@ namespace UnrealBuildTool
 			// Solution file header. Note that a leading newline is required for file type detection to work correclty in the shell.
 			if (Settings.ProjectFileFormat == VCProjectFileFormat.VisualStudio2022)
 			{
-				VCSolutionFileContent.AppendLine();
 				VCSolutionFileContent.AppendLine("Microsoft Visual Studio Solution File, Format Version 12.00");
 				VCSolutionFileContent.AppendLine("# Visual Studio Version 17");
 				VCSolutionFileContent.AppendLine("VisualStudioVersion = 17.0.31314.256");
@@ -732,7 +731,6 @@ namespace UnrealBuildTool
 			}
 			else if (Settings.ProjectFileFormat == VCProjectFileFormat.VisualStudio2026)
 			{
-				VCSolutionFileContent.AppendLine();
 				VCSolutionFileContent.AppendLine("Microsoft Visual Studio Solution File, Format Version 12.00");
 				VCSolutionFileContent.AppendLine("# Visual Studio Version 18");
 				VCSolutionFileContent.AppendLine("VisualStudioVersion = 18.0.11205.157 d18.0");
@@ -740,7 +738,7 @@ namespace UnrealBuildTool
 			}
 			else
 			{
-				throw new BuildException("Unexpected ProjectFileFormat");
+				throw new BuildException("Unexpected ProjectFileFormat '{0}'", Settings.ProjectFileFormat);
 			}
 
 			IDictionary<PrimaryProjectFolder, Guid> ProjectFolderGuids = GenerateProjectFolderGuids(RootFolder);
