UnrealBuildTool: Don't allow preserving psyms for shared targets

#rnx
#Synced-CL 48888070

diff --git a/Engine/Source/Programs/UnrealBuildTool/Platform/Linux/UEBuildLinux.cs b/Engine/Source/Programs/UnrealBuildTool/Platform/Linux/UEBuildLinux.cs
--- a/Engine/Source/Programs/UnrealBuildTool/Platform/Linux/UEBuildLinux.cs
+++ b/Engine/Source/Programs/UnrealBuildTool/Platform/Linux/UEBuildLinux.cs
@@ -735,7 +735,7 @@ namespace UnrealBuildTool
 				Options |= ClangToolChainOptions.TuneDebugInfoForLLDB;
 			}
 
-			if (Target.LinuxPlatform.bPreservePSYM)
+			if (Target.LinuxPlatform.bPreservePSYM && Target.BuildEnvironment == TargetBuildEnvironment.Unique)
 			{
 				Options |= ClangToolChainOptions.PreservePSYM;
 			}
