Fix an issue with swift source code changes were not detected and recompiled properly

#rnx

diff --git a/Engine/Source/Programs/UnrealBuildTool/Platform/Apple/AppleToolChain.cs b/Engine/Source/Programs/UnrealBuildTool/Platform/Apple/AppleToolChain.cs
--- a/Engine/Source/Programs/UnrealBuildTool/Platform/Apple/AppleToolChain.cs
+++ b/Engine/Source/Programs/UnrealBuildTool/Platform/Apple/AppleToolChain.cs
@@ -606,7 +606,7 @@ namespace UnrealBuildTool
 				HeaderAction.CommandArguments = String.Join(" ", Arguments);
 				foreach (FileItem SourceFile in InputFiles)
 				{
-					CompileAction.PrerequisiteItems.Add(SourceFile);
+					HeaderAction.PrerequisiteItems.Add(SourceFile);
 				}
 				HeaderAction.ProducedItems.Add(OutputInteropHeader);
 				// swiftc won't touch the output header if the generated content doesn't change, so always delete the output so that 
