Explicitly linking Clang Code Coverage lib to enable full code coverage data

#rnx
#Synced-CL 48915187

diff --git a/Engine/Source/Programs/UnrealBuildTool/Platform/Windows/VCToolChain.cs b/Engine/Source/Programs/UnrealBuildTool/Platform/Windows/VCToolChain.cs
--- a/Engine/Source/Programs/UnrealBuildTool/Platform/Windows/VCToolChain.cs
+++ b/Engine/Source/Programs/UnrealBuildTool/Platform/Windows/VCToolChain.cs
@@ -1930,6 +1930,28 @@ namespace UnrealBuildTool
 					LinkEnvironment.Libraries.Add(FileReference.Combine(UBSanRuntimeDir, UBSanRuntimeLib));
 				}
 			}
+			
+			if (LinkEnvironment.bCodeCoverage)
+			{
+				// Link clang_rt.profile-x86_64.lib
+				if (Target.WindowsPlatform.Compiler.IsClang())
+				{
+					DirectoryReference ProfileRuntimeDir;
+					string ProfileArchSuffix;
+					if (EnvVars.Architecture == UnrealArch.X64)
+					{
+						VersionNumber ClangVersion = Target.WindowsPlatform.Compiler == WindowsCompiler.Intel ? MicrosoftPlatformSDK.GetClangVersionForIntelCompiler(EnvVars.CompilerPath) : EnvVars.CompilerVersion;
+						ProfileRuntimeDir = DirectoryReference.Combine(EnvVars.CompilerDir, "lib", "clang", $"{ClangVersion.Components[0]}", "lib", "windows");
+						ProfileArchSuffix = "x86_64";
+					}
+					else
+					{
+						throw new BuildException("Unsupported build architecture for Code Coverage");
+					}
+					string ProfileRuntimeLib = $"clang_rt.profile-{ProfileArchSuffix}.lib";
+					LinkEnvironment.Libraries.Add(FileReference.Combine(ProfileRuntimeDir, ProfileRuntimeLib));
+				}
+			}
 		}
 
 		protected virtual void AppendLibArguments(LinkEnvironment LinkEnvironment, List<string> Arguments)
