EpicGames.Build: Downgrade write lock message to trace

#rnx
#Synced-CL 48951480

diff --git a/Engine/Source/Programs/UnrealBuildTool/System/FileHashCache.cs b/Engine/Source/Programs/UnrealBuildTool/System/FileHashCache.cs
--- a/Engine/Source/Programs/UnrealBuildTool/System/FileHashCache.cs
+++ b/Engine/Source/Programs/UnrealBuildTool/System/FileHashCache.cs
@@ -441,7 +441,7 @@ namespace UnrealBuildTool
 
 					if (!isLocked)
 					{
-						logger.LogInformation("Get exception trying to write {Location} ({Message}). Will wait 1 second and try again", fileItem.Location, e.Message);
+						logger.LogTrace("Get exception trying to write {Location} ({Message}). Will wait 100ms and try again", fileItem.Location, e.Message);
 						Monitor.Enter(fileItem);
 						isLocked = true;
 					}
