[UBA]
* Fixed so UBT only retries local process three times before moving on

#Synced-CL 48993900

diff --git a/Engine/Source/Programs/UnrealBuildTool/Executors/UnrealBuildAccelerator/UBAExecutor.cs b/Engine/Source/Programs/UnrealBuildTool/Executors/UnrealBuildAccelerator/UBAExecutor.cs
--- a/Engine/Source/Programs/UnrealBuildTool/Executors/UnrealBuildAccelerator/UBAExecutor.cs
+++ b/Engine/Source/Programs/UnrealBuildTool/Executors/UnrealBuildAccelerator/UBAExecutor.cs
@@ -732,6 +732,12 @@ namespace UnrealBuildTool
 			}
 		}
 
+		class ProcessUserData
+		{
+			public required LinkedAction Action;
+			public int RetryCount = 0;
+		}
+
 		/// <summary>
 		/// Executes the provided actions
 		/// </summary>
@@ -777,7 +783,8 @@ namespace UnrealBuildTool
 
 				List<string> logLines = info.LogLines;
 
-				LinkedAction action = (LinkedAction)info.UserData;
+				ProcessUserData userData = (ProcessUserData)info.UserData;
+				LinkedAction action = userData.Action;
 
 				string additionalDescription = String.Empty;
 
@@ -861,7 +868,9 @@ namespace UnrealBuildTool
 								// If vfs paths are used, then native will not work
 								bool mustRunDetoured = !UBAConfig.CanRetryNatively || UBAConfig.bStoreObjFilesCompressed || action.RootPaths.bUseVfs;
 
-								if (mustRunDetoured)
+								++userData.RetryCount;
+
+								if (mustRunDetoured && userData.RetryCount <= 3)
 								{
 									logger.LogInformation("{Description} {StatusDescription}: Exited with error code {ExitCode}. This action will retry", action.CommandDescription, action.StatusDescription, exitCode);
 									++forcedRetryActions;
@@ -994,7 +1003,8 @@ namespace UnrealBuildTool
 				uint bucket = GetCacheBucket(logger, action);
 				uint memoryGroup = GetMemoryGroup(action);
 				ulong predictedMemoryUsage = 0; // uba would do a better job scheduling if we knew approx memory usage
-				scheduler.EnqueueProcess(startInfo, action.Weight, enableDetour, canRunRemotely, prereqActionsSortIndex, knownInputs, knownInputsCount, bucket, memoryGroup, predictedMemoryUsage, action);
+				ProcessUserData userData = new ProcessUserData() { Action = action };
+				scheduler.EnqueueProcess(startInfo, action.Weight, enableDetour, canRunRemotely, prereqActionsSortIndex, knownInputs, knownInputsCount, bucket, memoryGroup, predictedMemoryUsage, userData);
 
 				++totalActions;
 				cacheTotalActions += bucket != 0 ? 0 : 1;
