[UBA]
* Added option -UBAWritePlaceholders to UBT which can be used to write zero-sized exe/dll/pdb

#Synced-CL 48993900

diff --git a/Engine/Source/Programs/UnrealBuildTool/Executors/UnrealBuildAccelerator/UBAExecutor.cs b/Engine/Source/Programs/UnrealBuildTool/Executors/UnrealBuildAccelerator/UBAExecutor.cs
--- a/Engine/Source/Programs/UnrealBuildTool/Executors/UnrealBuildAccelerator/UBAExecutor.cs
+++ b/Engine/Source/Programs/UnrealBuildTool/Executors/UnrealBuildAccelerator/UBAExecutor.cs
@@ -285,6 +285,12 @@ namespace UnrealBuildTool
 				config.AddValue("Storage", "SkipCleanupOnShutdown", "true"); // For faster shutdown
 			}
 
+			if (UBAConfig.bWritePlaceholders)
+			{
+				config.AddValue("Session", "WritePlaceholders", "true");
+			}
+
+
 			EpicGames.UBA.ILogger ubaLogger = EpicGames.UBA.ILogger.GlobalLogger!;
 			_server = IServer.CreateServer(UBAConfig.MaxWorkers, UBAConfig.SendSize, ubaLogger, UBAConfig.bUseQuic);
 			_storageServer = IStorageServer.CreateStorageServer(_server, ubaLogger, new StorageServerCreateInfo(_rootDirRef.FullName, ((ulong)UBAConfig.StoreCapacityGb) * 1000 * 1000 * 1000, !UBAConfig.bStoreRaw, UBAConfig.Zone));
diff --git a/Engine/Source/Programs/UnrealBuildTool/Executors/UnrealBuildAccelerator/UnrealBuildAcceleratorConfig.cs b/Engine/Source/Programs/UnrealBuildTool/Executors/UnrealBuildAccelerator/UnrealBuildAcceleratorConfig.cs
--- a/Engine/Source/Programs/UnrealBuildTool/Executors/UnrealBuildAccelerator/UnrealBuildAcceleratorConfig.cs
+++ b/Engine/Source/Programs/UnrealBuildTool/Executors/UnrealBuildAccelerator/UnrealBuildAcceleratorConfig.cs
@@ -310,6 +310,13 @@ namespace UnrealBuildTool
 		[CommandLine("-UBANoDetour", Value = "false")]
 		public bool bAllowDetour { get; set; } = true;
 
+		/// <summary>
+		/// Disable all detouring making all actions run outside uba
+		/// </summary>
+		[XmlConfigFile(Category = "UnrealBuildAccelerator", Name = "WritePlaceholders")]
+		[CommandLine("-UBAWritePlaceholders", Value = "true")]
+		public bool bWritePlaceholders { get; set; } = false;
+
 		/// <summary>
 		/// Max number of cache download tasks that can execute in parallel
 		/// </summary>
