UnrealBuildTool: Fix NullReferenceException in IWYU

#rnx
#Synced-CL 48951480

diff --git a/Engine/Source/Programs/UnrealBuildTool/Modes/IWYUMode.cs b/Engine/Source/Programs/UnrealBuildTool/Modes/IWYUMode.cs
--- a/Engine/Source/Programs/UnrealBuildTool/Modes/IWYUMode.cs
+++ b/Engine/Source/Programs/UnrealBuildTool/Modes/IWYUMode.cs
@@ -458,9 +458,12 @@ namespace UnrealBuildTool
 				Logger.LogInformation($"Creating BuildTarget for {TargetName}...");
 				UEBuildTarget Target = UEBuildTarget.Create(TargetDescriptors[0], BuildConfiguration, Logger);
 
-				VCEnvironment env = Target.Rules.WindowsPlatform.Environment!;
-				systemPaths.Add(env.ToolChainDir.FullName.Replace('\\', '/'));
-				systemPaths.Add(env.WindowsSdkDir.FullName.Replace('\\', '/'));
+				if (TargetDescriptors[0].Platform.IsInGroup(UnrealPlatformGroup.Microsoft))
+				{
+					VCEnvironment env = Target.Rules.WindowsPlatform.Environment!;
+					systemPaths.Add(env.ToolChainDir.FullName.Replace('\\', '/'));
+					systemPaths.Add(env.WindowsSdkDir.FullName.Replace('\\', '/'));
+				}
 
 				Logger.LogInformation($"Calculating file filter for IWYU...");
 
