UnrealBuildTool: Don't regenerate cpp compile enviroment unnecessarily when generating a VSCode workspace

#rnx

#Synced-CL 49007367

diff --git a/Engine/Source/Programs/UnrealBuildTool/ProjectFiles/ProjectFileGenerator.cs b/Engine/Source/Programs/UnrealBuildTool/ProjectFiles/ProjectFileGenerator.cs
--- a/Engine/Source/Programs/UnrealBuildTool/ProjectFiles/ProjectFileGenerator.cs
+++ b/Engine/Source/Programs/UnrealBuildTool/ProjectFiles/ProjectFileGenerator.cs
@@ -2422,7 +2422,7 @@ namespace UnrealBuildTool
 					void Gen(UEBuildTarget Target)
 					{
 						using ITimelineEvent _ = Timeline.ScopeEvent(Target.TargetName);
-						AddTargetForIntellisense(Target, Logger);
+						ConcurrentBag<(UEBuildModuleCPP, CppCompileEnvironment)> ModulesForTarget = [];
 
 						// If the project generator just cares about the result of UEBuildTarget.Create, skip generating the compile environments.
 						if (ShouldGenerateIntelliSenseCompileEnvironments() && BuildTargetToProjectFile.TryGetValue(Target, out ProjectFile? TargetProjectFile))
@@ -2453,6 +2453,7 @@ namespace UnrealBuildTool
 
 										ConcurrentBag<(UEBuildModuleCPP, CppCompileEnvironment)> list = toAdd.GetOrAdd(ProjectFileForIDE, _ => new());
 										list.Add((Module, ModuleCompileEnvironment));
+										ModulesForTarget.Add((Module, ModuleCompileEnvironment));
 									}
 									Progress.Write(Interlocked.Increment(ref TasksFinished), TotalTasks);
 								}
@@ -2460,6 +2461,8 @@ namespace UnrealBuildTool
 							Parallel2.ForEach(Target.Binaries, Binary => ProcessBinary(Binary));
 						}
 
+						AddTargetForIntellisense(Target, ModulesForTarget, Logger);
+
 						Progress.Write(Interlocked.Increment(ref TasksFinished), TotalTasks);
 					}
 
@@ -2484,7 +2487,13 @@ namespace UnrealBuildTool
 			}
 		}
 
+		[Obsolete("Replace with 3 paramater version of AddTargetForIntellisense")]
 		protected virtual void AddTargetForIntellisense(UEBuildTarget Target, ILogger Logger)
+		{
+			AddTargetForIntellisense(Target, [], Logger);
+		}
+
+		protected virtual void AddTargetForIntellisense(UEBuildTarget Target, IEnumerable<(UEBuildModuleCPP, CppCompileEnvironment)> Modules, ILogger Logger)
 		{
 
 		}
diff --git a/Engine/Source/Programs/UnrealBuildTool/ProjectFiles/VisualStudioCode/VSCodeProjectFileGenerator.cs b/Engine/Source/Programs/UnrealBuildTool/ProjectFiles/VisualStudioCode/VSCodeProjectFileGenerator.cs
--- a/Engine/Source/Programs/UnrealBuildTool/ProjectFiles/VisualStudioCode/VSCodeProjectFileGenerator.cs
+++ b/Engine/Source/Programs/UnrealBuildTool/ProjectFiles/VisualStudioCode/VSCodeProjectFileGenerator.cs
@@ -314,11 +314,11 @@ namespace UnrealBuildTool
 			public readonly UnrealTargetPlatform Platform;
 			public readonly UnrealTargetConfiguration Configuration;
 			public readonly CppStandardVersion CppStandard;
-			public readonly FileReference? CompilerPath;
+			public readonly FileReference CompilerPath;
 			public readonly DirectoryReference? SysRootPath;
 			public readonly Dictionary<DirectoryReference, string> ModuleCommandLines;
 
-			public BuildTarget(string InName, TargetType InType, UnrealTargetPlatform InPlatform, UnrealTargetConfiguration InConfiguration, CppStandardVersion InCppStandard, FileReference? InCompilerPath, DirectoryReference? InSysRootPath, Dictionary<DirectoryReference, string> InModulesCommandLines)
+			public BuildTarget(string InName, TargetType InType, UnrealTargetPlatform InPlatform, UnrealTargetConfiguration InConfiguration, CppStandardVersion InCppStandard, FileReference InCompilerPath, DirectoryReference? InSysRootPath, Dictionary<DirectoryReference, string> InModulesCommandLines)
 			{
 				Name = InName;
 				Type = InType;
@@ -336,12 +336,17 @@ namespace UnrealBuildTool
 			}
 		}
 
-		protected override void AddTargetForIntellisense(UEBuildTarget Target, ILogger Logger)
+		protected override void AddTargetForIntellisense(UEBuildTarget Target, IEnumerable<(UEBuildModuleCPP, CppCompileEnvironment)> Modules, ILogger Logger)
 		{
-			base.AddTargetForIntellisense(Target, Logger);
+			base.AddTargetForIntellisense(Target, Modules, Logger);
+
+			if (!Modules.Any())
+			{
+				return;
+			}
 
 			bool UsingClang = true;
-			FileReference? CompilerPath = null;
+			FileReference CompilerPath;
 			DirectoryReference? SysRootPath = null;
 			if (OperatingSystem.IsWindows())
 			{
@@ -351,7 +356,12 @@ namespace UnrealBuildTool
 			}
 			else if (OperatingSystem.IsLinux())
 			{
-				CompilerPath = FileReference.FromString(LinuxCommon.WhichClang(Logger));
+				string? ClangPath = LinuxCommon.WhichClang(Logger);
+				if (ClangPath == null)
+				{
+					return;
+				}
+				CompilerPath = new(ClangPath);
 				string? InternalSDKPath = UEBuildPlatform.GetSDK(UnrealTargetPlatform.Linux)?.GetInternalSDKPath();
 				if (!String.IsNullOrEmpty(InternalSDKPath))
 				{
@@ -366,51 +376,48 @@ namespace UnrealBuildTool
 			}
 			else
 			{
-				throw new Exception("Unknown platform " + HostPlatform.ToString());
+				return;
 			}
 
 			// we do not need to keep track of which binary the invocation belongs to, only which target, as such we join all binaries into a single set
 			Dictionary<DirectoryReference, string> ModuleDirectoryToCompileCommand = new Dictionary<DirectoryReference, string>();
 
-			// Generate a compile environment for each module in the binary
-			CppCompileEnvironment GlobalCompileEnvironment = Target.CreateCompileEnvironmentForProjectFiles(Logger);
-			foreach (UEBuildBinary Binary in Target.Binaries)
+			CppStandardVersion CppStandard = CppStandardVersion.EngineDefault;
+			foreach ((UEBuildModuleCPP Module, CppCompileEnvironment ModuleCompileEnvironment) in Modules)
 			{
-				CppCompileEnvironment BinaryCompileEnvironment = Binary.CreateBinaryCompileEnvironment(GlobalCompileEnvironment);
-				foreach (UEBuildModuleCPP Module in Binary.Modules.OfType<UEBuildModuleCPP>())
+				List<FileReference> ForceIncludePaths = new List<FileReference>(ModuleCompileEnvironment.ForceIncludeFiles.Select(x => x.Location));
+				if (ModuleCompileEnvironment.PrecompiledHeaderIncludeFilename != null)
 				{
-					CppCompileEnvironment ModuleCompileEnvironment = Module.CreateCompileEnvironmentForIntellisense(Target.Rules, BinaryCompileEnvironment, Logger);
-
-					List<FileReference> ForceIncludePaths = new List<FileReference>(ModuleCompileEnvironment.ForceIncludeFiles.Select(x => x.Location));
-					if (ModuleCompileEnvironment.PrecompiledHeaderIncludeFilename != null)
-					{
-						ForceIncludePaths.Add(ModuleCompileEnvironment.PrecompiledHeaderIncludeFilename);
-					}
+					ForceIncludePaths.Add(ModuleCompileEnvironment.PrecompiledHeaderIncludeFilename);
+				}
 
-					StringBuilder CommandBuilder = new StringBuilder();
+				StringBuilder CommandBuilder = new StringBuilder();
 
-					foreach (FileReference ForceIncludeFile in ForceIncludePaths)
-					{
-						CommandBuilder.AppendFormat("{0} \"{1}\" {2}", UsingClang ? "-include" : "/FI", ForceIncludeFile.FullName, Environment.NewLine);
-					}
-					foreach (string Definition in ModuleCompileEnvironment.Definitions)
-					{
-						CommandBuilder.AppendFormat("{0} \"{1}\" {2}", UsingClang ? "-D" : "/D", Definition, Environment.NewLine);
-					}
-					foreach (DirectoryReference IncludePath in ModuleCompileEnvironment.UserIncludePaths)
-					{
-						CommandBuilder.AppendFormat("{0} \"{1}\" {2}", UsingClang ? "-I" : "/I", IncludePath, Environment.NewLine);
-					}
-					foreach (DirectoryReference IncludePath in ModuleCompileEnvironment.SystemIncludePaths)
-					{
-						CommandBuilder.AppendFormat("{0} \"{1}\" {2}", UsingClang ? "-I" : "/I", IncludePath, Environment.NewLine);
-					}
+				foreach (FileReference ForceIncludeFile in ForceIncludePaths)
+				{
+					CommandBuilder.AppendFormat("{0} \"{1}\" {2}", UsingClang ? "-include" : "/FI", ForceIncludeFile.FullName, Environment.NewLine);
+				}
+				foreach (string Definition in ModuleCompileEnvironment.Definitions)
+				{
+					CommandBuilder.AppendFormat("{0} \"{1}\" {2}", UsingClang ? "-D" : "/D", Definition, Environment.NewLine);
+				}
+				foreach (DirectoryReference IncludePath in ModuleCompileEnvironment.UserIncludePaths)
+				{
+					CommandBuilder.AppendFormat("{0} \"{1}\" {2}", UsingClang ? "-I" : "/I", IncludePath, Environment.NewLine);
+				}
+				foreach (DirectoryReference IncludePath in ModuleCompileEnvironment.SystemIncludePaths)
+				{
+					CommandBuilder.AppendFormat("{0} \"{1}\" {2}", UsingClang ? "-I" : "/I", IncludePath, Environment.NewLine);
+				}
 
-					ModuleDirectoryToCompileCommand.TryAdd(Module.ModuleDirectory, CommandBuilder.ToString());
+				ModuleDirectoryToCompileCommand.TryAdd(Module.ModuleDirectory, CommandBuilder.ToString());
+				if (Module.Rules.CppStandard != null && Module.Rules.CppStandard > CppStandard)
+				{
+					CppStandard = (CppStandardVersion)Module.Rules.CppStandard;
 				}
 			}
 
-			BuildTargets.Add(new BuildTarget(Target.TargetName, Target.TargetType, Target.Platform, Target.Configuration, GlobalCompileEnvironment.CppStandard, CompilerPath, SysRootPath, ModuleDirectoryToCompileCommand));
+			BuildTargets.Add(new BuildTarget(Target.TargetName, Target.TargetType, Target.Platform, Target.Configuration, CppStandard, CompilerPath, SysRootPath, ModuleDirectoryToCompileCommand));
 		}
 
 		private class ProjectData
@@ -642,7 +649,7 @@ namespace UnrealBuildTool
 							}
 
 							string Name = String.Format("{0} {1} {2} {3} ({4})", ProjectTarget.Name, ProjectTarget.Type, BuildTarget.Platform, BuildTarget.Configuration, Project.Name);
-							WriteConfiguration(Name, Project.Name, Project.SourceProject.SourceFiles.Select(x => x.Reference), BuildTarget.CppStandard, BuildTarget.CompilerPath!, BuildTarget.SysRootPath, BuildTarget.ModuleCommandLines, OutFile, OutputDirectory);
+							WriteConfiguration(Name, Project.Name, Project.SourceProject.SourceFiles.Select(x => x.Reference), BuildTarget.CppStandard, BuildTarget.CompilerPath, BuildTarget.SysRootPath, BuildTarget.ModuleCommandLines, OutFile, OutputDirectory);
 
 							CompilerPath = BuildTarget.CompilerPath;
 
@@ -679,7 +686,7 @@ namespace UnrealBuildTool
 						DefaultConfigName = "Win32";
 					}
 
-					WriteConfiguration(DefaultConfigName, "Default", AllSourceFiles, CppStandard, CompilerPath!, SysRootPath, AllModuleCommandLines, OutFile, OutputDirectory);
+					WriteConfiguration(DefaultConfigName, "Default", AllSourceFiles, CppStandard, CompilerPath, SysRootPath, AllModuleCommandLines, OutFile, OutputDirectory);
 				}
 				OutFile.EndArray();
 			}
@@ -688,7 +695,7 @@ namespace UnrealBuildTool
 			OutFile.Write(FileReference.Combine(OutputDirectory, "c_cpp_properties.json"));
 		}
 
-		private void WriteConfiguration(string Name, string ProjectName, IEnumerable<FileReference> SourceFiles, CppStandardVersion CppStandard, FileReference CompilerPath, DirectoryReference? SysRootPath, Dictionary<DirectoryReference, string> ModuleCommandLines, JsonFile OutFile, DirectoryReference OutputDirectory)
+		private void WriteConfiguration(string Name, string ProjectName, IEnumerable<FileReference> SourceFiles, CppStandardVersion CppStandard, FileReference? CompilerPath, DirectoryReference? SysRootPath, Dictionary<DirectoryReference, string> ModuleCommandLines, JsonFile OutFile, DirectoryReference OutputDirectory)
 		{
 			OutFile.BeginObject();
 
@@ -728,6 +735,7 @@ namespace UnrealBuildTool
 
 			switch (CppStandard)
 			{
+				default:
 				case CppStandardVersion.Cpp20:
 					OutFile.AddField("cppStandard", "c++20");
 					break;
@@ -735,8 +743,8 @@ namespace UnrealBuildTool
 					OutFile.AddField("cppStandard", "c++23");
 					break;
 				case CppStandardVersion.Latest:
-				default:
-					throw new BuildException($"Unsupported C++ standard type set: {CppStandard}");
+					OutFile.AddField("cppStandard", "c++26");
+					break;
 			}
 
 			string CompilerPrefix = HostPlatform == UnrealTargetPlatform.Win64 ? "msvc" : "clang";
@@ -749,10 +757,10 @@ namespace UnrealBuildTool
 				OutFile.AddField("configurationProvider", "epic.ue");
 			}
 
-			if (ShouldGenerateIntelliSenseData())
+			if (ShouldGenerateIntelliSenseData() && CompilerPath != null)
 			{
 				FileReference CompileCommands = FileReference.Combine(OutputDirectory, String.Format("compileCommands_{0}.json", ProjectName));
-				WriteCompileCommands(CompileCommands, SourceFiles, CompilerPath!, ModuleCommandLines);
+				WriteCompileCommands(CompileCommands, SourceFiles, CompilerPath, ModuleCommandLines);
 				OutFile.AddField("compileCommands", MakePathString(CompileCommands, bInAbsolute: true, bForceSkipQuotes: true));
 			}
 
