UnrealBuildTool: Analysis cleanup of QueryMode

#rnx
#Synced-CL 49007367

diff --git a/Engine/Source/Programs/UnrealBuildTool/Modes/QueryMode.cs b/Engine/Source/Programs/UnrealBuildTool/Modes/QueryMode.cs
--- a/Engine/Source/Programs/UnrealBuildTool/Modes/QueryMode.cs
+++ b/Engine/Source/Programs/UnrealBuildTool/Modes/QueryMode.cs
@@ -29,18 +29,18 @@ namespace UnrealBuildTool
 	{
 		internal class CompileSettings
 		{
-			public List<string> IncludePaths { get; set; } = new();
-			public List<string> Defines { get; set; } = new();
+			public List<string> IncludePaths { get; set; } = [];
+			public List<string> Defines { get; set; } = [];
 			public string? Standard { get; set; }
-			public List<string> ForcedIncludes { get; set; } = new();
+			public List<string> ForcedIncludes { get; set; } = [];
 			public string? CompilerPath { get; set; }
-			public List<string> CompilerArgs { get; set; } = new();
-			public Dictionary<string, List<string>> CompilerAdditionalArgs { get; set; } = new();
+			public List<string> CompilerArgs { get; set; } = [];
+			public Dictionary<string, List<string>> CompilerAdditionalArgs { get; set; } = [];
 			public string? WindowsSdkVersion { get; set; }
 		}
 
-		public ConcurrentDictionary<UEBuildModule, CompileSettings> ModuleToCompileSettings = new();
-		public ConcurrentDictionary<DirectoryReference, UEBuildModule> DirToModule = new();
+		public ConcurrentDictionary<UEBuildModule, CompileSettings> ModuleToCompileSettings = [];
+		public ConcurrentDictionary<DirectoryReference, UEBuildModule> DirToModule = [];
 
 		public UEBuildModule? FindModuleForFile(FileReference File)
 		{
@@ -75,7 +75,7 @@ namespace UnrealBuildTool
 	{
 		public string? Description { get; set; }
 		public string? BinaryPath { get; set; }
-		public List<string> Arguments { get; set; } = new();
+		public List<string> Arguments { get; set; } = [];
 	}
 
 	internal class TargetConfigs
@@ -83,8 +83,8 @@ namespace UnrealBuildTool
 		public string TargetPath { get; set; } = String.Empty;
 		public string ProjectPath { get; set; } = String.Empty;
 		public string TargetType { get; set; } = String.Empty;
-		public List<string> Platforms { get; set; } = new();
-		public List<string> Configurations { get; set; } = new();
+		public List<string> Platforms { get; set; } = [];
+		public List<string> Configurations { get; set; } = [];
 	}
 
 	[ToolMode("Query", ToolModeOptions.BuildPlatforms | ToolModeOptions.XmlConfig | ToolModeOptions.UseStartupTraceListener)]
@@ -211,7 +211,7 @@ namespace UnrealBuildTool
 			{
 				Utils.TryParseProjectFileArgument(Arguments, Logger, out FileReference? ProjectFileArg);
 
-				List<UnrealTargetPlatform> Platforms = new();
+				List<UnrealTargetPlatform> Platforms = [];
 				foreach (UnrealTargetPlatform Platform in UnrealTargetPlatform.GetValidPlatforms())
 				{
 					// If there is a build platform present, add it to the SupportedPlatforms list
@@ -225,10 +225,10 @@ namespace UnrealBuildTool
 					}
 				}
 
-				List<UnrealTargetConfiguration> AllowedTargetConfigurations = new List<UnrealTargetConfiguration>();
-				AllowedTargetConfigurations = Enum.GetValues(typeof(UnrealTargetConfiguration)).Cast<UnrealTargetConfiguration>().ToList();
+				List<UnrealTargetConfiguration> AllowedTargetConfigurations = [];
+				AllowedTargetConfigurations = [.. Enum.GetValues(typeof(UnrealTargetConfiguration)).Cast<UnrealTargetConfiguration>()];
 
-				List<UnrealTargetConfiguration> Configurations = new();
+				List<UnrealTargetConfiguration> Configurations = [];
 				foreach (UnrealTargetConfiguration CurConfiguration in AllowedTargetConfigurations)
 				{
 					if (CurConfiguration != UnrealTargetConfiguration.Unknown)
@@ -240,10 +240,10 @@ namespace UnrealBuildTool
 					}
 				}
 
-				List<FileReference> Projects = ProjectFileArg != null ? new List<FileReference>(new[] { ProjectFileArg }) : NativeProjects.EnumerateProjectFiles(Logger).ToList();
+				List<FileReference> Projects = ProjectFileArg != null ? [ProjectFileArg] : [.. NativeProjects.EnumerateProjectFiles(Logger)];
 				List<FileReference> AllTargetFiles = ProjectFileGenerator.DiscoverTargets(Projects, Logger, null, Platforms, bIncludeEngineSource: bIncludeEngineSource, bIncludeTempTargets: false);
 
-				Dictionary<string, TargetConfigs> Targets = new();
+				Dictionary<string, TargetConfigs> Targets = [];
 				string? DefaultTarget = null;
 				foreach (FileReference TargetFilePath in AllTargetFiles)
 				{
@@ -251,18 +251,18 @@ namespace UnrealBuildTool
 					{
 						FileReference? ProjectPath = Projects.FirstOrDefault(p => TargetFilePath.IsUnderDirectory(p.Directory));
 						string TargetType = String.Empty;
-						List<UnrealTargetConfiguration> SupportedConfigurations = new(Configurations);
-						List<UnrealTargetPlatform> SupportedPlatforms = new(Platforms);
+						List<UnrealTargetConfiguration> SupportedConfigurations = [.. Configurations];
+						List<UnrealTargetPlatform> SupportedPlatforms = [.. Platforms];
 						if (bLoadTargets)
 						{
-							List<string> RawArgs = new List<string> { TargetFilePath.GetFileNameWithoutAnyExtensions(), UnrealTargetConfiguration.Development.ToString(), UnrealTargetPlatform.Win64.ToString() };
+							List<string> RawArgs = [TargetFilePath.GetFileNameWithoutAnyExtensions(), UnrealTargetConfiguration.Development.ToString(), UnrealTargetPlatform.Win64.ToString()];
 							if (ProjectPath != null)
 							{
 								RawArgs.Add($"-Project={ProjectPath.FullName}");
 							}
-							CommandLineArguments Args = new CommandLineArguments(RawArgs.ToArray());
+							CommandLineArguments Args = new CommandLineArguments([.. RawArgs]);
 
-							List<TargetDescriptor> TargetDescriptors = new();
+							List<TargetDescriptor> TargetDescriptors = [];
 							TargetDescriptor.ParseSingleCommandLine(Args, false, false, false, TargetDescriptors, NullLogger.Instance);
 
 							// Ensure the intermediate environment does not conflict with normal builds
@@ -281,8 +281,8 @@ namespace UnrealBuildTool
 								}
 							}
 							TargetType = CurrentTarget.Rules.Type.ToString();
-							SupportedConfigurations = CurrentTarget.Rules.SupportedConfigurations.ToList();
-							SupportedPlatforms = CurrentTarget.Rules.SupportedPlatforms.ToList();
+							SupportedConfigurations = [.. CurrentTarget.Rules.SupportedConfigurations];
+							SupportedPlatforms = [.. CurrentTarget.Rules.SupportedPlatforms];
 						}
 
 						string TargetName = TargetFilePath.GetFileNameWithoutAnyExtensions();
@@ -291,8 +291,8 @@ namespace UnrealBuildTool
 							TargetPath = TargetFilePath.FullName,
 							ProjectPath = ProjectPath?.ToString() ?? String.Empty,
 							TargetType = TargetType,
-							Configurations = SupportedConfigurations.Select(x => x.ToString()).ToList(),
-							Platforms = SupportedPlatforms.Select(x => x.ToString()).ToList()
+							Configurations = [.. SupportedConfigurations.Select(x => x.ToString())],
+							Platforms = [.. SupportedPlatforms.Select(x => x.ToString())]
 						});
 						if (DefaultTarget == null || TargetName == "UnrealEditor")
 						{
@@ -343,15 +343,15 @@ namespace UnrealBuildTool
 			}
 
 			Utils.TryParseProjectFileArgument(Arguments, Logger, out FileReference? ProjectFileArg);
-			List<string> RawArgs = new List<string> { TargetName, TargetConfiguration, TargetPlatform };
+			List<string> RawArgs = [TargetName, TargetConfiguration, TargetPlatform];
 			if (ProjectFileArg != null)
 			{
 				RawArgs.Add(ProjectFileArg.ToString());
 			}
 			RawArgs.AddRange(Arguments.GetUnusedArguments());
 			RawArgs.Add("-NoPCHChain"); // Currently unsupported
-			CommandLineArguments Args = new CommandLineArguments(RawArgs.ToArray());
-			List<TargetDescriptor> TargetDescriptors = new();
+			CommandLineArguments Args = new CommandLineArguments([.. RawArgs]);
+			List<TargetDescriptor> TargetDescriptors = [];
 
 			TargetDescriptor.ParseSingleCommandLine(Args, false, false, false, TargetDescriptors, Logger);
 			if (TargetDescriptors.Count == 0)
@@ -474,11 +474,11 @@ namespace UnrealBuildTool
 							}
 							Settings.Defines = ModuleCompileEnvironment.Definitions;
 							Settings.Standard = ModuleCompileEnvironment.CppStandard.ToString();
-							Settings.ForcedIncludes = ModuleCompileEnvironment.ForceIncludeFiles.Select(x => x.ToString()).ToList();
+							Settings.ForcedIncludes = [.. ModuleCompileEnvironment.ForceIncludeFiles.Select(x => x.ToString())];
 							Settings.CompilerPath = TargetToolChain.GetCppCompilerPath()?.ToString();
-							Settings.CompilerArgs = TargetToolChain.GetGlobalCommandLineArgs(ModuleCompileEnvironmentForArgs).ToList();
-							Settings.CompilerAdditionalArgs.Add("c", TargetToolChain.GetCCommandLineArgs(ModuleCompileEnvironmentForArgs).ToList());
-							Settings.CompilerAdditionalArgs.Add("cpp", TargetToolChain.GetCPPCommandLineArgs(ModuleCompileEnvironmentForArgs).ToList());
+							Settings.CompilerArgs = [.. TargetToolChain.GetGlobalCommandLineArgs(ModuleCompileEnvironmentForArgs)];
+							Settings.CompilerAdditionalArgs.Add("c", [.. TargetToolChain.GetCCommandLineArgs(ModuleCompileEnvironmentForArgs)]);
+							Settings.CompilerAdditionalArgs.Add("cpp", [.. TargetToolChain.GetCPPCommandLineArgs(ModuleCompileEnvironmentForArgs)]);
 							Settings.WindowsSdkVersion = CurrentTarget.Rules.WindowsPlatform.WindowsSdkVersion;
 							CurrentTargetIntellisenseInfo.ModuleToCompileSettings.TryAdd(Module, Settings);
 
@@ -492,7 +492,7 @@ namespace UnrealBuildTool
 					DirToModule = CurrentTargetIntellisenseInfo?.DirToModule.ToImmutableSortedDictionary(x => x.Key.ToString(), x => x.Value.Name),
 					ModuleToCompileSettings = CurrentTargetIntellisenseInfo?.ModuleToCompileSettings.ToImmutableSortedDictionary(x => x.Key.Name, x => x.Value),
 					LaunchSettings = CurrentLaunchSettings,
-					IsTestTarget = CurrentTarget.Rules.IsTestTarget
+					CurrentTarget.Rules.IsTestTarget
 				};
 
 				await WriteResultsAsync(Result, JsonOptions, Logger);
@@ -525,15 +525,15 @@ namespace UnrealBuildTool
 			}
 
 			Utils.TryParseProjectFileArgument(Arguments, Logger, out FileReference? ProjectFileArg);
-			List<string> RawArgs = new List<string> { TargetName, TargetConfiguration, TargetPlatform };
+			List<string> RawArgs = [TargetName, TargetConfiguration, TargetPlatform];
 			if (ProjectFileArg != null)
 			{
 				RawArgs.Add(ProjectFileArg.ToString());
 			}
 			RawArgs.AddRange(Arguments.GetUnusedArguments());
 			RawArgs.Add("-NoPCHChain"); // Currently unsupported
-			CommandLineArguments Args = new CommandLineArguments(RawArgs.ToArray());
-			List<TargetDescriptor> TargetDescriptors = new();
+			CommandLineArguments Args = new CommandLineArguments([.. RawArgs]);
+			List<TargetDescriptor> TargetDescriptors = [];
 
 			TargetDescriptor.ParseSingleCommandLine(Args, false, false, false, TargetDescriptors, Logger);
 			if (TargetDescriptors.Count == 0)
