[UBT]
* Added option -UBASharedMemoryTempFile=<filename> which can be used to set a temporary file as backing for shared memory. Very useful to reduce committed memory charge but with the risk of more I/O under RAM pressure

#Synced-CL 49025682

diff --git a/Engine/Source/Programs/UnrealBuildTool/Executors/UnrealBuildAccelerator/UBAExecutor.cs b/Engine/Source/Programs/UnrealBuildTool/Executors/UnrealBuildAccelerator/UBAExecutor.cs
--- a/Engine/Source/Programs/UnrealBuildTool/Executors/UnrealBuildAccelerator/UBAExecutor.cs
+++ b/Engine/Source/Programs/UnrealBuildTool/Executors/UnrealBuildAccelerator/UBAExecutor.cs
@@ -290,6 +290,16 @@ namespace UnrealBuildTool
 				config.AddValue("Session", "WritePlaceholders", "true");
 			}
 
+			if (!String.IsNullOrEmpty(UBAConfig.SharedMemoryTempFile))
+			{
+				config.AddValue("Storage", "SharedMemoryTempFile", UBAConfig.SharedMemoryTempFile);
+			}
+
+			if (!UBAConfig.bAllowDetour)
+			{
+				config.AddValue("Session", "AllowMemoryMaps", "false");
+			}
+
 			EpicGames.UBA.ILogger ubaLogger = EpicGames.UBA.ILogger.GlobalLogger!;
 			_server = IServer.CreateServer(UBAConfig.MaxWorkers, UBAConfig.SendSize, ubaLogger, UBAConfig.bUseQuic);
 			_storageServer = IStorageServer.CreateStorageServer(_server, ubaLogger, new StorageServerCreateInfo(_rootDirRef.FullName, ((ulong)UBAConfig.StoreCapacityGb) * 1000 * 1000 * 1000, !UBAConfig.bStoreRaw, UBAConfig.Zone));
diff --git a/Engine/Source/Programs/UnrealBuildTool/Executors/UnrealBuildAccelerator/UnrealBuildAcceleratorConfig.cs b/Engine/Source/Programs/UnrealBuildTool/Executors/UnrealBuildAccelerator/UnrealBuildAcceleratorConfig.cs
--- a/Engine/Source/Programs/UnrealBuildTool/Executors/UnrealBuildAccelerator/UnrealBuildAcceleratorConfig.cs
+++ b/Engine/Source/Programs/UnrealBuildTool/Executors/UnrealBuildAccelerator/UnrealBuildAcceleratorConfig.cs
@@ -311,12 +311,19 @@ namespace UnrealBuildTool
 		public bool bAllowDetour { get; set; } = true;
 
 		/// <summary>
-		/// Disable all detouring making all actions run outside uba
+		/// Write zero-byte placeholder files for exe/dll/pdb instead of the full files. This can be useful for validation builds or builds just populating cache
 		/// </summary>
 		[XmlConfigFile(Category = "UnrealBuildAccelerator", Name = "WritePlaceholders")]
 		[CommandLine("-UBAWritePlaceholders", Value = "true")]
 		public bool bWritePlaceholders { get; set; } = false;
 
+		/// <summary>
+		/// Makes the shared memory used in UBA to be backed by a Temp+DeleteOnClose file. This trades committed memory charge for the risk of more I/O under RAM pressure
+		/// </summary>
+		[XmlConfigFile(Category = "UnrealBuildAccelerator", Name = "SharedMemoryTempFile")]
+		[CommandLine("-UBASharedMemoryTempFile")]
+		public string SharedMemoryTempFile { get; set; } = String.Empty;
+
 		/// <summary>
 		/// Max number of cache download tasks that can execute in parallel
 		/// </summary>
