Persist the result of executed actions to disk, in order to enable better visibility of the cost of individual actions.

#Synced-CL 49029782

diff --git a/Engine/Source/Programs/Shared/EpicGames.IoHash/EpicGames.IoHash.csproj b/Engine/Source/Programs/Shared/EpicGames.IoHash/EpicGames.IoHash.csproj
--- a/Engine/Source/Programs/Shared/EpicGames.IoHash/EpicGames.IoHash.csproj
+++ b/Engine/Source/Programs/Shared/EpicGames.IoHash/EpicGames.IoHash.csproj
@@ -36,6 +36,7 @@
     </Compile>
   </ItemGroup>
   <ItemGroup>
+    <!-- Blake3 2.0.0 has a massive performance regression on Windows x64 - takes 2.25x as long for large (1MB) inputs -->
     <PackageReference Include="Blake3" Version="1.1.0" />
     <PackageReference Include="System.Memory" Version="4.6.3" />
   </ItemGroup>
diff --git a/Engine/Source/Programs/Shared/EpicGames.UBA/IScheduler.cs b/Engine/Source/Programs/Shared/EpicGames.UBA/IScheduler.cs
--- a/Engine/Source/Programs/Shared/EpicGames.UBA/IScheduler.cs
+++ b/Engine/Source/Programs/Shared/EpicGames.UBA/IScheduler.cs
@@ -74,7 +74,7 @@ namespace EpicGames.UBA
 		/// <summary>
 		/// Peak memory used, requires a job object so will only be non-zero for Windows hosts.
 		/// </summary>
-		long PeakMemoryUsed { get; }
+		ulong PeakMemoryUsed { get; }
 
 		/// <summary>
 		/// UserData that was provided in EnqueueProcess
diff --git a/Engine/Source/Programs/Shared/EpicGames.UBA/Impl/SchedulerImpl.cs b/Engine/Source/Programs/Shared/EpicGames.UBA/Impl/SchedulerImpl.cs
--- a/Engine/Source/Programs/Shared/EpicGames.UBA/Impl/SchedulerImpl.cs
+++ b/Engine/Source/Programs/Shared/EpicGames.UBA/Impl/SchedulerImpl.cs
@@ -60,7 +60,7 @@ namespace EpicGames.UBA.Impl
 		public string? ExecutingHost => Marshal.PtrToStringAuto(ProcessHandle_GetExecutingHost(_processHandle));
 		public TimeSpan TotalProcessorTime => new((long)ProcessHandle_GetTotalProcessorTime(_processHandle));
 		public TimeSpan TotalWallTime => new((long)ProcessHandle_GetTotalWallTime(_processHandle));
-		public long PeakMemoryUsed => (long)ProcessHandle_GetPeakMemory(_processHandle);
+		public ulong PeakMemoryUsed => ProcessHandle_GetPeakMemory(_processHandle);
 		public object UserData => _userData;
 		public nint ProcessHandle => _processHandle;
 	}
diff --git a/Engine/Source/Programs/UnrealBuildTool/Executors/ActionLogger.cs b/Engine/Source/Programs/UnrealBuildTool/Executors/ActionLogger.cs
--- a/Engine/Source/Programs/UnrealBuildTool/Executors/ActionLogger.cs
+++ b/Engine/Source/Programs/UnrealBuildTool/Executors/ActionLogger.cs
@@ -20,7 +20,7 @@ namespace UnrealBuildTool
 	/// <param name="PeakMemoryUsed">Peak Memory Used</param>
 	/// <param name="Action">Action</param>
 	/// <param name="AdditionalDescription">Additional description of action</param>
-	record ExecuteResults(List<string> LogLines, int ExitCode, TimeSpan ExecutionTime, TimeSpan ProcessorTime, long PeakMemoryUsed, LinkedAction Action, string? AdditionalDescription = null);
+	record ExecuteResults(List<string> LogLines, int ExitCode, TimeSpan ExecutionTime, TimeSpan ProcessorTime, ulong PeakMemoryUsed, LinkedAction Action, string? AdditionalDescription = null);
 
 	class ActionLogger : IDisposable
 	{
@@ -152,7 +152,7 @@ namespace UnrealBuildTool
 			int exitCode = executeTaskResult?.ExitCode ?? Int32.MaxValue;
 			TimeSpan executionTime = executeTaskResult?.ExecutionTime ?? TimeSpan.Zero;
 			TimeSpan processorTime = executeTaskResult?.ProcessorTime ?? TimeSpan.Zero;
-			long peakMemory = executeTaskResult?.PeakMemoryUsed ?? 0;
+			ulong peakMemory = executeTaskResult?.PeakMemoryUsed ?? 0;
 			string? additionalDescription = executeTaskResult?.AdditionalDescription;
 
 			// Write it to the log
@@ -200,7 +200,7 @@ namespace UnrealBuildTool
 				{
 					if (processorTime.Ticks > 0 && peakMemory > 0)
 					{
-						compilationTimes = $" (Wall: {executionTime.TotalSeconds:0.00}s CPU: {processorTime.TotalSeconds:0.00}s Mem: {StringUtils.FormatBytesString(peakMemory)})";
+						compilationTimes = $" (Wall: {executionTime.TotalSeconds:0.00}s CPU: {processorTime.TotalSeconds:0.00}s Mem: {StringUtils.FormatBytesString((long)peakMemory)})";
 					}
 					else if (processorTime.Ticks > 0)
 					{
@@ -208,7 +208,7 @@ namespace UnrealBuildTool
 					}
 					else if (peakMemory > 0)
 					{
-						compilationTimes = $" (Wall: {executionTime.TotalSeconds:0.00}s Mem: {StringUtils.FormatBytesString(peakMemory)})";
+						compilationTimes = $" (Wall: {executionTime.TotalSeconds:0.00}s Mem: {StringUtils.FormatBytesString((long)peakMemory)})";
 					}
 					else
 					{
@@ -366,7 +366,7 @@ namespace UnrealBuildTool
 					string Description = $"{(Action.CommandDescription ?? Action.CommandPath.GetFileNameWithoutExtension())} {Action.StatusDescription}".Trim();
 					if (Result.ProcessorTime.Ticks > 0 && Result.PeakMemoryUsed > 0)
 					{
-						Logger.Log(LogLevel, "{Description} [ Wall Time {ExecutionTime:0.00}s / CPU Time {ProcessorTime:0.00}s / Mem {Memory} ]", Description, Result.ExecutionTime.TotalSeconds, Result.ProcessorTime.TotalSeconds, StringUtils.FormatBytesString(Result.PeakMemoryUsed));
+						Logger.Log(LogLevel, "{Description} [ Wall Time {ExecutionTime:0.00}s / CPU Time {ProcessorTime:0.00}s / Mem {Memory} ]", Description, Result.ExecutionTime.TotalSeconds, Result.ProcessorTime.TotalSeconds, StringUtils.FormatBytesString((long)Result.PeakMemoryUsed));
 					}
 					else if (Result.ProcessorTime.Ticks > 0)
 					{
@@ -374,7 +374,7 @@ namespace UnrealBuildTool
 					}
 					else if (Result.PeakMemoryUsed > 0)
 					{
-						Logger.Log(LogLevel, "{Description} [ Time {ExecutionTime:0.00}s / Mem {Memory} ]", Description, Result.ExecutionTime.TotalSeconds, StringUtils.FormatBytesString(Result.PeakMemoryUsed));
+						Logger.Log(LogLevel, "{Description} [ Time {ExecutionTime:0.00}s / Mem {Memory} ]", Description, Result.ExecutionTime.TotalSeconds, StringUtils.FormatBytesString((long)Result.PeakMemoryUsed));
 					}
 					else
 					{
diff --git a/Engine/Source/Programs/UnrealBuildTool/Executors/UnrealBuildAccelerator/UBAExecutor.cs b/Engine/Source/Programs/UnrealBuildTool/Executors/UnrealBuildAccelerator/UBAExecutor.cs
--- a/Engine/Source/Programs/UnrealBuildTool/Executors/UnrealBuildAccelerator/UBAExecutor.cs
+++ b/Engine/Source/Programs/UnrealBuildTool/Executors/UnrealBuildAccelerator/UBAExecutor.cs
@@ -13,6 +13,8 @@ using EpicGames.Core;
 using EpicGames.UBA;
 using Microsoft.Extensions.Logging;
 using UnrealBuildBase;
+using UnrealBuildTool.Actions.ResultStore;
+using UnrealBuildTool.Storage.Impl;
 using ILogger = Microsoft.Extensions.Logging.ILogger;
 
 namespace UnrealBuildTool
@@ -27,6 +29,7 @@ namespace UnrealBuildTool
 		readonly Task _initTask;
 		readonly CancellationTokenSource _cancellationSource;
 		readonly bool _bCompactOutput;
+		readonly IActionResultStore _resultStore;
 
 		IServer? _server;
 		IStorageServer? _storageServer;
@@ -167,6 +170,8 @@ namespace UnrealBuildTool
 			_cancellationSource = new();
 			_initTask = new(() => { Init(targetDescriptors, logger); }, TaskCreationOptions.LongRunning);
 			_initTask.Start();
+
+			_resultStore = new ContentAddressableActionResultStore(new FileSystemStorageProvider());
 		}
 
 		private void Init(IEnumerable<TargetDescriptor> targetDescriptors, ILogger logger)
@@ -581,7 +586,7 @@ namespace UnrealBuildTool
 
 				OnSessionCreated?.Invoke(_session!);
 
-				bool success = ExecuteActionsInternal(logger, inputActions, _server!, _storageServer!, _session!, _cacheClients.Values);
+				bool success = await ExecuteActionsInternalAsync(logger, inputActions, _server!, _storageServer!, _session!, _cacheClients.Values);
 
 				if (!UBAConfig.bDisableRemote)
 				{
@@ -757,7 +762,7 @@ namespace UnrealBuildTool
 		/// Executes the provided actions
 		/// </summary>
 		/// <returns>True if all the tasks successfully executed, or false if any of them failed.</returns>
-		bool ExecuteActionsInternal(ILogger logger, IEnumerable<LinkedAction> inputActions, IServer server, IStorageServer storage, ISessionServer session, IEnumerable<ICacheClient> cache)
+		private async Task<bool> ExecuteActionsInternalAsync(ILogger logger, IEnumerable<LinkedAction> inputActions, IServer server, IStorageServer storage, ISessionServer session, IEnumerable<ICacheClient> cache)
 		{
 			_ubaStartTimeUtc = DateTime.UtcNow;
 
@@ -782,11 +787,13 @@ namespace UnrealBuildTool
 			int remoteProcessedActions = 0;
 			int localRetryActions = 0;
 			int forcedRetryActions = 0;
-			long peakProcessMemory = 0;
-			long peakTotalMemory = 0;
+			ulong peakProcessMemory = 0;
+			ulong peakTotalMemory = 0;
 			bool shouldCancel = false;
 			string cancelReason = "user cancellation";
 
+			List<Task> resultStoreTasks = [];
+
 			scheduler.SetProcessFinishedCallback((info) =>
 			{
 				int exitCode = info.ExitCode;
@@ -923,7 +930,7 @@ namespace UnrealBuildTool
 
 				TimeSpan processorTime = TimeSpan.Zero;
 				TimeSpan executionTime = TimeSpan.Zero;
-				long peakMemory = 0;
+				ulong peakMemory = 0;
 				if (executionType != ProcessExecutionType.Cache)
 				{
 					processorTime = info.TotalProcessorTime;
@@ -980,7 +987,9 @@ namespace UnrealBuildTool
 					}
 				}
 
-				actionLogger.AddActionToLog(new(logLines, exitCode, executionTime, processorTime, peakMemory, action, additionalDescription));
+				ExecuteResults result = new(logLines, exitCode, executionTime, processorTime, peakMemory, action, additionalDescription);
+				actionLogger.AddActionToLog(result);
+				resultStoreTasks.Add(Task.Run(() => _resultStore.StoreResultAsync(action, ActionResult.From(result), _cancellationSource.Token)));
 
 				return ProcessFinishedResponse.None;
 			});
@@ -1092,6 +1101,17 @@ namespace UnrealBuildTool
 				{
 					logger.LogInformation("Cancelled execution after {Actions}/{TotalActions} actions ({Reason})", succeededActions + failedActions + cacheHitActions, totalActions, cancelReason);
 				}
+
+				Task whenAllTask = Task.WhenAll(resultStoreTasks);
+				try
+				{
+					await whenAllTask;
+				}
+				catch (Exception) when (whenAllTask.IsFaulted)
+				{
+					logger.LogWarning(whenAllTask.Exception, "Some action result store tasks did not succeed");
+				}
+
 				return success;
 			}
 			finally
diff --git a/Engine/Source/Programs/UnrealBuildTool/System/Telemetry.cs b/Engine/Source/Programs/UnrealBuildTool/System/Telemetry.cs
--- a/Engine/Source/Programs/UnrealBuildTool/System/Telemetry.cs
+++ b/Engine/Source/Programs/UnrealBuildTool/System/Telemetry.cs
@@ -117,8 +117,8 @@ namespace UnrealBuildTool
 		public DateTime StartUTC { get; }
 		[JsonConverter(typeof(DurationDataRouterEventConverter))]
 		public TimeSpan Duration { get; }
-		public long PeakProcessMemory { get; }
-		public long PeakTotalMemory { get; }
+		public ulong PeakProcessMemory { get; }
+		public ulong PeakTotalMemory { get; }
 
 		public bool Result { get; }
 		public int TotalActions { get; }
@@ -134,7 +134,7 @@ namespace UnrealBuildTool
 		public double CacheRate => CacheHitActions + CacheMissActions > 0 ? CacheHitActions / (double)CacheCheckActions : 0.0;
 
 		public TelemetryExecutorEvent(string executor, DateTime startUTC, bool result, int totalActions, int succeededActions, int failedActions, int cacheHitActions, int cacheMissActions, DateTime timestamp,
-			long peakProcessMemory, long peakTotalMemory)
+			ulong peakProcessMemory, ulong peakTotalMemory)
 			: base(timestamp)
 		{
 			Executor = executor;
@@ -179,7 +179,7 @@ namespace UnrealBuildTool
 			int retriedLocalActions, int retriedDisabledActions,
 			int totalCoordinators, int succeededCoordinators, int failedCoordinators,
 			TimeSpan durationWaitingForRemote, string connectionMode,
-			DateTime timestamp, long peakProcessMemory, long peakTotalMemory)
+			DateTime timestamp, ulong peakProcessMemory, ulong peakTotalMemory)
 			: base(executor, startUTC, result, totalActions, succeededActions, failedActions, cacheHitActions, cacheMissActions, timestamp, peakProcessMemory, peakTotalMemory)
 		{
 			LocalActions = localActions;
