As part of creating Datasmith Exporter plugins compiled with UBT, we discovered that some SDKs have built in resource files that cannot be skipped and they include icons that clash with the builtin Default.rc2 definition.
This modifications provide a straightforward solution that enables us to utilize UBT to build external plugins without this issue.

diff --git a/Engine/Source/Programs/UnrealBuildTool/Platform/Windows/UEBuildWindows.cs b/Engine/Source/Programs/UnrealBuildTool/Platform/Windows/UEBuildWindows.cs
--- a/Engine/Source/Programs/UnrealBuildTool/Platform/Windows/UEBuildWindows.cs
+++ b/Engine/Source/Programs/UnrealBuildTool/Platform/Windows/UEBuildWindows.cs
@@ -607,6 +607,12 @@ namespace UnrealBuildTool
 		/// </summary>
 		public string? ApplicationIcon;
 
+		/// <summary>
+		/// Whether to define or not the resource icon to be used in Default.rc2.
+		/// </summary>
+		[CommandLine("-NoDefineIconResource", Value ="false")]
+		public bool bDefineIconResource = true;
+
 		/// <summary>
 		/// Enables strict standard conformance mode (/permissive-).
 		/// </summary>
@@ -1058,6 +1064,8 @@ namespace UnrealBuildTool
 		public string? ManifestFile => Inner.ManifestFile;
 		
 		public string? ApplicationIcon => Inner.ApplicationIcon;
+		
+		public bool bDefineIconResource => Inner.bDefineIconResource;
 
 		public bool bNeedsLegacyStdioDefinitionsLib => Inner.bNeedsLegacyStdioDefinitionsLib;
 
diff --git a/Engine/Source/Programs/UnrealBuildTool/Platform/Windows/VCToolChain.cs b/Engine/Source/Programs/UnrealBuildTool/Platform/Windows/VCToolChain.cs
--- a/Engine/Source/Programs/UnrealBuildTool/Platform/Windows/VCToolChain.cs
+++ b/Engine/Source/Programs/UnrealBuildTool/Platform/Windows/VCToolChain.cs
@@ -2881,12 +2881,15 @@ namespace UnrealBuildTool
 					}
 				}
 
-				// Figure the icon to use. We can only use a custom icon when compiling to a project-specific intermediate directory (and not for the shared editor executable, for example).
-				FileItem IconFile = FileItem.GetItemByFileReference(GetApplicationIcon(!CompileEnvironment.bUseSharedBuildEnvironment ? Target.ProjectFile : null));
-				CompileAction.PrerequisiteItems.Add(IconFile);
+				if (Target.WindowsPlatform.bDefineIconResource)
+				{
+					// Figure the icon to use. We can only use a custom icon when compiling to a project-specific intermediate directory (and not for the shared editor executable, for example).
+					FileItem IconFile = FileItem.GetItemByFileReference(GetApplicationIcon(!CompileEnvironment.bUseSharedBuildEnvironment ? Target.ProjectFile : null));
+					CompileAction.PrerequisiteItems.Add(IconFile);
 
-				// Setup the compile environment, setting the icon to use via a macro. This is used in Default.rc2.
-				AddDefinition(Arguments, $"BUILD_ICON_FILE_NAME=\"\\\"{NormalizeCommandLinePath(IconFile, RootPaths).Replace("/", "\\\\")}\\\"\"");
+					// Setup the compile environment, setting the icon to use via a macro. This is used in Default.rc2.
+					AddDefinition(Arguments, $"BUILD_ICON_FILE_NAME=\"\\\"{NormalizeCommandLinePath(IconFile, RootPaths).Replace("/", "\\\\")}\\\"\"");
+				}
 
 				// Apply the target settings for the resources
 				if (!CompileEnvironment.bUseSharedBuildEnvironment)
