Fix issue CA1052 (Static holder types should be Static or NotInheritable) in UBT.

#Synced-CL 49065719

diff --git a/Engine/Source/Programs/UnrealBuildTool/Configuration/Ini/ConfigValueParser.cs b/Engine/Source/Programs/UnrealBuildTool/Configuration/Ini/ConfigValueParser.cs
--- a/Engine/Source/Programs/UnrealBuildTool/Configuration/Ini/ConfigValueParser.cs
+++ b/Engine/Source/Programs/UnrealBuildTool/Configuration/Ini/ConfigValueParser.cs
@@ -14,7 +14,7 @@ namespace UnrealBuildTool
 	/// <summary>
 	/// Functionality for parsing values from a config value
 	/// </summary>
-	public class ConfigValueParser
+	public static class ConfigValueParser
 	{
 		/// <summary>
 		/// Attempts to parse the given line as the given type. Struct member variables are expected to match
diff --git a/Engine/Source/Programs/UnrealBuildTool/Configuration/Rules/RulesCompiler.cs b/Engine/Source/Programs/UnrealBuildTool/Configuration/Rules/RulesCompiler.cs
--- a/Engine/Source/Programs/UnrealBuildTool/Configuration/Rules/RulesCompiler.cs
+++ b/Engine/Source/Programs/UnrealBuildTool/Configuration/Rules/RulesCompiler.cs
@@ -15,7 +15,7 @@ namespace UnrealBuildTool
 	/// <summary>
 	/// Class which compiles (and caches) rules assemblies for different folders.
 	/// </summary>
-	public class RulesCompiler
+	public static class RulesCompiler
 	{
 		/// <summary>
 		/// 
diff --git a/Engine/Source/Programs/UnrealBuildTool/Platform/IOS/IOSCodeSigning.cs b/Engine/Source/Programs/UnrealBuildTool/Platform/IOS/IOSCodeSigning.cs
--- a/Engine/Source/Programs/UnrealBuildTool/Platform/IOS/IOSCodeSigning.cs
+++ b/Engine/Source/Programs/UnrealBuildTool/Platform/IOS/IOSCodeSigning.cs
@@ -18,25 +18,8 @@ namespace UnrealBuildTool
 	/// <summary>
 	/// Utilities related to code sign for Apple Platforms
 	/// </summary>
-	public class AppleCodeSign
+	public static class AppleCodeSign
 	{
-		/// <summary>
-		/// A logger
-		/// </summary>
-		static ILogger? _logger;
-
-		/// <summary>
-		/// Constructor
-		/// </summary>
-		public AppleCodeSign()
-		{
-			ILoggerFactory loggerFactory = LoggerFactory.Create(builder =>
-			{
-				builder.AddConsole();
-			});
-			_logger = loggerFactory.CreateLogger("AppleCodeSign");
-		}
-
 		/// <summary>
 		/// The shared provision library directory (on PC)
 		/// </summary>
@@ -167,7 +150,7 @@ namespace UnrealBuildTool
 
 				bool bCertTimeIsValid = (EffectiveDate < Now) && (ExpirationDate > Now);
 				string CertLine = String.Format("CERTIFICATE-Name:{0},Validity:{1},StartDate:{2},EndDate:{3}", GetFriendlyNameFromCert(TestCert), bCertTimeIsValid ? "VALID" : "EXPIRED", EffectiveDate.ToString("o"), ExpirationDate.ToString("o"));
-				_logger?.LogInformation(CertLine);
+				Log.Logger.LogInformation(CertLine);
 				Certs.Add(CertLine);
 			}
 			return Certs;
@@ -184,7 +167,7 @@ namespace UnrealBuildTool
 
 			if (!Directory.Exists(ProvisionDirectory))
 			{
-				_logger?.LogError("Could not find provision directory '{ProvisionDirectory}'.", ProvisionDirectory);
+				Log.Logger.LogError("Could not find provision directory '{ProvisionDirectory}'.", ProvisionDirectory);
 				MatchedProvision = null;
 				return Provisions;
 			}
@@ -291,14 +274,14 @@ namespace UnrealBuildTool
 				}
 				if (p != null)
 				{
-					_logger?.LogInformation("PROVISION-File:{Filename},Name:{Name},Validity:{Validity},StartDate:{StartDate},EndDate:{EndDate},Type:{Type}", Path.GetFileName(ProvisionFile), p.ProvisionName, Validity, EffectiveDate.ToString(), ExpirationDate.ToString(), bDistribution ? "DISTRIBUTION" : "DEVELOPMENT");
+					Log.Logger.LogInformation("PROVISION-File:{Filename},Name:{Name},Validity:{Validity},StartDate:{StartDate},EndDate:{EndDate},Type:{Type}", Path.GetFileName(ProvisionFile), p.ProvisionName, Validity, EffectiveDate.ToString(), ExpirationDate.ToString(), bDistribution ? "DISTRIBUTION" : "DEVELOPMENT");
 				}
 				if (ProvisionFile != null)
 				{
 					Provisions.Add(new FileReference(ProvisionFile));
 				}
 			}
-			_logger?.LogInformation(" MATCHED-Provision:{Name},File:{Filename},Cert:{Cert} ", SelectedProvision, SelectedFile, SelectedCert);
+			Log.Logger.LogInformation(" MATCHED-Provision:{Name},File:{Filename},Cert:{Cert} ", SelectedProvision, SelectedFile, SelectedCert);
 			if (!String.IsNullOrEmpty(SelectedProvision))
 			{
 				MatchedProvision = new FileReference(SelectedProvision);
@@ -317,7 +300,7 @@ namespace UnrealBuildTool
 		/// </summary>
 		public static X509Certificate2? FindCertificate(MobileProvision ProvisionToWorkFrom)
 		{
-			_logger?.LogInformation("  Looking for a certificate that matches the application identifier '{ApplicationIdentifier}'", ProvisionToWorkFrom.ApplicationIdentifier);
+			Log.Logger.LogInformation("  Looking for a certificate that matches the application identifier '{ApplicationIdentifier}'", ProvisionToWorkFrom.ApplicationIdentifier);
 
 			X509Certificate2? Result = null;
 
@@ -336,18 +319,18 @@ namespace UnrealBuildTool
 				CertTool.BeginOutputReadLine();
 				CertTool.WaitForExit();
 
-				_logger?.LogInformation("  Running {CertToolPath} {CertToolArguments}\n{CertToolData}", CertTool.StartInfo.FileName, CertTool.StartInfo.Arguments, CertToolData);
+				Log.Logger.LogInformation("  Running {CertToolPath} {CertToolArguments}\n{CertToolData}", CertTool.StartInfo.FileName, CertTool.StartInfo.Arguments, CertToolData);
 
 				if (CertTool.ExitCode == 0)
 				{
-					_logger?.LogInformation("  Provisioning profile contains the following certificate hashes:");
+					Log.Logger.LogInformation("  Provisioning profile contains the following certificate hashes:");
 					foreach (X509Certificate2 SourceCert in ProvisionToWorkFrom.DeveloperCertificates)
 					{
 						X509Certificate2? ValidInTimeCert = null;
 						// see if certificate can be found by serial number
 						string CertHash = SourceCert.GetCertHashString();
 
-						_logger?.LogInformation("    {CertHash}", CertHash);
+						Log.Logger.LogInformation("    {CertHash}", CertHash);
 
 						if (CertToolData.Contains(CertHash))
 						{
@@ -358,7 +341,7 @@ namespace UnrealBuildTool
 						{
 							// Found a cert in the valid time range, quit now!
 							Result = ValidInTimeCert;
-							_logger?.LogInformation("  Matches!");
+							Log.Logger.LogInformation("  Matches!");
 							break;
 						}
 					}
@@ -375,7 +358,7 @@ namespace UnrealBuildTool
 				{
 					X509Certificate2Collection FoundCerts = Store.Certificates.Find(X509FindType.FindBySerialNumber, SourceCert.SerialNumber, false);
 
-					_logger?.LogInformation("  .. Provision entry SN '{SerialNumber}' matched {NumFoundCerts} installed certificate(s)", SourceCert.SerialNumber, FoundCerts.Count);
+					Log.Logger.LogInformation("  .. Provision entry SN '{SerialNumber}' matched {NumFoundCerts} installed certificate(s)", SourceCert.SerialNumber, FoundCerts.Count);
 
 					X509Certificate2? ValidInTimeCert = null;
 					foreach (X509Certificate2 TestCert in FoundCerts)
@@ -386,7 +369,7 @@ namespace UnrealBuildTool
 
 						bool bCertTimeIsValid = (EffectiveDate < Now) && (ExpirationDate > Now);
 
-						_logger?.LogInformation("  .. .. Installed certificate '{CertFriendlyName}' is {Validity} (range '{StartDate}' to '{EndDate}')", GetFriendlyNameFromCert(TestCert), bCertTimeIsValid ? "valid (choosing it)" : "EXPIRED", TestCert.GetEffectiveDateString(), TestCert.GetExpirationDateString());
+						Log.Logger.LogInformation("  .. .. Installed certificate '{CertFriendlyName}' is {Validity} (range '{StartDate}' to '{EndDate}')", GetFriendlyNameFromCert(TestCert), bCertTimeIsValid ? "valid (choosing it)" : "EXPIRED", TestCert.GetEffectiveDateString(), TestCert.GetExpirationDateString());
 						if (bCertTimeIsValid)
 						{
 							ValidInTimeCert = TestCert;
@@ -407,7 +390,7 @@ namespace UnrealBuildTool
 
 			if (Result == null)
 			{
-				_logger?.LogInformation("  .. Failed to find a valid certificate that was in date");
+				Log.Logger.LogInformation("  .. Failed to find a valid certificate that was in date");
 			}
 
 			return Result;
@@ -537,12 +520,12 @@ namespace UnrealBuildTool
 					// print out the provision and cert name
 					Provision = new FileReference(FileName);
 					CertName = GetFriendlyNameFromCert(Cert!);
-					_logger?.LogInformation("CERTIFICATE-{Name},PROVISION-{ProvisionFilePath}", CertName, Provision.FullName);
+					Log.Logger.LogInformation("CERTIFICATE-{Name},PROVISION-{ProvisionFilePath}", CertName, Provision.FullName);
 					return true;
 				}
 			}
 
-			_logger?.LogInformation("No matching Signing Data found!");
+			Log.Logger.LogInformation("No matching Signing Data found!");
 			Provision = null;
 			CertName = "";
 			return false;
@@ -1380,22 +1363,8 @@ namespace UnrealBuildTool
 	/// This class understands how to get the embedded plist in a .mobileprovision file.  It doesn't
 	/// understand the full format and is not capable of writing a new one out or anything similar.
 	/// </summary>
-	public class MobileProvisionParser
+	public static class MobileProvisionParser
 	{
-		static ILogger? _logger;
-
-		/// <summary>
-		/// The mobile provision parser
-		/// </summary>
-		public MobileProvisionParser()
-		{
-			ILoggerFactory loggerFactory = LoggerFactory.Create(builder =>
-			{
-				builder.AddConsole();
-			});
-			_logger = loggerFactory.CreateLogger("MobileProvisionParser");
-		}
-
 		private static int StrStrByteArray(byte[] Haystack, int Offset, string Needle)
 		{
 			byte[] NeedleBytes = Encoding.UTF8.GetBytes(Needle);
@@ -1454,7 +1423,7 @@ namespace UnrealBuildTool
 			}
 
 			// Unable to find the start of the plist data
-			_logger?.LogError("Failed to find embedded plist in .mobileprovision file");
+			Log.Logger.LogError("Failed to find embedded plist in .mobileprovision file");
 			return null;
 		}
 
@@ -1492,7 +1461,7 @@ namespace UnrealBuildTool
 	/// <summary>
 	/// The code signing config class
 	/// </summary>
-	public class CodeSigningConfig
+	public static class CodeSigningConfig
 	{
 		/// <summary>
 		/// The display name in title bars, popups, etc ...
diff --git a/Engine/Source/Programs/UnrealBuildTool/Platform/Linux/LinuxExports.cs b/Engine/Source/Programs/UnrealBuildTool/Platform/Linux/LinuxExports.cs
--- a/Engine/Source/Programs/UnrealBuildTool/Platform/Linux/LinuxExports.cs
+++ b/Engine/Source/Programs/UnrealBuildTool/Platform/Linux/LinuxExports.cs
@@ -8,7 +8,7 @@ namespace UnrealBuildTool
 	/// <summary>
 	/// Public Linux functions exposed to UAT
 	/// </summary>
-	public class LinuxExports
+	public static class LinuxExports
 	{
 		/// <summary>
 		/// 
diff --git a/Engine/Source/Programs/UnrealBuildTool/System/DataDrivenPlatformInfo.cs b/Engine/Source/Programs/UnrealBuildTool/System/DataDrivenPlatformInfo.cs
--- a/Engine/Source/Programs/UnrealBuildTool/System/DataDrivenPlatformInfo.cs
+++ b/Engine/Source/Programs/UnrealBuildTool/System/DataDrivenPlatformInfo.cs
@@ -12,7 +12,7 @@ namespace UnrealBuildTool
 	/// <summary>
 	///  Class to manage looking up data driven platform information (loaded from .ini files instead of in code)
 	/// </summary>
-	public class DataDrivenPlatformInfo
+	public static class DataDrivenPlatformInfo
 	{
 		/// <summary>
 		/// All data driven information about a platform
diff --git a/Engine/Source/Programs/UnrealBuildTool/System/DeployExports.cs b/Engine/Source/Programs/UnrealBuildTool/System/DeployExports.cs
--- a/Engine/Source/Programs/UnrealBuildTool/System/DeployExports.cs
+++ b/Engine/Source/Programs/UnrealBuildTool/System/DeployExports.cs
@@ -8,7 +8,7 @@ namespace UnrealBuildTool
 	/// <summary>
 	/// General deploy support functions, exported for UAT
 	/// </summary>
-	public class DeployExports
+	public static class DeployExports
 	{
 		/// <summary>
 		/// The max length of the deploy folder name.
diff --git a/Engine/Source/Programs/UnrealBuildTool/System/DynamicCompilation.cs b/Engine/Source/Programs/UnrealBuildTool/System/DynamicCompilation.cs
--- a/Engine/Source/Programs/UnrealBuildTool/System/DynamicCompilation.cs
+++ b/Engine/Source/Programs/UnrealBuildTool/System/DynamicCompilation.cs
@@ -25,7 +25,7 @@ namespace UnrealBuildTool
 	/// <summary>
 	/// Methods for dynamically compiling C# source files
 	/// </summary>
-	public class DynamicCompilation
+	public static class DynamicCompilation
 	{
 		/// <summary>
 		/// Checks to see if the assembly needs compilation
diff --git a/Engine/Source/Programs/UnrealBuildTool/System/InstalledPlatformInfo.cs b/Engine/Source/Programs/UnrealBuildTool/System/InstalledPlatformInfo.cs
--- a/Engine/Source/Programs/UnrealBuildTool/System/InstalledPlatformInfo.cs
+++ b/Engine/Source/Programs/UnrealBuildTool/System/InstalledPlatformInfo.cs
@@ -53,7 +53,7 @@ namespace UnrealBuildTool
 	/// <summary>
 	/// Contains methods to allow querying the available installed platforms
 	/// </summary>
-	public class InstalledPlatformInfo
+	public static class InstalledPlatformInfo
 	{
 		/// <summary>
 		/// Information about a single installed platform configuration
diff --git a/Engine/Source/Programs/UnrealBuildTool/System/PlatformExports.cs b/Engine/Source/Programs/UnrealBuildTool/System/PlatformExports.cs
--- a/Engine/Source/Programs/UnrealBuildTool/System/PlatformExports.cs
+++ b/Engine/Source/Programs/UnrealBuildTool/System/PlatformExports.cs
@@ -11,7 +11,7 @@ namespace UnrealBuildTool
 	/// <summary>
 	/// General platform support functions, exported for UAT
 	/// </summary>
-	public class PlatformExports
+	public static class PlatformExports
 	{
 		/// <summary>
 		/// Checks whether the given platform is available
