Ensure UPL file deployments end up in Modern Xcode archived builds. To do so, instead of having UPL deploy directly to the App bundle, instead deploy to a staged directory which gets consumed in the script phase (Deploy Exe and Staged Data) of the Xcode project.

#Synced-CL 48848898

diff --git a/Engine/Source/Programs/UnrealBuildTool/Platform/IOS/UEDeployIOS.cs b/Engine/Source/Programs/UnrealBuildTool/Platform/IOS/UEDeployIOS.cs
--- a/Engine/Source/Programs/UnrealBuildTool/Platform/IOS/UEDeployIOS.cs
+++ b/Engine/Source/Programs/UnrealBuildTool/Platform/IOS/UEDeployIOS.cs
@@ -552,18 +552,25 @@ namespace UnrealBuildTool
 
 			List<(string UEArch, string NativeArch)> ProjectArches = new() { ("None", "None") };
 
-			string BundlePath;
+			string UPLBuildDir = Path.Combine(ProjectDirectory, "Intermediate", "IOS", "UPLBuild");
 
-			// get the receipt
-			if (bIsUnrealGame)
+			if(AppleExports.CreatingAppOnWindows(ProjectFile))
 			{
-				//               ReceiptFilename = TargetReceipt.GetDefaultPath(Unreal.EngineDirectory, "UnrealGame", UnrealTargetPlatform.IOS, Config, "");
-				BundlePath = Path.Combine(Unreal.EngineDirectory.ToString(), "Intermediate", "IOS-Deploy", "UnrealGame", Config.ToString(), "Payload", "UnrealGame.app");
+				// get the receipt
+				if (bIsUnrealGame)
+				{
+					//               ReceiptFilename = TargetReceipt.GetDefaultPath(Unreal.EngineDirectory, "UnrealGame", UnrealTargetPlatform.IOS, Config, "");
+					UPLBuildDir = Path.Combine(Unreal.EngineDirectory.ToString(), "Intermediate", "IOS-Deploy", "UnrealGame", Config.ToString(), "Payload", "UnrealGame.app");
+				}
+				else
+				{
+					//                ReceiptFilename = TargetReceipt.GetDefaultPath(new DirectoryReference(ProjectDirectory), GameName, UnrealTargetPlatform.IOS, Config, "");
+					UPLBuildDir = AppDirectory;//Path.Combine(ProjectDirectory, "Binaries", "IOS", "Payload", ProjectName + ".app");
+				}
 			}
 			else
 			{
-				//                ReceiptFilename = TargetReceipt.GetDefaultPath(new DirectoryReference(ProjectDirectory), GameName, UnrealTargetPlatform.IOS, Config, "");
-				BundlePath = AppDirectory;//Path.Combine(ProjectDirectory, "Binaries", "IOS", "Payload", ProjectName + ".app");
+				Directory.CreateDirectory(UPLBuildDir);
 			}
 
 			string RelativeEnginePath = Unreal.EngineDirectory.MakeRelativeTo(DirectoryReference.GetCurrentDirectory());
@@ -571,7 +578,7 @@ namespace UnrealBuildTool
 			UnrealPluginLanguage UPL = new UnrealPluginLanguage(ProjectFile, UPLScripts, ProjectArches, null, UnrealTargetPlatform.IOS, Logger);
 
 			// Passing in true for distribution is not ideal here but given the way that ios packaging happens and this call chain it seems unavoidable for now, maybe there is a way to correctly pass it in that I can't find?
-			UPL.Init(ProjectArches.Select(x => x.NativeArch), true, RelativeEnginePath, BundlePath, ProjectDirectory, Config.ToString(), false);
+			UPL.Init(ProjectArches.Select(x => x.NativeArch), true, RelativeEnginePath, UPLBuildDir, ProjectDirectory, Config.ToString(), false);
 
 			return GenerateIOSPList(ProjectFile, Config, ProjectDirectory, bIsUnrealGame, GameName, bIsClient, ProjectName, InEngineDir, AppDirectory, UPL, BundleID, bBuildAsFramework, Logger);
 		}
diff --git a/Engine/Source/Programs/UnrealBuildTool/ProjectFiles/Xcode/XcodeProject.cs b/Engine/Source/Programs/UnrealBuildTool/ProjectFiles/Xcode/XcodeProject.cs
--- a/Engine/Source/Programs/UnrealBuildTool/ProjectFiles/Xcode/XcodeProject.cs
+++ b/Engine/Source/Programs/UnrealBuildTool/ProjectFiles/Xcode/XcodeProject.cs
@@ -1645,7 +1645,7 @@ namespace UnrealBuildTool.XcodeProjectXcconfig
 				// make a string like --exclude=/Info.plist --exclude=/Manifest_* ...
 				string ExcludeString = String.Join(" ", Exclusions.Select(x => (x[0] == '+' ? "--include" : "--exclude") + $"=\\\"{x.Substring(1)}\\\""));
 
-				// copy uecommandline.txt for IOS type platforms
+				// copy uecommandline.txt and UPL build output for IOS type platforms
 				if (Platform != UnrealTargetPlatform.Mac)
 				{
 					CopyScript.AddRange(new string[]
@@ -1653,6 +1653,10 @@ namespace UnrealBuildTool.XcodeProjectXcconfig
 						"if [[ -e \\\"${STAGED_DIR}/uecommandline.txt\\\" ]]; then",
 						"  echo \\\"Syncing ${STAGED_DIR}/uecommandline.txt to ${CONFIGURATION_BUILD_DIR}/${CONTENTS_FOLDER_PATH}\\\"",
 						"  cp \\\"${STAGED_DIR}/uecommandline.txt\\\" \\\"${CONFIGURATION_BUILD_DIR}/${CONTENTS_FOLDER_PATH}\\\"",
+						"fi",
+						"UPL_BUILD_DIR=${UE_PROJECT_DIR}/Intermediate/IOS/UPLBuild/",
+						"if [[ -e \\\"${UPL_BUILD_DIR}\\\" ]]; then",
+						"cp -r \\\"${UPL_BUILD_DIR}\\\" \\\"${CONFIGURATION_BUILD_DIR}/${CONTENTS_FOLDER_PATH}\\\"",
 						"fi"
 					});
 				}
