[Android] Extracts several blocks of code out of UEDeployAndroid into re-usable functions for AndroidToolChain
These are then used by both UEDeployAndroid and AndroidExports for exposure to UAT
The primary reason for this is to enable automated generation of code only preflight builds
The following are now possible independently of packaging an APK
- Obtaining build tool versions
- Stripping symbols
- Signing APKs

#Synced-CL 48971518

diff --git a/Engine/Source/Programs/UnrealBuildTool/Platform/Android/AndroidExports.cs b/Engine/Source/Programs/UnrealBuildTool/Platform/Android/AndroidExports.cs
--- a/Engine/Source/Programs/UnrealBuildTool/Platform/Android/AndroidExports.cs
+++ b/Engine/Source/Programs/UnrealBuildTool/Platform/Android/AndroidExports.cs
@@ -92,8 +92,19 @@ namespace UnrealBuildTool
 	/// </summary>
 	public static class AndroidExports
 	{
+		private static AndroidToolChain? CachedToolChain;
+		private static AndroidToolChain GetToolChain(FileReference ProjectFile, ILogger Logger)
+		{
+			if (CachedToolChain == null || !CachedToolChain.Equals(ProjectFile, Logger))
+			{
+				CachedToolChain = new AndroidToolChain(ProjectFile, Logger);
+			}
+
+			return CachedToolChain;
+		}
+
 		/// <summary>
-		/// 
+		///
 		/// </summary>
 		/// <param name="ProjectFile"></param>
 		/// <param name="InForcePackageData"></param>
@@ -104,7 +115,7 @@ namespace UnrealBuildTool
 		}
 
 		/// <summary>
-		/// 
+		///
 		/// </summary>
 		/// <returns></returns>
 		public static bool ShouldMakeSeparateApks()
@@ -113,7 +124,7 @@ namespace UnrealBuildTool
 		}
 
 		/// <summary>
-		/// 
+		///
 		/// </summary>
 		/// <param name="NDKArch"></param>
 		/// <returns></returns>
@@ -123,19 +134,80 @@ namespace UnrealBuildTool
 		}
 
 		/// <summary>
-		/// 
+		///
 		/// </summary>
 		/// <param name="SourceFile"></param>
 		/// <param name="TargetFile"></param>
 		/// <param name="Logger">Logger for output</param>
-		public static void StripSymbols(FileReference SourceFile, FileReference TargetFile, ILogger Logger)
+		/// <param name="bStripAll">Whether to strip unneeded symbols or just debug symbols</param>
+		public static void StripSymbols(FileReference SourceFile, FileReference TargetFile, ILogger Logger, bool bStripAll = false)
 		{
-			AndroidToolChain ToolChain = new AndroidToolChain(null, Logger);
-			ToolChain.StripSymbols(SourceFile, TargetFile, Logger);
+			AndroidToolChain.StripSymbols(SourceFile, TargetFile, Logger, bStripAll);
 		}
 
 		/// <summary>
-		/// 
+		/// Replaces a shared object file within an apk, copying the resulting apk to a requested destination
+		/// </summary>
+		/// <param name="SourceApk">Base apk</param>
+		/// <param name="TargetApk">Output location</param>
+		/// <param name="SOFile">The so file to replace</param>
+		/// <param name="RelativeSOFile">Path of the new SO file, relative to the apk</param>
+		/// <param name="Logger">Logger</param>
+		public static void ReplaceSO(FileReference SourceApk, FileReference TargetApk, FileReference SOFile, string RelativeSOFile, ILogger Logger)
+		{
+			UEDeployAndroid.CopyAPKAndReplaceSO(SourceApk, TargetApk, SOFile, RelativeSOFile, Logger);
+		}
+
+		/// <summary>
+		/// Signs the source APK for debug/development usage
+		/// </summary>
+		/// <param name="ProjectFile">UE Project file</param>
+		/// <param name="SourceApkFile">APK to sign</param>
+		/// <param name="TargetApkFile">Location the signed APK is output to</param>
+		/// <param name="Logger">Logger for output</param>
+		/// <param name="bZipAlign">Whether or not to align the zip before signing</param>
+		public static void SignDebugApk(FileReference ProjectFile, FileReference SourceApkFile, FileReference TargetApkFile, ILogger Logger, bool bZipAlign = false)
+		{
+			if (ProjectFile == null || !FileReference.Exists(ProjectFile))
+			{
+				throw new BuildException("Failed to sign apk. Project file reference is null or it does not exist");
+			}
+			if (SourceApkFile == null || !FileReference.Exists(SourceApkFile))
+			{
+				throw new BuildException("Failed to sign apk. Source file reference is null or it does not exist");
+			}
+			if (TargetApkFile == null)
+			{
+				throw new BuildException("Failed to sign apk. Target file reference is null");
+			}
+
+			GetToolChain(ProjectFile, Logger).SignApk(SourceApkFile, TargetApkFile, bZipAlign);
+		}
+
+		/// <summary>
+		/// Obtains the version of Android build tools
+		/// </summary>
+		/// <param name="ProjectFile">UE project file</param>
+		/// <param name="Logger">Logger for output</param>
+		/// <returns>Version of Android build tools</returns>
+		public static string GetBuildToolsVersion(FileReference ProjectFile, ILogger Logger)
+		{
+			return GetToolChain(ProjectFile, Logger).GetBuildToolsVersion();
+		}
+
+		/// <summary>
+		/// Obtains the path of Android build tools
+		/// </summary>
+		/// <param name="ProjectFile">UE project file</param>
+		/// <param name="Logger">Logger for output</param>
+		/// <returns>Path to Android build tools</returns>
+		public static string GetBuildToolsPath(FileReference ProjectFile, ILogger Logger)
+		{
+			return GetToolChain(ProjectFile, Logger).GetBuildToolsPath();
+		}
+
+		/// <summary>
+		///
 		/// </summary>
 		/// <param name="Target"></param>
 		/// <param name="Logger"></param>
diff --git a/Engine/Source/Programs/UnrealBuildTool/Platform/Android/AndroidToolChain.cs b/Engine/Source/Programs/UnrealBuildTool/Platform/Android/AndroidToolChain.cs
--- a/Engine/Source/Programs/UnrealBuildTool/Platform/Android/AndroidToolChain.cs
+++ b/Engine/Source/Programs/UnrealBuildTool/Platform/Android/AndroidToolChain.cs
@@ -17,6 +17,9 @@ namespace UnrealBuildTool
 	{
 		// Minimum NDK API level to allow
 		public const int MinimumNDKAPILevel = 26;
+		public const int MinimumBestVersion = (28 << 24) | (0 << 16) | (3 << 8);
+		public const string MinimumBestVersionString = "28.0.3";
+		public const string FallbackBestVersionString = "35.0.1";
 
 		// this is architectures with the dash, which we match in filenames that have inlined arch name
 		public static readonly Dictionary<UnrealArch, string> AllCpuSuffixes = new()
@@ -251,6 +254,11 @@ namespace UnrealBuildTool
 			SysrootPath = DirectoryReference.FromString(Path.Combine(NDKPath!, @"toolchains/llvm", ArchitecturePath!, "sysroot").Replace("\\", "/"))!;
 		}
 
+		public bool Equals(FileReference InProjectFile, ILogger InLogger)
+		{
+			return ProjectFile == InProjectFile && Logger == InLogger;
+		}
+
 		IEnumerable<string> AdditionalPaths(CppRootPaths Roots)
 		{
 			return new[]
@@ -1430,7 +1438,7 @@ namespace UnrealBuildTool
 			}
 		}
 
-		public virtual string GetStripPath(FileReference SourceFile)
+		public static string GetStripPath(FileReference SourceFile)
 		{
 			string StripExe;
 			if (SourceFile.FullName.Contains("-arm64"))
@@ -1453,22 +1461,190 @@ namespace UnrealBuildTool
 			return StripExe;
 		}
 
-		public void StripSymbols(FileReference SourceFile, FileReference TargetFile, ILogger Logger)
+		public string GetBuildToolsVersion()
 		{
-			if (SourceFile != TargetFile)
+			string HomePath = Environment.ExpandEnvironmentVariables("%ANDROID_HOME%");
+			if (string.IsNullOrEmpty(HomePath))
 			{
-				// Strip command only works in place so we need to copy original if target is different
-				File.Copy(SourceFile.FullName, TargetFile.FullName, true);
+				throw new BuildException("Failed to find %ANDROID_HOME%. Ensure environment is configured properly");
 			}
 
+			string? BestVersionString = null;
+			uint BestVersion = 0;
+
+			ConfigHierarchy Ini = ConfigCache.ReadHierarchy(ConfigHierarchyType.Engine, DirectoryReference.FromFile(ProjectFile), UnrealTargetPlatform.Android);
+			Ini.GetString("/Script/AndroidRuntimeSettings.AndroidRuntimeSettings", "BuildToolsOverride", out BestVersionString);
+
+			if (BestVersionString == null || String.IsNullOrEmpty(BestVersionString) || BestVersionString == "latest")
+			{
+				// get a list of the directories in build-tools.. may be more than one set installed (or none which is bad)
+				string[] Subdirs = Directory.GetDirectories(Path.Combine(HomePath, "build-tools"));
+				if (Subdirs.Length == 0)
+				{
+					throw new BuildException("Failed to find %ANDROID_HOME%/build-tools subdirectory. Run SDK manager and install build-tools.");
+				}
+
+				// valid directories will have a source.properties with the Pkg.Revision (there is no guarantee we can use the directory name as revision)
+				foreach (string CandidateDir in Subdirs)
+				{
+					string AaptFilename = Path.Combine(CandidateDir, OperatingSystem.IsWindows() ? "aapt.exe" : "aapt");
+					string RevisionString = "";
+					uint RevisionValue = 0;
+
+					if (File.Exists(AaptFilename))
+					{
+						string SourcePropFilename = Path.Combine(CandidateDir, "source.properties");
+						if (File.Exists(SourcePropFilename))
+						{
+							string[] PropertyContents = File.ReadAllLines(SourcePropFilename);
+							foreach (string PropertyLine in PropertyContents)
+							{
+								if (PropertyLine.StartsWith("Pkg.Revision="))
+								{
+									RevisionString = PropertyLine.Substring(13);
+									RevisionValue = GetRevisionValue(RevisionString);
+									break;
+								}
+							}
+						}
+					}
+
+					// remember it if newer version or haven't found one yet
+					if (RevisionValue > BestVersion || BestVersionString == null)
+					{
+						BestVersion = RevisionValue;
+						BestVersionString = RevisionString;
+					}
+				}
+			}
+
+			if (BestVersionString == null)
+			{
+				BestVersionString = FallbackBestVersionString;
+				Logger.LogWarning("Failed to find %ANDROID_HOME%/build-tools subdirectory. Will attempt to use {BestVersionString}.", BestVersionString);
+			}
+
+			BestVersion = GetRevisionValue(BestVersionString);
+
+			// With Gradle enabled, use at least the minimum version (will be installed by Gradle if missing)
+			if (BestVersion < MinimumBestVersion)
+			{
+				BestVersionString = MinimumBestVersionString;
+			}
+
+			return BestVersionString;
+		}
+
+		public string GetBuildToolsPath()
+		{
+			string HomePath = Environment.ExpandEnvironmentVariables("%ANDROID_HOME%");
+			return Path.Combine(HomePath, "build-tools", GetBuildToolsVersion());
+		}
+
+		public static void StripSymbols(FileReference SourceFile, FileReference TargetFile, UnrealArch UnrealArch, ILogger Logger, bool bStripAll = false)
+		{
+			string StripCommand = bStripAll ? "--strip-unneeded" : "--strip-debug";
+
 			ProcessStartInfo StartInfo = new ProcessStartInfo();
-			StartInfo.FileName = GetStripPath(SourceFile).Trim('"');
-			StartInfo.Arguments = " --strip-debug \"" + TargetFile.FullName + "\"";
+			StartInfo.FileName = GetStripExecutablePath(UnrealArch).Trim('"');
+			StartInfo.Arguments = $"{StripCommand} -o \"{TargetFile.FullName}\" \"{SourceFile.FullName}\"";
 			StartInfo.UseShellExecute = false;
 			StartInfo.CreateNoWindow = true;
 			Utils.RunLocalProcessAndLogOutput(StartInfo, Logger);
 		}
 
+		public static void StripSymbols(FileReference SourceFile, FileReference TargetFile, ILogger Logger, bool bStripAll = false)
+		{
+			UnrealArch Architecture = SourceFile.FullName.Contains("arm64", StringComparison.OrdinalIgnoreCase) ? UnrealArch.Arm64 : UnrealArch.X64;
+			StripSymbols(SourceFile, TargetFile, Architecture, Logger, bStripAll);
+		}
+
+		public void SignApk(FileReference SourceApk, FileReference TargetApk, bool bZipAlign = false)
+		{
+			string KeyStore = Path.Combine(Environment.GetFolderPath(Environment.SpecialFolder.UserProfile), ".android", "debug.keystore");
+			string KeyStorePassword = "android";
+			string KeyAlias = "androiddebugkey";
+			string KeyPassword = "android";
+
+			SignApk(SourceApk, TargetApk, KeyStore, KeyStorePassword, KeyAlias, KeyPassword, bZipAlign);
+		}
+
+		public void SignApk(FileReference SourceApk, FileReference TargetApk, string KeyStore, string KeyStorePassword, string KeyAlias, string KeyPassword, bool bZipAlign = false)
+		{
+			Process Proc;
+			ProcessStartInfo StartInfo;
+			string BuildTools = GetBuildToolsPath();
+
+			if (bZipAlign)
+			{
+				string ZipAlign = Path.Combine(BuildTools, "zipalign" + (OperatingSystem.IsWindows() ? ".exe" : ""));
+				StartInfo = new ProcessStartInfo();
+				StartInfo.FileName = ZipAlign;
+				StartInfo.UseShellExecute = false;
+				StartInfo.WindowStyle = ProcessWindowStyle.Minimized;
+
+				StartInfo.Arguments = $"-P 16 -f -v 4 {SourceApk.FullName} {TargetApk.FullName}";
+
+				Proc = new Process();
+				Proc.StartInfo = StartInfo;
+				Proc.Start();
+				Proc.WaitForExit();
+			}
+
+			string ApkSigner = Path.Combine(BuildTools, "apksigner" + (OperatingSystem.IsWindows() ? ".bat" : ""));
+			StartInfo = new ProcessStartInfo();
+			StartInfo.UseShellExecute = false;
+			StartInfo.WindowStyle = ProcessWindowStyle.Minimized;
+
+			if (OperatingSystem.IsWindows())
+			{
+				StartInfo.FileName = "cmd.exe";
+				StartInfo.ArgumentList.Add("/c");
+				StartInfo.ArgumentList.Add(ApkSigner);
+				StartInfo.ArgumentList.Add("sign");
+				StartInfo.ArgumentList.Add("--ks");
+				StartInfo.ArgumentList.Add(KeyStore);
+				StartInfo.ArgumentList.Add("--ks-pass");
+				StartInfo.ArgumentList.Add("pass:" + KeyStorePassword);
+				StartInfo.ArgumentList.Add("--ks-key-alias");
+				StartInfo.ArgumentList.Add(KeyAlias);
+				StartInfo.ArgumentList.Add("--key-pass");
+				StartInfo.ArgumentList.Add("pass:" + KeyPassword);
+				StartInfo.ArgumentList.Add("--out");
+				StartInfo.ArgumentList.Add(TargetApk.FullName);
+
+				if (bZipAlign)
+				{
+					StartInfo.ArgumentList.Add(TargetApk.FullName);
+				}
+				else
+				{
+					StartInfo.ArgumentList.Add(SourceApk.FullName);
+				}
+			}
+			else
+			{
+				StartInfo.FileName = "/bin/sh";
+				StartInfo.ArgumentList.Add("-c");
+				StartInfo.ArgumentList.Add("\"" + ApkSigner + "\" sign --ks \"" + KeyStore + "\" --ks-pass pass:" + KeyStorePassword + " --ks-key-alias " + KeyAlias + " --key-pass pass:" + KeyPassword + " --out \"" + TargetApk.FullName + "2\" \"" + SourceApk.FullName + "\"");
+			}
+
+			Proc = new Process();
+			Proc.StartInfo = StartInfo;
+			Proc.Start();
+			Proc.WaitForExit();
+
+			if (Proc.ExitCode != 0)
+			{
+				string Args = "";
+				foreach (string Arg in StartInfo.ArgumentList)
+				{
+					Args += Arg + " ";
+				}
+				throw new BuildException("{0} failed with args {1}", StartInfo.FileName, Args);
+			}
+		}
+
 		public ClangSanitizer BuildWithSanitizer()
 		{
 			if (Options.HasFlag(ClangToolChainOptions.EnableAddressSanitizer))
@@ -1494,5 +1670,30 @@ namespace UnrealBuildTool
 
 			return ClangSanitizer.None;
 		}
+
+		public static uint GetRevisionValue(string VersionString)
+		{
+			if (VersionString == null)
+			{
+				return 0;
+			}
+
+			// read up to 4 sections (ie. 20.0.3.5), first section most significant
+			// each section assumed to be 0 to 255 range
+			uint Value = 0;
+			try
+			{
+				string[] Sections = VersionString.Split(".".ToCharArray());
+				Value |= (Sections.Length > 0) ? (UInt32.Parse(Sections[0]) << 24) : 0;
+				Value |= (Sections.Length > 1) ? (UInt32.Parse(Sections[1]) << 16) : 0;
+				Value |= (Sections.Length > 2) ? (UInt32.Parse(Sections[2]) << 8) : 0;
+				Value |= (Sections.Length > 3) ? UInt32.Parse(Sections[3]) : 0;
+			}
+			catch (Exception)
+			{
+				// ignore poorly formed version
+			}
+			return Value;
+		}
 	};
 }
diff --git a/Engine/Source/Programs/UnrealBuildTool/Platform/Android/UEDeployAndroid.cs b/Engine/Source/Programs/UnrealBuildTool/Platform/Android/UEDeployAndroid.cs
--- a/Engine/Source/Programs/UnrealBuildTool/Platform/Android/UEDeployAndroid.cs
+++ b/Engine/Source/Programs/UnrealBuildTool/Platform/Android/UEDeployAndroid.cs
@@ -290,32 +290,7 @@ namespace UnrealBuildTool
 		private string? CachedBuildToolsVersion = null;
 		private string? LastAndroidHomePath = null;
 
-		private uint GetRevisionValue(string VersionString)
-		{
-			if (VersionString == null)
-			{
-				return 0;
-			}
-
-			// read up to 4 sections (ie. 20.0.3.5), first section most significant
-			// each section assumed to be 0 to 255 range
-			uint Value = 0;
-			try
-			{
-				string[] Sections = VersionString.Split(".".ToCharArray());
-				Value |= (Sections.Length > 0) ? (UInt32.Parse(Sections[0]) << 24) : 0;
-				Value |= (Sections.Length > 1) ? (UInt32.Parse(Sections[1]) << 16) : 0;
-				Value |= (Sections.Length > 2) ? (UInt32.Parse(Sections[2]) << 8) : 0;
-				Value |= (Sections.Length > 3) ? UInt32.Parse(Sections[3]) : 0;
-			}
-			catch (Exception)
-			{
-				// ignore poorly formed version
-			}
-			return Value;
-		}
-
-		private string GetBuildToolsVersion()
+		private string GetBuildToolsVersion(AndroidToolChain ToolChain)
 		{
 			// return cached path if ANDROID_HOME has not changed
 			string HomePath = Environment.ExpandEnvironmentVariables("%ANDROID_HOME%");
@@ -324,75 +299,7 @@ namespace UnrealBuildTool
 				return CachedBuildToolsVersion;
 			}
 
-			uint BestVersion = 0;
-
-			ConfigHierarchy Ini = GetConfigCacheIni(ConfigHierarchyType.Engine);
-			Ini.GetString("/Script/AndroidRuntimeSettings.AndroidRuntimeSettings", "BuildToolsOverride", out string? BestVersionString);
-
-			if (BestVersionString == null || String.IsNullOrEmpty(BestVersionString) || BestVersionString == "latest")
-			{
-				// get a list of the directories in build-tools.. may be more than one set installed (or none which is bad)
-				string[] Subdirs = Directory.GetDirectories(Path.Combine(HomePath, "build-tools"));
-				if (Subdirs.Length == 0)
-				{
-					throw new BuildException("Failed to find %ANDROID_HOME%/build-tools subdirectory. Run SDK manager and install build-tools.");
-				}
-
-				// valid directories will have a source.properties with the Pkg.Revision (there is no guarantee we can use the directory name as revision)
-				foreach (string CandidateDir in Subdirs)
-				{
-					string AaptFilename = Path.Combine(CandidateDir, OperatingSystem.IsWindows() ? "aapt.exe" : "aapt");
-					string RevisionString = "";
-					uint RevisionValue = 0;
-
-					if (File.Exists(AaptFilename))
-					{
-						string SourcePropFilename = Path.Combine(CandidateDir, "source.properties");
-						if (File.Exists(SourcePropFilename))
-						{
-							string[] PropertyContents = File.ReadAllLines(SourcePropFilename);
-							foreach (string PropertyLine in PropertyContents)
-							{
-								if (PropertyLine.StartsWith("Pkg.Revision="))
-								{
-									RevisionString = PropertyLine.Substring(13);
-									RevisionValue = GetRevisionValue(RevisionString);
-									break;
-								}
-							}
-						}
-					}
-
-					// remember it if newer version or haven't found one yet
-					if (RevisionValue > BestVersion || BestVersionString == null)
-					{
-						BestVersion = RevisionValue;
-						BestVersionString = RevisionString;
-					}
-				}
-			}
-
-			if (BestVersionString == null)
-			{
-				BestVersionString = "35.0.1";
-				Logger.LogWarning("Failed to find %ANDROID_HOME%/build-tools subdirectory. Will attempt to use {BestVersionString}.", BestVersionString);
-			}
-
-			BestVersion = GetRevisionValue(BestVersionString);
-
-			// with Gradle enabled use at least 28.0.3 (will be installed by Gradle if missing)
-			if (BestVersion < ((28 << 24) | (0 << 16) | (3 << 8)))
-			{
-				BestVersionString = "28.0.3";
-			}
-
-			// don't allow higher than 30.0.3 for now (will be installed by Gradle if missing)
-			//			if (BestVersion > ((30 << 24) | (0 << 16) | (3 << 8)))
-			//			{
-			//				BestVersionString = "30.0.3";
-			//			}
-
-			CachedBuildToolsVersion = BestVersionString;
+			CachedBuildToolsVersion = ToolChain.GetBuildToolsVersion();
 			LastAndroidHomePath = HomePath;
 
 			Logger.LogInformation("Building with Build Tools version '{CachedBuildToolsVersion}'", CachedBuildToolsVersion);
@@ -1266,13 +1173,7 @@ namespace UnrealBuildTool
 
 		private static void StripDebugSymbols(string SourceFileName, string TargetFileName, UnrealArch UnrealArch, ILogger Logger, bool bStripAll = false)
 		{
-			ProcessStartInfo StartInfo = new ProcessStartInfo();
-			StartInfo.FileName = AndroidToolChain.GetStripExecutablePath(UnrealArch).Trim('"');
-			string StripCommand = bStripAll ? "--strip-unneeded" : "--strip-debug";
-			StartInfo.Arguments = $"{StripCommand} -o \"{TargetFileName}\" \"{SourceFileName}\"";
-			StartInfo.UseShellExecute = false;
-			StartInfo.CreateNoWindow = true;
-			Utils.RunLocalProcessAndLogOutput(StartInfo, Logger);
+			AndroidToolChain.StripSymbols(new FileReference(SourceFileName), new FileReference(TargetFileName), UnrealArch, Logger, bStripAll);
 		}
 
 		private List<string> GetSODependencies(string SOFilename)
@@ -3699,7 +3600,7 @@ namespace UnrealBuildTool
 			}
 		}
 
-		private readonly uint[] CRCTablesSB8 = {
+		private static uint[] CRCTablesSB8 = {
 			0x00000000, 0x77073096, 0xee0e612c, 0x990951ba, 0x076dc419, 0x706af48f, 0xe963a535, 0x9e6495a3, 0x0edb8832, 0x79dcb8a4, 0xe0d5e91e, 0x97d2d988, 0x09b64c2b, 0x7eb17cbd, 0xe7b82d07, 0x90bf1d91,
 			0x1db71064, 0x6ab020f2, 0xf3b97148, 0x84be41de, 0x1adad47d, 0x6ddde4eb, 0xf4d4b551, 0x83d385c7, 0x136c9856, 0x646ba8c0, 0xfd62f97a, 0x8a65c9ec, 0x14015c4f, 0x63066cd9, 0xfa0f3d63, 0x8d080df5,
 			0x3b6e20c8, 0x4c69105e, 0xd56041e4, 0xa2677172, 0x3c03e4d1, 0x4b04d447, 0xd20d85fd, 0xa50ab56b, 0x35b5a8fa, 0x42b2986c, 0xdbbbc9d6, 0xacbcf940, 0x32d86ce3, 0x45df5c75, 0xdcd60dcf, 0xabd13d59,
@@ -3718,210 +3619,99 @@ namespace UnrealBuildTool
 			0xbdbdf21c, 0xcabac28a, 0x53b39330, 0x24b4a3a6, 0xbad03605, 0xcdd70693, 0x54de5729, 0x23d967bf, 0xb3667a2e, 0xc4614ab8, 0x5d681b02, 0x2a6f2b94, 0xb40bbe37, 0xc30c8ea1, 0x5a05df1b, 0x2d02ef8d
 		};
 
-		private bool CopyAPKAndReplaceSO(ILogger Logger, string SourceAPK, string DestAPK, string SourceSOFile, string DestSOFile)
+		/// <summary>
+		/// Replaces a SO in an existing apk, performing an alignment.
+		/// This function can throw IO exceptions, ensure usage of proper try/catch handling
+		/// </summary>
+		/// <param name="SourceApk">The base apk to align</param>
+		/// <param name="TargetApk">Path to output the newly aligned apk</param>
+		/// <param name="Logger">Logger</param>
+		/// <param name="SourceSOFile">SO file to use as a replacement</param>
+		/// <param name="RelativeSOFile">Path of the SO file to replace, relative to the apk</param>
+		public static void CopyAPKAndReplaceSO(FileReference SourceApk, FileReference TargetApk, FileReference SourceSOFile, string RelativeSOFile, ILogger Logger)
 		{
-			BinaryReader SourceReader;
-			BinaryWriter DestWriter;
-			BinaryWriter DirectoryWriter;
-
-			string DirectoryFilename = DestAPK + ".dir";
+			string DirectoryFilename = TargetApk.FullName + ".dir";
 
-			try
+			using BinaryReader SourceReader = new BinaryReader(new FileStream(SourceApk.FullName, FileMode.Open, FileAccess.Read));
+			using BinaryWriter DestWriter = new BinaryWriter(new FileStream(TargetApk.FullName, FileMode.Create));
+			using BinaryWriter DirectoryWriter = new BinaryWriter(new FileStream(DirectoryFilename, FileMode.Create));
 			{
-				SourceReader = new BinaryReader(new FileStream(SourceAPK, FileMode.Open, FileAccess.Read));
-			}
-			catch (IOException)
-			{
-				return false;
-			}
+				byte[] CopyBuffer = new byte[65536];
+				ushort TotalEntries = 0;
+				uint Remaining;
 
-			try
-			{
-				DestWriter = new BinaryWriter(new FileStream(DestAPK, FileMode.Create));
-			}
-			catch (IOException)
-			{
-				return false;
-			}
+				uint Signature = 0x02014b50;
+				ushort DirVersion = 0;
+				ushort DiskNumber = 0;
+				ushort DiskStart = 0;
+				ushort InternalAttr = 0;
+				uint ExternalAttr = 0;
+				ushort CommentLen = 0;
 
-			try
-			{
-				DirectoryWriter = new BinaryWriter(new FileStream(DirectoryFilename, FileMode.Create));
-			}
-			catch (IOException)
-			{
-				return false;
-			}
-
-			byte[] CopyBuffer = new byte[65536];
-			ushort TotalEntries = 0;
-			uint Remaining;
-
-			uint Signature = 0x02014b50;
-			ushort DirVersion = 0;
-			ushort DiskNumber = 0;
-			ushort DiskStart = 0;
-			ushort InternalAttr = 0;
-			uint ExternalAttr = 0;
-			ushort CommentLen = 0;
-
-			// copy files
-			while (true)
-			{
-				uint Header = SourceReader.ReadUInt32();
-				if (Header != 0x04034b50)
+				// copy files
+				while (true)
 				{
-					SourceReader.Close();
-					break;
-				}
+					uint Header = SourceReader.ReadUInt32();
+					if (Header != 0x04034b50)
+					{
+						SourceReader.Close();
+						break;
+					}
 
-				ushort Version = SourceReader.ReadUInt16();
-				ushort Flags = SourceReader.ReadUInt16();
-				ushort Compression = SourceReader.ReadUInt16();
-				ushort ModTime = SourceReader.ReadUInt16();
-				ushort ModDate = SourceReader.ReadUInt16();
-				uint CRC32Value = SourceReader.ReadUInt32();
-				uint CompressedSize = SourceReader.ReadUInt32();
-				uint UncompressedSize = SourceReader.ReadUInt32();
-				ushort FilenameLen = SourceReader.ReadUInt16();
-				ushort ExtraLen = SourceReader.ReadUInt16();
-				byte[] Filename = new byte[FilenameLen];
-				SourceReader.BaseStream.Read(Filename, 0, FilenameLen);
-				byte[] Extra = new byte[ExtraLen];
-				SourceReader.BaseStream.Read(Extra, 0, ExtraLen);
-
-				string FilenameStr = Encoding.UTF8.GetString(Filename, 0, FilenameLen);
-				if (!FilenameStr.Equals(DestSOFile, StringComparison.Ordinal))
-				{
-					TotalEntries++;
-					uint Location = (uint)DestWriter.BaseStream.Position;
-
-					// calculate location after header and any additional alignment needed (only if uncompressed)
-					uint HeaderSize = (uint)(4 + 2 + 2 + 2 + 2 + 2 + 4 + 4 + 4 + 2 + 2 + FilenameLen + ExtraLen);
-					uint HeaderEnd = Location + HeaderSize;
-					uint DataStart = (uint)((HeaderEnd + 3) & ~3);
-					ushort ExtraAlign = (ushort)(Compression != 0 ? 0 : (DataStart - HeaderEnd));
-
-					// write local header
-					DestWriter.Write(Header);
-					DestWriter.Write(Version);
-					DestWriter.Write(Flags);
-					DestWriter.Write(Compression);
-					DestWriter.Write(ModTime);
-					DestWriter.Write(ModDate);
-					DestWriter.Write(CRC32Value);
-					DestWriter.Write(CompressedSize);
-					DestWriter.Write(UncompressedSize);
-					DestWriter.Write(FilenameLen);
-					DestWriter.Write((ushort)(ExtraLen + ExtraAlign));
-					DestWriter.BaseStream.Write(Filename, 0, FilenameLen);
-					DestWriter.BaseStream.Write(Extra, 0, ExtraLen);
-					while (ExtraAlign-- > 0)
-					{
-						DestWriter.Write((byte)0);
-					}
-
-					// copy file data as-is
-					Remaining = CompressedSize;
-					while (Remaining > 0)
-					{
-						int CopySize = Remaining < 65536 ? (int)Remaining : 65536;
-						int BytesRead = SourceReader.BaseStream.Read(CopyBuffer, 0, CopySize);
-						DestWriter.BaseStream.Write(CopyBuffer, 0, BytesRead);
-						Remaining -= (uint)BytesRead;
-					}
-
-					// write central directory entry for this file
-					DirectoryWriter.Write(Signature);
-					DirectoryWriter.Write(DirVersion);
-					DirectoryWriter.Write(Version);
-					DirectoryWriter.Write(Flags);
-					DirectoryWriter.Write(Compression);
-					DirectoryWriter.Write(ModTime);
-					DirectoryWriter.Write(ModDate);
-					DirectoryWriter.Write(CRC32Value);
-					DirectoryWriter.Write(CompressedSize);
-					DirectoryWriter.Write(UncompressedSize);
-					DirectoryWriter.Write(FilenameLen);
-					DirectoryWriter.Write(ExtraLen);
-					DirectoryWriter.Write(CommentLen);
-					DirectoryWriter.Write(DiskStart);
-					DirectoryWriter.Write(InternalAttr);
-					DirectoryWriter.Write(ExternalAttr);
-					DirectoryWriter.Write(Location);
-					DirectoryWriter.BaseStream.Write(Filename, 0, FilenameLen);
-					DirectoryWriter.BaseStream.Write(Extra, 0, ExtraLen);
-					// write comment here if there was one
-				}
-				else
-				{
-					// replace with the new .SO
-					try
+					ushort Version = SourceReader.ReadUInt16();
+					ushort Flags = SourceReader.ReadUInt16();
+					ushort Compression = SourceReader.ReadUInt16();
+					ushort ModTime = SourceReader.ReadUInt16();
+					ushort ModDate = SourceReader.ReadUInt16();
+					uint CRC32Value = SourceReader.ReadUInt32();
+					uint CompressedSize = SourceReader.ReadUInt32();
+					uint UncompressedSize = SourceReader.ReadUInt32();
+					ushort FilenameLen = SourceReader.ReadUInt16();
+					ushort ExtraLen = SourceReader.ReadUInt16();
+					byte[] Filename = new byte[FilenameLen];
+					SourceReader.BaseStream.Read(Filename, 0, FilenameLen);
+					byte[] Extra = new byte[ExtraLen];
+					SourceReader.BaseStream.Read(Extra, 0, ExtraLen);
+
+					string FilenameStr = Encoding.UTF8.GetString(Filename, 0, FilenameLen);
+					if (!FilenameStr.Equals(RelativeSOFile))
 					{
-						FileStream SourceFile = File.OpenRead(SourceSOFile);
+						TotalEntries++;
+						uint Location = (uint)DestWriter.BaseStream.Position;
 
-						UncompressedSize = (uint)SourceFile.Length;
-						CompressedSize = 0;
-						Compression = 8;        // deflate
+						// calculate location after header and any additional alignment needed (only if uncompressed)
+						uint HeaderSize = (uint)(4 + 2 + 2 + 2 + 2 + 2 + 4 + 4 + 4 + 2 + 2 + FilenameLen + ExtraLen);
+						uint HeaderEnd = Location + HeaderSize;
+						uint DataStart = (uint)((HeaderEnd + 3) & ~3);
+						ushort ExtraAlign = (ushort)(Compression != 0 ? 0 : (DataStart - HeaderEnd));
+
+						// write local header
+						DestWriter.Write(Header);
+						DestWriter.Write(Version);
+						DestWriter.Write(Flags);
+						DestWriter.Write(Compression);
+						DestWriter.Write(ModTime);
+						DestWriter.Write(ModDate);
+						DestWriter.Write(CRC32Value);
+						DestWriter.Write(CompressedSize);
+						DestWriter.Write(UncompressedSize);
+						DestWriter.Write(FilenameLen);
+						DestWriter.Write((ushort)(ExtraLen + ExtraAlign));
+						DestWriter.BaseStream.Write(Filename, 0, FilenameLen);
+						DestWriter.BaseStream.Write(Extra, 0, ExtraLen);
+						while (ExtraAlign-- > 0)
+						{
+							DestWriter.Write((byte)0);
+						}
 
-						// calculate CRC32
-						CRC32Value = 0x0;
-						CRC32Value = ~CRC32Value;
-						Remaining = UncompressedSize;
+						// copy file data as-is
+						Remaining = CompressedSize;
 						while (Remaining > 0)
 						{
 							int CopySize = Remaining < 65536 ? (int)Remaining : 65536;
-							int BytesRead = SourceFile.Read(CopyBuffer, 0, CopySize);
+							int BytesRead = SourceReader.BaseStream.Read(CopyBuffer, 0, CopySize);
+							DestWriter.BaseStream.Write(CopyBuffer, 0, BytesRead);
 							Remaining -= (uint)BytesRead;
-							for (uint Index = 0; Index < BytesRead; Index++)
-							{
-								CRC32Value = (CRC32Value >> 8) ^ CRCTablesSB8[(CRC32Value ^ (uint)CopyBuffer[Index]) & 0xff];
-							}
-						}
-						CRC32Value = ~CRC32Value;
-						SourceFile.Position = 0;
-
-						TotalEntries++;
-						uint Location = (uint)DestWriter.BaseStream.Position;
-
-						using (MemoryStream CompressedStream = new MemoryStream())
-						{
-							using (DeflateStream Compressor = new DeflateStream(CompressedStream, CompressionMode.Compress, leaveOpen: true))
-							{
-								SourceFile.CopyTo(Compressor);
-								SourceFile.Close();
-							}
-
-							byte[] CompressedBuffer = CompressedStream.GetBuffer();
-							CompressedSize = (uint)CompressedStream.Length;
-
-							// calculate location after header and any additional alignment needed (only if uncompressed)
-							uint HeaderSize = (uint)(4 + 2 + 2 + 2 + 2 + 2 + 4 + 4 + 4 + 2 + 2 + FilenameLen + ExtraLen);
-							uint HeaderEnd = Location + HeaderSize;
-							uint DataStart = (uint)((HeaderEnd + 3) & ~3);
-							ushort ExtraAlign = (ushort)(Compression != 0 ? 0 : (DataStart - HeaderEnd));
-
-							// write local header
-							DestWriter.Write(Header);
-							DestWriter.Write(Version);
-							DestWriter.Write(Flags);
-							DestWriter.Write(Compression);
-							DestWriter.Write(ModTime);
-							DestWriter.Write(ModDate);
-							DestWriter.Write(CRC32Value);
-							DestWriter.Write(CompressedSize);
-							DestWriter.Write(UncompressedSize);
-							DestWriter.Write(FilenameLen);
-							DestWriter.Write((ushort)(ExtraLen + ExtraAlign));
-							DestWriter.BaseStream.Write(Filename, 0, FilenameLen);
-							DestWriter.BaseStream.Write(Extra, 0, ExtraLen);
-							while (ExtraAlign-- > 0)
-							{
-								DestWriter.Write((byte)0);
-							}
-
-							// write compressed data
-							DestWriter.BaseStream.Write(CompressedBuffer, 0, (int)CompressedSize);
 						}
 
 						// write central directory entry for this file
@@ -3946,46 +3736,139 @@ namespace UnrealBuildTool
 						DirectoryWriter.BaseStream.Write(Extra, 0, ExtraLen);
 						// write comment here if there was one
 					}
-					catch (IOException e)
+					else
 					{
-						Logger.LogInformation("Failed to add {SourceSOFile} to APK: Reason = {Reason}", SourceSOFile, e.ToString());
-						DirectoryWriter.Close();
-						File.Delete(DirectoryFilename);
-						DestWriter.Close();
-						File.Delete(DestAPK);
-						return false;
+						// replace with the new .SO
+						try
+						{
+							FileStream SourceFile = File.OpenRead(SourceSOFile.FullName);
+
+							UncompressedSize = (uint)SourceFile.Length;
+							CompressedSize = 0;
+							Compression = 8;        // deflate
+
+							// calculate CRC32
+							CRC32Value = 0x0;
+							CRC32Value = ~CRC32Value;
+							Remaining = UncompressedSize;
+							while (Remaining > 0)
+							{
+								int CopySize = Remaining < 65536 ? (int)Remaining : 65536;
+								int BytesRead = SourceFile.Read(CopyBuffer, 0, CopySize);
+								Remaining -= (uint)BytesRead;
+								for (uint Index = 0; Index < BytesRead; Index++)
+								{
+									CRC32Value = (CRC32Value >> 8) ^ CRCTablesSB8[(CRC32Value ^ (uint)CopyBuffer[Index]) & 0xff];
+								}
+							}
+							CRC32Value = ~CRC32Value;
+							SourceFile.Position = 0;
+
+							TotalEntries++;
+							uint Location = (uint)DestWriter.BaseStream.Position;
+
+							using (MemoryStream CompressedStream = new MemoryStream())
+							{
+								using (DeflateStream Compressor = new DeflateStream(CompressedStream, CompressionMode.Compress, leaveOpen: true))
+								{
+									SourceFile.CopyTo(Compressor);
+									SourceFile.Close();
+								}
+
+								byte[] CompressedBuffer = CompressedStream.GetBuffer();
+								CompressedSize = (uint)CompressedStream.Length;
+
+								// calculate location after header and any additional alignment needed (only if uncompressed)
+								uint HeaderSize = (uint)(4 + 2 + 2 + 2 + 2 + 2 + 4 + 4 + 4 + 2 + 2 + FilenameLen + ExtraLen);
+								uint HeaderEnd = Location + HeaderSize;
+								uint DataStart = (uint)((HeaderEnd + 3) & ~3);
+								ushort ExtraAlign = (ushort)(Compression != 0 ? 0 : (DataStart - HeaderEnd));
+
+								// write local header
+								DestWriter.Write(Header);
+								DestWriter.Write(Version);
+								DestWriter.Write(Flags);
+								DestWriter.Write(Compression);
+								DestWriter.Write(ModTime);
+								DestWriter.Write(ModDate);
+								DestWriter.Write(CRC32Value);
+								DestWriter.Write(CompressedSize);
+								DestWriter.Write(UncompressedSize);
+								DestWriter.Write(FilenameLen);
+								DestWriter.Write((ushort)(ExtraLen + ExtraAlign));
+								DestWriter.BaseStream.Write(Filename, 0, FilenameLen);
+								DestWriter.BaseStream.Write(Extra, 0, ExtraLen);
+								while (ExtraAlign-- > 0)
+								{
+									DestWriter.Write((byte)0);
+								}
+
+								// write compressed data
+								DestWriter.BaseStream.Write(CompressedBuffer, 0, (int)CompressedSize);
+							}
+
+							// write central directory entry for this file
+							DirectoryWriter.Write(Signature);
+							DirectoryWriter.Write(DirVersion);
+							DirectoryWriter.Write(Version);
+							DirectoryWriter.Write(Flags);
+							DirectoryWriter.Write(Compression);
+							DirectoryWriter.Write(ModTime);
+							DirectoryWriter.Write(ModDate);
+							DirectoryWriter.Write(CRC32Value);
+							DirectoryWriter.Write(CompressedSize);
+							DirectoryWriter.Write(UncompressedSize);
+							DirectoryWriter.Write(FilenameLen);
+							DirectoryWriter.Write(ExtraLen);
+							DirectoryWriter.Write(CommentLen);
+							DirectoryWriter.Write(DiskStart);
+							DirectoryWriter.Write(InternalAttr);
+							DirectoryWriter.Write(ExternalAttr);
+							DirectoryWriter.Write(Location);
+							DirectoryWriter.BaseStream.Write(Filename, 0, FilenameLen);
+							DirectoryWriter.BaseStream.Write(Extra, 0, ExtraLen);
+							// write comment here if there was one
+						}
+						catch (IOException e)
+						{
+							Logger.LogInformation("Failed to add {SourceSOFile} to APK: Reason = {Reason}", SourceSOFile, e.ToString());
+							DirectoryWriter.Close();
+							File.Delete(DirectoryFilename);
+							DestWriter.Close();
+							FileReference.Delete(TargetApk);
+							throw;
+						}
 					}
 				}
-			}
 
-			uint CentralDirectoryLocation = (uint)DestWriter.BaseStream.Position;
+				uint CentralDirectoryLocation = (uint)DestWriter.BaseStream.Position;
 
-			// copy the directory to the end of the APK
-			uint DirectorySize = (uint)DirectoryWriter.BaseStream.Position;
-			DirectoryWriter.BaseStream.Position = 0;
-			Remaining = DirectorySize;
-			while (Remaining > 0)
-			{
-				int CopySize = Remaining < 65536 ? (int)Remaining : 65536;
-				int BytesRead = DirectoryWriter.BaseStream.Read(CopyBuffer, 0, CopySize);
-				DestWriter.BaseStream.Write(CopyBuffer, 0, BytesRead);
-				Remaining -= (uint)BytesRead;
-			}
-			DirectoryWriter.Close();
-			File.Delete(DirectoryFilename);
+				// copy the directory to the end of the APK
+				uint DirectorySize = (uint)DirectoryWriter.BaseStream.Position;
+				DirectoryWriter.BaseStream.Position = 0;
+				Remaining = DirectorySize;
+				while (Remaining > 0)
+				{
+					int CopySize = Remaining < 65536 ? (int)Remaining : 65536;
+					int BytesRead = DirectoryWriter.BaseStream.Read(CopyBuffer, 0, CopySize);
+					DestWriter.BaseStream.Write(CopyBuffer, 0, BytesRead);
+					Remaining -= (uint)BytesRead;
+				}
+				DirectoryWriter.Close();
+				File.Delete(DirectoryFilename);
 
-			// write end of central directory record
-			Signature = 0x06054b50;
-			DestWriter.Write(Signature);
-			DestWriter.Write(DiskNumber);
-			DestWriter.Write(DiskStart);
-			DestWriter.Write(TotalEntries);
-			DestWriter.Write(TotalEntries);
-			DestWriter.Write(DirectorySize);
-			DestWriter.Write(CentralDirectoryLocation);
-			DestWriter.Write(CommentLen);
-			DestWriter.Close();
-			return true;
+				// write end of central directory record
+				Signature = 0x06054b50;
+				DestWriter.Write(Signature);
+				DestWriter.Write(DiskNumber);
+				DestWriter.Write(DiskStart);
+				DestWriter.Write(TotalEntries);
+				DestWriter.Write(TotalEntries);
+				DestWriter.Write(DirectorySize);
+				DestWriter.Write(CentralDirectoryLocation);
+				DestWriter.Write(CommentLen);
+				DestWriter.Close();
+			}
 		}
 
 		private void CreateGradlePropertiesFiles(AndroidToolChain ToolChain, int MinSDKVersion, int TargetSDKVersion, string CompileSDKVersion, string BuildToolsVersion, string PackageName,
@@ -4408,7 +4291,7 @@ dependencies.implementation('androidx.tracing:tracing:1.2.0')
 
 			// do this here so we'll stop early if there is a problem with the SDK API level (cached so later calls will return the same)
 			string SDKAPILevel = GetSdkApiLevel(ToolChain);
-			string BuildToolsVersion = GetBuildToolsVersion();
+			string BuildToolsVersion = GetBuildToolsVersion(ToolChain);
 
 			// cache some tools paths
 			//string NDKBuildPath = Environment.ExpandEnvironmentVariables("%NDKROOT%/ndk-build" + (OperatingSystem.IsWindows() ? ".cmd" : ""));
@@ -4903,7 +4786,15 @@ dependencies.implementation('androidx.tracing:tracing:1.2.0')
 					Logger.LogInformation("\n===={Time}====PERFORMING FINAL APK PACKAGE OPERATION====={Arch}===========================================", DateTime.Now.ToString(), Arch);
 
 					string TempAPK = Path.Combine(DestApkDirectory, "TEMP_" + Path.GetFileName(DestApkName));
-					CopyAPKAndReplaceSO(Logger, DestApkName, TempAPK, StrippedSOName, "lib/" + NDKArch + "/libUnreal.so");
+					try
+					{
+						FileReference DestApkFile = new FileReference(DestApkName);
+						FileReference TempApkFile = new FileReference(TempAPK);
+						FileReference StrippedSOFile = new FileReference(StrippedSOName);
+						CopyAPKAndReplaceSO(DestApkFile, TempApkFile, StrippedSOFile, "lib/" + NDKArch + "/libUnreal.so", Logger);
+					}
+					catch { }
+
 
 					// defaults for non-distribution build (debug)
 					string KeyStore = Path.Combine(Environment.GetFolderPath(Environment.SpecialFolder.UserProfile), ".android", "debug.keystore");
@@ -4968,59 +4859,16 @@ dependencies.implementation('androidx.tracing:tracing:1.2.0')
 						throw new BuildException("Keystore file is missing. Check the DistributionSettings section in the Android tab of Project Settings");
 					}
 
-					// sign the APK
-					string BuildToolsPath = Path.Combine(Environment.ExpandEnvironmentVariables("%ANDROID_HOME%"), "Build-Tools", BuildToolsVersion);
-					string APKSignerExecutable = Path.Combine(BuildToolsPath, "apksigner" + (OperatingSystem.IsWindows() ? ".bat" : ""));
-
-					ProcessStartInfo StartInfo = new ProcessStartInfo();
-					StartInfo.WorkingDirectory = IntermediateAndroidPath;
-					if (OperatingSystem.IsWindows())
-					{
-						StartInfo.FileName = "cmd.exe";
-						StartInfo.ArgumentList.Add("/c");
-						StartInfo.ArgumentList.Add(APKSignerExecutable);
-						StartInfo.ArgumentList.Add("sign");
-						StartInfo.ArgumentList.Add("--ks");
-						StartInfo.ArgumentList.Add(KeyStore);
-						StartInfo.ArgumentList.Add("--ks-pass");
-						StartInfo.ArgumentList.Add("pass:" + KeyStorePassword);
-						StartInfo.ArgumentList.Add("--ks-key-alias");
-						StartInfo.ArgumentList.Add(KeyAlias);
-						StartInfo.ArgumentList.Add("--key-pass");
-						StartInfo.ArgumentList.Add("pass:" + KeyPassword);
-						StartInfo.ArgumentList.Add("--out");
-						StartInfo.ArgumentList.Add(DestApkName);
-						StartInfo.ArgumentList.Add(TempAPK);
-					}
-					else
+					Logger.LogInformation("Applying apksigner to apk...");
+					try
 					{
-						StartInfo.FileName = "/bin/sh";
-						StartInfo.ArgumentList.Add("-c");
-						StartInfo.ArgumentList.Add("\"" + APKSignerExecutable + "\" sign--ks \"" + KeyStore + "\" --ks-pass pass:" + KeyStorePassword + " --ks-key-alias " + KeyAlias + " --key-pass pass:" + KeyPassword + " --out \"" + DestApkName + "2\" \"" + TempAPK + "\"");
+						ToolChain.SignApk(new FileReference(TempAPK), new FileReference(DestApkName), KeyStore, KeyStorePassword, KeyAlias, KeyPassword);
 					}
-					StartInfo.UseShellExecute = false;
-					StartInfo.WindowStyle = ProcessWindowStyle.Minimized;
-
-					Logger.LogInformation("Applying apksigner to apk...");
-
-					Process Proc = new Process();
-					Proc.StartInfo = StartInfo;
-					Proc.Start();
-					Proc.WaitForExit();
-
-					// clean up work files
-					SafeDeleteFile(TempAPK);
-					SafeDeleteFile(DestApkName + ".idsig");
-					SafeDeleteFile(StrippedSOName);
-
-					if (Proc.ExitCode != 0)
+					finally
 					{
-						string Args = "";
-						foreach (string Arg in StartInfo.ArgumentList)
-						{
-							Args += Arg + " ";
-						}
-						throw new BuildException("{0} failed with args {1}", StartInfo.FileName, Args);
+						SafeDeleteFile(TempAPK);
+						SafeDeleteFile(DestApkName + ".idsig");
+						SafeDeleteFile(StrippedSOName);
 					}
 
 					// copy .so with symbols if requested
