In UBT, use a thread-safe collection to collect the tasks for writing out action results.

#Synced-CL 49066283

diff --git a/Engine/Source/Programs/UnrealBuildTool/Executors/UnrealBuildAccelerator/UBAExecutor.cs b/Engine/Source/Programs/UnrealBuildTool/Executors/UnrealBuildAccelerator/UBAExecutor.cs
--- a/Engine/Source/Programs/UnrealBuildTool/Executors/UnrealBuildAccelerator/UBAExecutor.cs
+++ b/Engine/Source/Programs/UnrealBuildTool/Executors/UnrealBuildAccelerator/UBAExecutor.cs
@@ -1,6 +1,7 @@
 // Copyright Epic Games, Inc. All Rights Reserved.
 
 using System;
+using System.Collections.Concurrent;
 using System.Collections.Generic;
 using System.IO;
 using System.Linq;
@@ -792,7 +793,7 @@ namespace UnrealBuildTool
 			bool shouldCancel = false;
 			string cancelReason = "user cancellation";
 
-			List<Task> resultStoreTasks = [];
+			ConcurrentBag<Task> resultStoreTasks = [];
 
 			scheduler.SetProcessFinishedCallback((info) =>
 			{
