[Android] Fix CI style warnings
#Synced-CL 48971518

diff --git a/Engine/Source/Programs/UnrealBuildTool/Platform/Android/AndroidToolChain.cs b/Engine/Source/Programs/UnrealBuildTool/Platform/Android/AndroidToolChain.cs
--- a/Engine/Source/Programs/UnrealBuildTool/Platform/Android/AndroidToolChain.cs
+++ b/Engine/Source/Programs/UnrealBuildTool/Platform/Android/AndroidToolChain.cs
@@ -1469,8 +1469,8 @@ namespace UnrealBuildTool
 				throw new BuildException("Failed to find %ANDROID_HOME%. Ensure environment is configured properly");
 			}
 
-			string? BestVersionString = null;
 			uint BestVersion = 0;
+			string? BestVersionString;
 
 			ConfigHierarchy Ini = ConfigCache.ReadHierarchy(ConfigHierarchyType.Engine, DirectoryReference.FromFile(ProjectFile), UnrealTargetPlatform.Android);
 			Ini.GetString("/Script/AndroidRuntimeSettings.AndroidRuntimeSettings", "BuildToolsOverride", out BestVersionString);
diff --git a/Engine/Source/Programs/UnrealBuildTool/Platform/Android/UEDeployAndroid.cs b/Engine/Source/Programs/UnrealBuildTool/Platform/Android/UEDeployAndroid.cs
--- a/Engine/Source/Programs/UnrealBuildTool/Platform/Android/UEDeployAndroid.cs
+++ b/Engine/Source/Programs/UnrealBuildTool/Platform/Android/UEDeployAndroid.cs
@@ -3600,7 +3600,7 @@ namespace UnrealBuildTool
 			}
 		}
 
-		private static uint[] CRCTablesSB8 = {
+		private static readonly uint[] CRCTablesSB8 = {
 			0x00000000, 0x77073096, 0xee0e612c, 0x990951ba, 0x076dc419, 0x706af48f, 0xe963a535, 0x9e6495a3, 0x0edb8832, 0x79dcb8a4, 0xe0d5e91e, 0x97d2d988, 0x09b64c2b, 0x7eb17cbd, 0xe7b82d07, 0x90bf1d91,
 			0x1db71064, 0x6ab020f2, 0xf3b97148, 0x84be41de, 0x1adad47d, 0x6ddde4eb, 0xf4d4b551, 0x83d385c7, 0x136c9856, 0x646ba8c0, 0xfd62f97a, 0x8a65c9ec, 0x14015c4f, 0x63066cd9, 0xfa0f3d63, 0x8d080df5,
 			0x3b6e20c8, 0x4c69105e, 0xd56041e4, 0xa2677172, 0x3c03e4d1, 0x4b04d447, 0xd20d85fd, 0xa50ab56b, 0x35b5a8fa, 0x42b2986c, 0xdbbbc9d6, 0xacbcf940, 0x32d86ce3, 0x45df5c75, 0xdcd60dcf, 0xabd13d59,
@@ -3674,7 +3674,7 @@ namespace UnrealBuildTool
 					SourceReader.BaseStream.Read(Extra, 0, ExtraLen);
 
 					string FilenameStr = Encoding.UTF8.GetString(Filename, 0, FilenameLen);
-					if (!FilenameStr.Equals(RelativeSOFile))
+					if (!FilenameStr.Equals(RelativeSOFile, StringComparison.OrdinalIgnoreCase))
 					{
 						TotalEntries++;
 						uint Location = (uint)DestWriter.BaseStream.Position;
@@ -4795,7 +4795,6 @@ dependencies.implementation('androidx.tracing:tracing:1.2.0')
 					}
 					catch { }
 
-
 					// defaults for non-distribution build (debug)
 					string KeyStore = Path.Combine(Environment.GetFolderPath(Environment.SpecialFolder.UserProfile), ".android", "debug.keystore");
 					string KeyStorePassword = "android";
