The Android AndroidPostApkInstallCommands step in generated .vcxproj references RunUAT.bat wrapper from Unreal root directory. This doesn't exist on installed/GitHub builds so instead reference the the actual Engine\Build\BatchFiles\RunUAT.bat.

#Synced-CL 48848898

diff --git a/Engine/Source/Programs/UnrealBuildTool/Platform/Android/AndroidProjectGenerator.cs b/Engine/Source/Programs/UnrealBuildTool/Platform/Android/AndroidProjectGenerator.cs
--- a/Engine/Source/Programs/UnrealBuildTool/Platform/Android/AndroidProjectGenerator.cs
+++ b/Engine/Source/Programs/UnrealBuildTool/Platform/Android/AndroidProjectGenerator.cs
@@ -211,7 +211,7 @@ namespace UnrealBuildTool
 					// It's critical to have AndroidDebugTarget here and for it to be before properties that use it (e.g. AndroidPostApkInstallCommands), otherwise MSBuild will evaluate it to empty string.
 					Out.AppendLine("    <AndroidDebugTarget></AndroidDebugTarget>");
 
-					string UatPath = Path.Combine(Unreal.RootDirectory.FullName, "RunUAT.bat");
+					string UatPath = Path.Combine(Unreal.RootDirectory.FullName, "Engine", "Build", "BatchFiles", "RunUAT.bat");
 					string ShortAbi = GetShortAbi(InVSSettings);
 
 					// AGDE specifies current debug target in AndroidDebugTarget property in a form of "model:serial:arch".
