[Android] Ensure toolchain is initialized before using Android export symbol stripping

diff --git a/Engine/Source/Programs/UnrealBuildTool/Platform/Android/AndroidExports.cs b/Engine/Source/Programs/UnrealBuildTool/Platform/Android/AndroidExports.cs
--- a/Engine/Source/Programs/UnrealBuildTool/Platform/Android/AndroidExports.cs
+++ b/Engine/Source/Programs/UnrealBuildTool/Platform/Android/AndroidExports.cs
@@ -93,7 +93,7 @@ namespace UnrealBuildTool
 	public static class AndroidExports
 	{
 		private static AndroidToolChain? CachedToolChain;
-		private static AndroidToolChain GetToolChain(FileReference ProjectFile, ILogger Logger)
+		private static AndroidToolChain GetToolChain(FileReference? ProjectFile, ILogger Logger)
 		{
 			if (CachedToolChain == null || !CachedToolChain.Equals(ProjectFile, Logger))
 			{
@@ -142,6 +142,12 @@ namespace UnrealBuildTool
 		/// <param name="bStripAll">Whether to strip unneeded symbols or just debug symbols</param>
 		public static void StripSymbols(FileReference SourceFile, FileReference TargetFile, ILogger Logger, bool bStripAll = false)
 		{
+			// This forces ToolChain's static init to run, which will populate the strip binary paths
+			if (CachedToolChain == null)
+			{
+				GetToolChain(null, Logger);
+			}
+
 			AndroidToolChain.StripSymbols(SourceFile, TargetFile, Logger, bStripAll);
 		}
 
diff --git a/Engine/Source/Programs/UnrealBuildTool/Platform/Android/AndroidToolChain.cs b/Engine/Source/Programs/UnrealBuildTool/Platform/Android/AndroidToolChain.cs
--- a/Engine/Source/Programs/UnrealBuildTool/Platform/Android/AndroidToolChain.cs
+++ b/Engine/Source/Programs/UnrealBuildTool/Platform/Android/AndroidToolChain.cs
@@ -254,7 +254,7 @@ namespace UnrealBuildTool
 			SysrootPath = DirectoryReference.FromString(Path.Combine(NDKPath!, @"toolchains/llvm", ArchitecturePath!, "sysroot").Replace("\\", "/"))!;
 		}
 
-		public bool Equals(FileReference InProjectFile, ILogger InLogger)
+		public bool Equals(FileReference? InProjectFile, ILogger InLogger)
 		{
 			return ProjectFile == InProjectFile && Logger == InLogger;
 		}
@@ -1542,23 +1542,29 @@ namespace UnrealBuildTool
 		}
 
 		public static void StripSymbols(FileReference SourceFile, FileReference TargetFile, UnrealArch UnrealArch, ILogger Logger, bool bStripAll = false)
+		{
+			string StripPath = GetStripExecutablePath(UnrealArch).Trim('"');
+			StripSymbols(StripPath, SourceFile, TargetFile, Logger, bStripAll);
+		}
+
+		public static void StripSymbols(FileReference SourceFile, FileReference TargetFile, ILogger Logger, bool bStripAll = false)
+		{
+			string StripPath = GetStripPath(SourceFile).Trim('"');
+			StripSymbols(StripPath, SourceFile, TargetFile, Logger, bStripAll);
+		}
+
+		private static void StripSymbols(string StripPath, FileReference SourceFile, FileReference TargetFile, ILogger Logger, bool bStripAll = false)
 		{
 			string StripCommand = bStripAll ? "--strip-unneeded" : "--strip-debug";
 
 			ProcessStartInfo StartInfo = new ProcessStartInfo();
-			StartInfo.FileName = GetStripExecutablePath(UnrealArch).Trim('"');
+			StartInfo.FileName = StripPath;
 			StartInfo.Arguments = $"{StripCommand} -o \"{TargetFile.FullName}\" \"{SourceFile.FullName}\"";
 			StartInfo.UseShellExecute = false;
 			StartInfo.CreateNoWindow = true;
 			Utils.RunLocalProcessAndLogOutput(StartInfo, Logger);
 		}
 
-		public static void StripSymbols(FileReference SourceFile, FileReference TargetFile, ILogger Logger, bool bStripAll = false)
-		{
-			UnrealArch Architecture = SourceFile.FullName.Contains("arm64", StringComparison.OrdinalIgnoreCase) ? UnrealArch.Arm64 : UnrealArch.X64;
-			StripSymbols(SourceFile, TargetFile, Architecture, Logger, bStripAll);
-		}
-
 		public void SignApk(FileReference SourceApk, FileReference TargetApk, bool bZipAlign = false)
 		{
 			string KeyStore = Path.Combine(Environment.GetFolderPath(Environment.SpecialFolder.UserProfile), ".android", "debug.keystore");
