[Android] Fixed APK repack step failing to copy files after replacing a shared object. The source stream was not advanced past the original entry?s data. This made parsing stop early and omit remaining files.
#Synced-CL 49130326

diff --git a/Engine/Source/Programs/UnrealBuildTool/Platform/Android/UEDeployAndroid.cs b/Engine/Source/Programs/UnrealBuildTool/Platform/Android/UEDeployAndroid.cs
--- a/Engine/Source/Programs/UnrealBuildTool/Platform/Android/UEDeployAndroid.cs
+++ b/Engine/Source/Programs/UnrealBuildTool/Platform/Android/UEDeployAndroid.cs
@@ -3742,7 +3742,7 @@ namespace UnrealBuildTool
 						try
 						{
 							FileStream SourceFile = File.OpenRead(SourceSOFile.FullName);
-
+							uint SourceCompressedSize = CompressedSize;
 							UncompressedSize = (uint)SourceFile.Length;
 							CompressedSize = 0;
 							Compression = 8;        // deflate
@@ -3805,6 +3805,9 @@ namespace UnrealBuildTool
 
 								// write compressed data
 								DestWriter.BaseStream.Write(CompressedBuffer, 0, (int)CompressedSize);
+
+								// advance the base stream forward by the source file size so we land on the next header
+								SourceReader.BaseStream.Position += SourceCompressedSize;
 							}
 
 							// write central directory entry for this file
