[UBT]
* Changed ApplePlatformSDK.HasRequiredManualSDKInternal to always return true for non-mac platforms. This was already always true for build machines and after discussing with Josh Adams and others we decided to match build farm settings for all devs

(No point preflighting since this only change non-farm behavior :-))

#Synced-CL 49171070

diff --git a/Engine/Source/Programs/UnrealBuildTool/Platform/Apple/ApplePlatformSDK.cs b/Engine/Source/Programs/UnrealBuildTool/Platform/Apple/ApplePlatformSDK.cs
--- a/Engine/Source/Programs/UnrealBuildTool/Platform/Apple/ApplePlatformSDK.cs
+++ b/Engine/Source/Programs/UnrealBuildTool/Platform/Apple/ApplePlatformSDK.cs
@@ -269,7 +269,7 @@ namespace UnrealBuildTool
 
 			// iTunes is technically only need to deploy to and run on connected devices.
 			// This code removes requirement for Windows builders to have Xcode installed.
-			if (Status == SDKStatus.Invalid && !OperatingSystem.IsMacOS() && Unreal.IsBuildMachine())
+			if (Status == SDKStatus.Invalid && !OperatingSystem.IsMacOS() && true) // Unreal.IsBuildMachine())
 			{
 				Status = SDKStatus.Valid;
 			}
