[UBT]
* Removed using namespace to make static analysis happy

#Synced-CL 49186122

diff --git a/Engine/Source/Programs/UnrealBuildTool/Platform/Apple/ApplePlatformSDK.cs b/Engine/Source/Programs/UnrealBuildTool/Platform/Apple/ApplePlatformSDK.cs
--- a/Engine/Source/Programs/UnrealBuildTool/Platform/Apple/ApplePlatformSDK.cs
+++ b/Engine/Source/Programs/UnrealBuildTool/Platform/Apple/ApplePlatformSDK.cs
@@ -8,7 +8,6 @@ using System.Text.RegularExpressions;
 using EpicGames.Core;
 using Microsoft.Extensions.Logging;
 using Microsoft.Win32;
-using UnrealBuildBase;
 
 ///////////////////////////////////////////////////////////////////
 // If you are looking for supported version numbers, look in the
