UnrealBuildTool support for observability of actions. Store a manifest of all actions ran for a given target, and add a mode to render the manifest.

#Synced-CL 49202056

diff --git a/Engine/Source/Programs/Shared/EpicGames.Core/CommandLineArguments.cs b/Engine/Source/Programs/Shared/EpicGames.Core/CommandLineArguments.cs
--- a/Engine/Source/Programs/Shared/EpicGames.Core/CommandLineArguments.cs
+++ b/Engine/Source/Programs/Shared/EpicGames.Core/CommandLineArguments.cs
@@ -10,6 +10,8 @@ using System.Diagnostics.CodeAnalysis;
 using System.Linq;
 using System.Reflection;
 using System.Text;
+using System.Text.Json;
+using System.Text.Json.Serialization;
 using Microsoft.Extensions.Logging;
 
 #pragma warning disable CA1710 // Identifiers should have correct suffix
@@ -79,6 +81,7 @@ namespace EpicGames.Core
 	/// </summary>
 	[DebuggerDisplay("Count = {Count}")]
 	[DebuggerTypeProxy(typeof(CommandLineArgumentListView))]
+	[JsonConverter(typeof(CommandLineArgumentsJsonConverter))]
 	public class CommandLineArguments : IReadOnlyList<string>, IReadOnlyCollection<string>, IEnumerable<string>, IEnumerable
 	{
 		/// <summary>
@@ -1446,4 +1449,48 @@ namespace EpicGames.Core
 			return result.ToString();
 		}
 	}
+
+	/// <summary>
+	/// Serializes command-line arguments to and from an array of strings.
+	/// </summary>
+	public class CommandLineArgumentsJsonConverter : JsonConverter<CommandLineArguments>
+	{
+		/// <inheritdoc />
+		public override CommandLineArguments? Read(ref Utf8JsonReader reader, Type typeToConvert, JsonSerializerOptions options)
+		{
+			List<string> arguments = [];
+			if (reader.TokenType != JsonTokenType.StartArray)
+			{
+				throw new JsonException();
+			}
+			while (reader.Read())
+			{
+				if (reader.TokenType == JsonTokenType.EndArray)
+				{
+					return new CommandLineArguments([.. arguments]);
+				}
+
+				if (reader.TokenType != JsonTokenType.String)
+				{
+					throw new JsonException();
+				}
+
+				string argument = reader.GetString()!;
+				arguments.Add(argument);
+			}
+
+			throw new JsonException();
+		}
+
+		/// <inheritdoc />
+		public override void Write(Utf8JsonWriter writer, CommandLineArguments value, JsonSerializerOptions options)
+		{
+			writer.WriteStartArray();
+			for (int i = 0; i < value.Count; ++i)
+			{
+				writer.WriteStringValue(value[i]);
+			}
+			writer.WriteEndArray();
+		}
+	}
 }
diff --git a/Engine/Source/Programs/UnrealBuildTool/Actions/ActionResult.cs b/Engine/Source/Programs/UnrealBuildTool/Actions/ActionResult.cs
--- a/Engine/Source/Programs/UnrealBuildTool/Actions/ActionResult.cs
+++ b/Engine/Source/Programs/UnrealBuildTool/Actions/ActionResult.cs
@@ -7,15 +7,16 @@ namespace UnrealBuildTool;
 /// <summary>
 /// The result of the execution of an action.
 /// </summary>
+/// <param name="Description">A human-readable description of the action.</param>
 /// <param name="ExitCode">The exit code the action returned after executing.</param>
 /// <param name="Success">Whether the action succeeded - actions may still fail even if the exit code was 0.</param>
-/// <param name="Output"></param>
-/// <param name="Metrics"></param>
-internal record class ActionResult(int ExitCode, bool Success, IReadOnlyList<string> Output, ActionMetrics Metrics)
+/// <param name="Output">The output of the action.</param>
+/// <param name="Metrics">The metrics associated with the action.</param>
+internal record class ActionResult(string Description, int ExitCode, bool Success, IReadOnlyList<string> Output, ActionMetrics Metrics)
 {
 	public static ActionResult From(ExecuteResults executeResults)
 	{
 		ActionMetrics metrics = new(executeResults.ExecutionTime, executeResults.ProcessorTime, (ulong)executeResults.PeakMemoryUsed);
-		return new(executeResults.ExitCode, executeResults.ExitCode == 0, executeResults.LogLines, metrics);
+		return new(executeResults.GetDescription(), executeResults.ExitCode, executeResults.ExitCode == 0, executeResults.LogLines, metrics);
 	}
 }
diff --git a/Engine/Source/Programs/UnrealBuildTool/Actions/ActionResultManifest.cs b/Engine/Source/Programs/UnrealBuildTool/Actions/ActionResultManifest.cs
new file mode 100644
--- /dev/null
+++ b/Engine/Source/Programs/UnrealBuildTool/Actions/ActionResultManifest.cs
@@ -0,0 +1,7 @@
+// Copyright Epic Games, Inc. All Rights Reserved.
+
+using System.Collections.Immutable;
+
+namespace UnrealBuildTool.Actions;
+
+internal record class ActionResultManifest(OriginalTargetDescriptor TargetDescriptor, ImmutableArray<ActionResult> Results);
diff --git a/Engine/Source/Programs/UnrealBuildTool/Actions/ResultStore/ContentAddressableActionResultStore.cs b/Engine/Source/Programs/UnrealBuildTool/Actions/ResultStore/ContentAddressableActionResultStore.cs
--- a/Engine/Source/Programs/UnrealBuildTool/Actions/ResultStore/ContentAddressableActionResultStore.cs
+++ b/Engine/Source/Programs/UnrealBuildTool/Actions/ResultStore/ContentAddressableActionResultStore.cs
@@ -1,40 +1,103 @@
 // Copyright Epic Games, Inc. All Rights Reserved.
 
 using System;
+using System.Buffers;
+using System.Collections.Concurrent;
+using System.Collections.Generic;
+using System.Collections.Immutable;
 using System.IO;
+using System.Linq;
 using System.Text.Json;
 using System.Threading;
 using System.Threading.Tasks;
 using EpicGames.Core;
+using Microsoft.Extensions.Logging;
 using UnrealBuildTool.Storage;
 
 namespace UnrealBuildTool.Actions.ResultStore;
 
-internal class ContentAddressableActionResultStore(IStorageProvider storageProvider)
+internal class ContentAddressableActionResultStore(IStorageProvider storageProvider, ILogger logger)
 	: IActionResultStore
 {
 	/// <inheritdoc />
-	public async Task<ActionResult?> LoadResultAsync(IExternalAction action, CancellationToken cancellationToken)
+	public async Task<ActionResult?> LoadResultAsync(LinkedAction action, CancellationToken cancellationToken)
 	{
 		IoHash actionHash = GetActionHash(action);
-		using Stream? stream = storageProvider.CreateReader(actionHash);
+		return await LoadResultImplAsync(actionHash, cancellationToken);
+	}
+
+	/// <inheritdoc />
+	public async Task<ActionResultManifest?> LoadManifestAsync(OriginalTargetDescriptor targetDescriptor, CancellationToken cancellationToken)
+	{
+		IoHash targetHash = _targetToHash.GetOrAdd(targetDescriptor, GetTargetHash);
+		await using Stream? stream = storageProvider.CreateReader(targetHash);
 		if (stream is null)
 		{
 			return null;
 		}
 
-		return await JsonSerializer.DeserializeAsync<ActionResult>(stream, cancellationToken: cancellationToken);
+		PartialActionResultManifest? partialManifest = await JsonSerializer.DeserializeAsync<PartialActionResultManifest>(stream, _jsonSerializerOptions, cancellationToken);
+		if (partialManifest is null)
+		{
+			logger.LogInformation("Found empty manifest for hash {TargetHash}", targetHash);
+			return null;
+		}
+
+		ActionResult?[] actionResults = await Task.WhenAll(partialManifest.Actions.AsParallel().Select(a => LoadResultImplAsync(a, cancellationToken)));
+		IEnumerable<ActionResult> resultsWithoutNulls = actionResults.OfType<ActionResult>();
+		return new ActionResultManifest(partialManifest.Target, [.. resultsWithoutNulls]);
 	}
 
 	/// <inheritdoc />
-	public async Task StoreResultAsync(IExternalAction action, ActionResult result, CancellationToken cancellationToken)
+	public async Task StoreResultAsync(LinkedAction action, ActionResult result, CancellationToken cancellationToken)
 	{
 		IoHash actionHash = GetActionHash(action);
-		using Stream stream = storageProvider.CreateWriter(actionHash, cancellationToken);
-		await JsonSerializer.SerializeAsync(stream, result, cancellationToken: cancellationToken);
+		await using Stream stream = storageProvider.CreateWriter(actionHash, cancellationToken);
+		await JsonSerializer.SerializeAsync(stream, result, _jsonSerializerOptions, cancellationToken);
+
+		if (action.Target is null)
+		{
+			return;
+		}
+
+		IoHash targetHash = _targetToHash.GetOrAdd(action.Target.Original, GetTargetHash);
+		ManifestBuilder manifestBuilder = _targetHashToStream
+			.GetOrAdd(targetHash, h => new(() => new(storageProvider, action.Target.Original, h, cancellationToken)))
+			.Value;
+		manifestBuilder.Actions.Add(actionHash);
+	}
+
+	/// <summary>
+	/// Write out the complete manifests to the path given by the hash of the target descriptors.
+	/// </summary>
+	public async ValueTask DisposeAsync()
+	{
+		await Task.WhenAll(_targetHashToStream.Values.AsParallel().Select(async st => await st.Value.DisposeAsync()));
+	}
+
+	// Reference equality should be sufficient for our purposes, and is much faster.
+	private readonly ConcurrentDictionary<OriginalTargetDescriptor, IoHash> _targetToHash = new(ReferenceEqualityComparer.Instance);
+	// Lazy allows us to make sure only one stream actually gets constructed.
+	private readonly ConcurrentDictionary<IoHash, Lazy<ManifestBuilder>> _targetHashToStream = [];
+
+	private static readonly JsonSerializerOptions _jsonSerializerOptions = new()
+	{
+		PreferredObjectCreationHandling = System.Text.Json.Serialization.JsonObjectCreationHandling.Populate,
+		IncludeFields = true,
+	};
+
+	private async Task<ActionResult?> LoadResultImplAsync(IoHash actionHash, CancellationToken cancellationToken)
+	{
+		await using Stream? stream = storageProvider.CreateReader(actionHash);
+		if (stream is null)
+		{
+			return null;
+		}
+
+		return await JsonSerializer.DeserializeAsync<ActionResult>(stream, _jsonSerializerOptions, cancellationToken);
 	}
 
-	private IoHash GetActionHash(IExternalAction action)
+	private IoHash GetActionHash(LinkedAction action)
 	{
 		using Blake3.Hasher hasher = Blake3.Hasher.New();
 		hasher.Update(action.CommandPath.FullName.AsSpan());
@@ -42,4 +105,39 @@ internal class ContentAddressableActionResultStore(IStorageProvider storageProvi
 		hasher.Update(action.CommandVersion.AsSpan());
 		return IoHash.FromBlake3(hasher);
 	}
+
+	private IoHash GetTargetHash(OriginalTargetDescriptor targetDescriptor)
+	{
+		byte[] byteArray = ArrayPool<byte>.Shared.Rent(64 * 1024);
+		using MemoryStream memoryStream = new(byteArray);
+		JsonSerializer.Serialize(memoryStream, targetDescriptor, _jsonSerializerOptions);
+		ReadOnlySpan<byte> usedStream = byteArray.AsSpan()[..(int)memoryStream.Position];
+		return IoHash.Compute(usedStream);
+	}
+
+	private record class PartialActionResultManifest(OriginalTargetDescriptor Target, ImmutableList<IoHash> Actions);
+
+	/// <summary>
+	/// Tracks actions and serializes a PartialActionResultManifest on dispose.
+	/// </summary>
+	private class ManifestBuilder(
+		IStorageProvider storageProvider,
+		OriginalTargetDescriptor targetDescriptor,
+		IoHash targetHash,
+		CancellationToken cancellationToken
+	) : IAsyncDisposable
+	{
+		/// <summary>
+		/// Even for 50,000 actions, this should only take about 1MB of storage, so keep it simple and just write out the manifest
+		/// once all actions are known.
+		/// </summary>
+		public ConcurrentBag<IoHash> Actions { get; } = [];
+
+		public async ValueTask DisposeAsync()
+		{
+			await using Stream stream = storageProvider.CreateWriter(targetHash, cancellationToken);
+			PartialActionResultManifest manifest = new(targetDescriptor, [.. Actions]);
+			await JsonSerializer.SerializeAsync(stream, manifest, _jsonSerializerOptions, cancellationToken);
+		}
+	}
 }
diff --git a/Engine/Source/Programs/UnrealBuildTool/Actions/ResultStore/IActionResultStore.cs b/Engine/Source/Programs/UnrealBuildTool/Actions/ResultStore/IActionResultStore.cs
--- a/Engine/Source/Programs/UnrealBuildTool/Actions/ResultStore/IActionResultStore.cs
+++ b/Engine/Source/Programs/UnrealBuildTool/Actions/ResultStore/IActionResultStore.cs
@@ -1,5 +1,6 @@
 // Copyright Epic Games, Inc. All Rights Reserved.
 
+using System;
 using System.Threading;
 using System.Threading.Tasks;
 
@@ -8,15 +9,24 @@ namespace UnrealBuildTool.Actions.ResultStore;
 /// <summary>
 /// Persists the result of executing actions.
 /// </summary>
-internal interface IActionResultStore
+internal interface IActionResultStore : IAsyncDisposable
 {
 	/// <summary>
 	/// Load the result of executing a given action.
 	/// </summary>
 	/// <param name="action">The action to check results of.</param>
 	/// <param name="cancellationToken">A token to allow cancelling the operation.</param>
-	/// <returns>The result of the action, if it was previously persisted.</returns>
-	Task<ActionResult?> LoadResultAsync(IExternalAction action, CancellationToken cancellationToken = default);
+	/// <returns>The result of the action, if it was previously persisted, otherwise null.</returns>
+	Task<ActionResult?> LoadResultAsync(LinkedAction action, CancellationToken cancellationToken = default);
+
+	/// <summary>
+	/// Load an action result manifest for a given target descriptor.
+	/// This will only exist if the target descriptor has been built previously.
+	/// </summary>
+	/// <param name="targetDescriptor">The target descriptor for which to load the manifest.</param>
+	/// <param name="cancellationToken">A token to allow cancelling the operation.</param>
+	/// <returns>The action result manifest, if it exists and is valid, otherwise null.</returns>
+	Task<ActionResultManifest?> LoadManifestAsync(OriginalTargetDescriptor targetDescriptor, CancellationToken cancellationToken = default);
 
 	/// <summary>
 	/// Store the result of executing a given action.
@@ -24,6 +34,5 @@ internal interface IActionResultStore
 	/// <param name="action">The action to store results of.</param>
 	/// <param name="result">The result of the action to store.</param>
 	/// <param name="cancellationToken">A token to allow cancelling the operation.</param>
-	/// <returns></returns>
-	Task StoreResultAsync(IExternalAction action, ActionResult result, CancellationToken cancellationToken = default);
+	Task StoreResultAsync(LinkedAction action, ActionResult result, CancellationToken cancellationToken = default);
 }
diff --git a/Engine/Source/Programs/UnrealBuildTool/Actions/ResultStore/NullActionResultStore.cs b/Engine/Source/Programs/UnrealBuildTool/Actions/ResultStore/NullActionResultStore.cs
--- a/Engine/Source/Programs/UnrealBuildTool/Actions/ResultStore/NullActionResultStore.cs
+++ b/Engine/Source/Programs/UnrealBuildTool/Actions/ResultStore/NullActionResultStore.cs
@@ -7,15 +7,27 @@ namespace UnrealBuildTool.Actions.ResultStore;
 
 internal class NullActionResultStore : IActionResultStore
 {
+
 	/// <inheritdoc />
-	public Task<ActionResult?> LoadResultAsync(IExternalAction action, CancellationToken cancellationToken)
+	public Task<ActionResult?> LoadResultAsync(LinkedAction action, CancellationToken cancellationToken)
 	{
 		return Task.FromResult<ActionResult?>(null);
 	}
 
+	public Task<ActionResultManifest?> LoadManifestAsync(OriginalTargetDescriptor targetDescriptor, CancellationToken cancellationToken = default)
+	{
+		return Task.FromResult<ActionResultManifest?>(null);
+	}
+
 	/// <inheritdoc />
-	public Task StoreResultAsync(IExternalAction action, ActionResult result, CancellationToken cancellationToken)
+	public Task StoreResultAsync(LinkedAction action, ActionResult result, CancellationToken cancellationToken)
 	{
 		return Task.CompletedTask;
 	}
+
+	/// <inheritdoc />
+	public ValueTask DisposeAsync()
+	{
+		return ValueTask.CompletedTask;
+	}
 }
diff --git a/Engine/Source/Programs/UnrealBuildTool/Configuration/Descriptors/OriginalTargetDescriptor.cs b/Engine/Source/Programs/UnrealBuildTool/Configuration/Descriptors/OriginalTargetDescriptor.cs
new file mode 100644
--- /dev/null
+++ b/Engine/Source/Programs/UnrealBuildTool/Configuration/Descriptors/OriginalTargetDescriptor.cs
@@ -0,0 +1,138 @@
+// Copyright Epic Games, Inc. All Rights Reserved.
+
+using System;
+using System.Collections.Immutable;
+using System.Runtime.InteropServices;
+using System.Text.Json.Serialization;
+using EpicGames.Core;
+
+namespace UnrealBuildTool;
+
+/// <summary>
+/// A target descriptor that is read-only and suitable for serialization and equality comparison.
+/// This should be kept up-to-date with <inheritdoc cref="TargetDescriptor"/>.
+/// "Original" because it stores the state of the descriptor when it was created.
+/// </summary>
+internal record class OriginalTargetDescriptor
+{
+	/// <inheritdoc cref="TargetDescriptor.ProjectFile"/>
+	public FileReference? ProjectFile { get; init; }
+
+	/// <inheritdoc cref="TargetDescriptor.Name"/>
+	public required string Name { get; init; }
+
+	/// <inheritdoc cref="TargetDescriptor.Platform"/>
+	public UnrealTargetPlatform Platform { get; init; }
+
+	/// <inheritdoc cref="TargetDescriptor.Configuration"/>
+	public UnrealTargetConfiguration Configuration { get; init; }
+
+	/// <inheritdoc cref="TargetDescriptor.Architectures"/>
+	public required UnrealArchitectures Architectures { get; init; }
+
+	/// <inheritdoc cref="TargetDescriptor.AdditionalArguments"/>
+	public required CommandLineArguments AdditionalArguments { get; init; }
+
+	/// <inheritdoc cref="TargetDescriptor.IsTestsTarget"/>
+	public bool IsTestsTarget { get; init; }
+
+	/// <inheritdoc cref="TargetDescriptor.ForeignPlugin"/>
+	public FileReference? ForeignPlugin { get; init; }
+
+	/// <inheritdoc cref="TargetDescriptor.bBuildPluginAsLocal"/>
+	public bool bBuildPluginAsLocal { get; init; }
+
+	/// <inheritdoc cref="TargetDescriptor.bBuildDependantPlugins"/>
+	public bool bBuildDependantPlugins { get; init; }
+
+	/// <inheritdoc cref="TargetDescriptor.OnlyModuleNames"/>
+	public required ImmutableHashSet<string> OnlyModuleNames { get; init; }
+
+	/// <inheritdoc cref="TargetDescriptor.FileLists"/>
+	public ImmutableArray<FileReference> FileLists { get; init; }
+
+	/// <inheritdoc cref="TargetDescriptor.SpecificFilesToCompile"/>
+	public ImmutableArray<FileReference> SpecificFilesToCompile { get; init; }
+
+	/// <inheritdoc cref="TargetDescriptor.RelativePathsToSpecificFilesToCompile"/>
+	public ImmutableArray<string> RelativePathsToSpecificFilesToCompile { get; init; }
+
+	/// <inheritdoc cref="TargetDescriptor.WorkingDir"/>
+	public string? WorkingDir { get; init; }
+
+	/// <inheritdoc cref="TargetDescriptor.bSingleFileBuildDependents"/>
+	public bool bSingleFileBuildDependents { get; init; }
+
+	/// <inheritdoc cref="TargetDescriptor.HotReloadMode"/>
+	public HotReloadMode HotReloadMode { get; set; }
+
+	/// <inheritdoc cref="TargetDescriptor.HotReloadModuleNameToSuffix"/>
+	public required ImmutableDictionary<string, int> HotReloadModuleNameToSuffix { get; init; }
+
+	/// <inheritdoc cref="TargetDescriptor.LiveCodingModules"/>
+	public FileReference? LiveCodingModules { get; init; }
+
+	/// <inheritdoc cref="TargetDescriptor.LiveCodingManifest"/>
+	public FileReference? LiveCodingManifest { get; init; }
+
+	/// <inheritdoc cref="TargetDescriptor.LiveCodingLimit"/>
+	public uint LiveCodingLimit { get; init; }
+
+	/// <inheritdoc cref="TargetDescriptor.bQuiet"/>
+	[JsonIgnore]
+	public bool bQuiet { get; init; }
+
+	/// <inheritdoc cref="TargetDescriptor.bRebuild"/>
+	[JsonIgnore]
+	public bool bRebuild { get; init; }
+
+	/// <inheritdoc cref="TargetDescriptor.bUseUnityBuild"/>
+	public bool bUseUnityBuild { get; init; }
+
+	/// <inheritdoc cref="TargetDescriptor.bForceUnityBuild"/>
+	public bool bForceUnityBuild { get; init; }
+
+	/// <inheritdoc cref="TargetDescriptor.bIWYU"/>
+	public bool bIWYU { get; init; }
+
+	/// <inheritdoc cref="TargetDescriptor.IntermediateEnvironment"/>
+	public UnrealIntermediateEnvironment IntermediateEnvironment { get; init; }
+
+	public static OriginalTargetDescriptor From(TargetDescriptor targetDescriptor)
+	{
+		Span<FileReference> fileLists = CollectionsMarshal.AsSpan(targetDescriptor.FileLists);
+		Span<FileReference> specificFilesToCompile = CollectionsMarshal.AsSpan(targetDescriptor.SpecificFilesToCompile);
+		Span<string> relativePathsToSpecificFiles = CollectionsMarshal.AsSpan(targetDescriptor.RelativePathsToSpecificFilesToCompile);
+		return new()
+		{
+			ProjectFile = targetDescriptor.ProjectFile,
+			Name = targetDescriptor.Name,
+			Platform = targetDescriptor.Platform,
+			Architectures = new(targetDescriptor.Architectures),
+			AdditionalArguments = targetDescriptor.AdditionalArguments,
+			IsTestsTarget = targetDescriptor.IsTestsTarget,
+			ForeignPlugin = targetDescriptor.ForeignPlugin,
+			bBuildPluginAsLocal = targetDescriptor.bBuildPluginAsLocal,
+			bBuildDependantPlugins = targetDescriptor.bBuildDependantPlugins,
+			OnlyModuleNames = ImmutableHashSet.Create(targetDescriptor.OnlyModuleNames.Comparer, [.. targetDescriptor.OnlyModuleNames]),
+			FileLists = ImmutableArray.Create(fileLists),
+			SpecificFilesToCompile = ImmutableArray.Create(specificFilesToCompile),
+			RelativePathsToSpecificFilesToCompile = ImmutableArray.Create(relativePathsToSpecificFiles),
+			WorkingDir = targetDescriptor.WorkingDir,
+			bSingleFileBuildDependents = targetDescriptor.bSingleFileBuildDependents,
+			HotReloadMode = targetDescriptor.HotReloadMode,
+			HotReloadModuleNameToSuffix = ImmutableDictionary.CreateRange<string, int>(
+				targetDescriptor.HotReloadModuleNameToSuffix.Comparer,
+				[.. targetDescriptor.HotReloadModuleNameToSuffix]),
+			LiveCodingModules = targetDescriptor.LiveCodingModules,
+			LiveCodingManifest = targetDescriptor.LiveCodingManifest,
+			LiveCodingLimit = targetDescriptor.LiveCodingLimit,
+			bQuiet = targetDescriptor.bQuiet,
+			bRebuild = targetDescriptor.bRebuild,
+			bUseUnityBuild = targetDescriptor.bUseUnityBuild,
+			bForceUnityBuild = targetDescriptor.bForceUnityBuild,
+			bIWYU = targetDescriptor.bIWYU,
+			IntermediateEnvironment = targetDescriptor.IntermediateEnvironment,
+		};
+	}
+}
diff --git a/Engine/Source/Programs/UnrealBuildTool/Configuration/Descriptors/TargetDescriptor.cs b/Engine/Source/Programs/UnrealBuildTool/Configuration/Descriptors/TargetDescriptor.cs
--- a/Engine/Source/Programs/UnrealBuildTool/Configuration/Descriptors/TargetDescriptor.cs
+++ b/Engine/Source/Programs/UnrealBuildTool/Configuration/Descriptors/TargetDescriptor.cs
@@ -4,7 +4,6 @@ using System.Collections.Generic;
 using System.Diagnostics.CodeAnalysis;
 using System.IO;
 using System.Linq;
-using System.Text;
 using EpicGames.Core;
 using Microsoft.Extensions.Logging;
 using UnrealBuildBase;
@@ -14,17 +13,17 @@ namespace UnrealBuildTool
 	/// <summary>
 	/// Describes all of the information needed to initialize a UEBuildTarget object
 	/// </summary>
-	class TargetDescriptor
+	internal sealed record class TargetDescriptor
 	{
-		public static string TEST_TARGETS_SUFFIX = "Tests";
+		public const string TEST_TARGETS_SUFFIX = "Tests";
 
-		public FileReference? ProjectFile;
-		public string Name;
-		public UnrealTargetPlatform Platform;
-		public UnrealTargetConfiguration Configuration;
-		public UnrealArchitectures Architectures;
-		public CommandLineArguments AdditionalArguments;
-		public bool IsTestsTarget = false;
+		public FileReference? ProjectFile { get; set; }
+		public required string Name { get; set; }
+		public UnrealTargetPlatform Platform { get; set; }
+		public UnrealTargetConfiguration Configuration { get; set; }
+		public required UnrealArchitectures Architectures { get; set; }
+		public required CommandLineArguments AdditionalArguments { get; set; }
+		public bool IsTestsTarget { get; set; }
 
 		public static string GetTestedName(string Name)
 		{
@@ -39,56 +38,56 @@ namespace UnrealBuildTool
 		/// Foreign plugin to compile against this target
 		/// </summary>
 		[CommandLine("-Plugin=")]
-		public FileReference? ForeignPlugin = null;
+		public FileReference? ForeignPlugin { get; set; }
 
 		/// <summary>
 		/// Whether we should treat the ForeignPlugin argument as a local plugin for building purposes
 		/// </summary>
 		[CommandLine("-BuildPluginAsLocal")]
-		public bool bBuildPluginAsLocal = false;
+		public bool bBuildPluginAsLocal { get; set; }
 
 		/// <summary>
 		/// When building a foreign plugin, whether to build plugins it depends on as well.
 		/// </summary>
 		[CommandLine("-BuildDependantPlugins")]
-		public bool bBuildDependantPlugins = false;
+		public bool bBuildDependantPlugins { get; set; }
 
 		/// <summary>
 		/// Set of module names to compile.
 		/// </summary>
 		[CommandLine("-Module=", ListSeparator = '+')]
-		public HashSet<string> OnlyModuleNames = new HashSet<string>(StringComparer.OrdinalIgnoreCase);
+		public HashSet<string> OnlyModuleNames { get; } = new(StringComparer.OrdinalIgnoreCase);
 
 		/// <summary>
 		/// Lists of files to compile
 		/// </summary>
 		[CommandLine("-FileList=")]
-		public List<FileReference> FileLists = new List<FileReference>();
+		public List<FileReference> FileLists { get; } = [];
 
 		/// <summary>
 		/// Individual file(s) to compile
 		/// </summary>
 		[CommandLine("-File=")]
 		[CommandLine("-SingleFile=")]
-		public List<FileReference> SpecificFilesToCompile = new List<FileReference>();
+		public List<FileReference> SpecificFilesToCompile { get; } = [];
 
 		/// <summary>
 		/// Relative path to file(s) to compile
 		/// </summary>
 		[CommandLine("-Files=", ListSeparator = ';')]
-		public List<string> RelativePathsToSpecificFilesToCompile = new List<string>();
+		public List<string> RelativePathsToSpecificFilesToCompile { get; } = [];
 
 		/// <summary>
 		/// Working directory when compiling with RelativePathsToSpecificFilesToCompile
 		/// </summary>
 		[CommandLine("-WorkingDir=")]
-		public string? WorkingDir = null;
+		public string? WorkingDir { get; set; }
 
 		/// <summary>
 		/// Will build all files that directly include any of the files provided in -SingleFile, not including generated code
 		/// </summary>
 		[CommandLine("-SingleFileBuildDependents")]
-		public bool bSingleFileBuildDependents;
+		public bool bSingleFileBuildDependents { get; set; }
 
 		/// <summary>
 		/// Whether to perform hot reload for this target
@@ -96,60 +95,60 @@ namespace UnrealBuildTool
 		[CommandLine("-NoHotReload", Value = nameof(HotReloadMode.Disabled))]
 		[CommandLine("-ForceHotReload", Value = nameof(HotReloadMode.FromIDE))]
 		[CommandLine("-LiveCoding", Value = nameof(HotReloadMode.LiveCoding))]
-		public HotReloadMode HotReloadMode = HotReloadMode.Default;
+		public HotReloadMode HotReloadMode { get; set; }
 
 		/// <summary>
 		/// Map of module name to suffix for hot reloading from the editor
 		/// </summary>
-		public Dictionary<string, int> HotReloadModuleNameToSuffix = new Dictionary<string, int>(StringComparer.OrdinalIgnoreCase);
+		public Dictionary<string, int> HotReloadModuleNameToSuffix { get; } = new(StringComparer.OrdinalIgnoreCase);
 
 		/// <summary>
 		/// Path to a file containing a list of modules that may be modified for live coding.
 		/// </summary>
 		[CommandLine("-LiveCodingModules=")]
-		public FileReference? LiveCodingModules = null;
+		public FileReference? LiveCodingModules { get; set; }
 
 		/// <summary>
 		/// Path to the manifest for passing info about the output to live coding
 		/// </summary>
 		[CommandLine("-LiveCodingManifest=")]
-		public FileReference? LiveCodingManifest = null;
+		public FileReference? LiveCodingManifest { get; set; }
 
 		/// <summary>
 		/// If a non-zero value, a live coding request will be terminated if more than the given number of actions are required.
 		/// </summary>
 		[CommandLine("-LiveCodingLimit=")]
-		public uint LiveCodingLimit = 0;
+		public uint LiveCodingLimit { get; set; }
 
 		/// <summary>
 		/// Suppress messages about building this target
 		/// </summary>
 		[CommandLine("-Quiet")]
-		public bool bQuiet;
+		public bool bQuiet { get; set; }
 
 		/// <summary>
 		/// Clean the target before trying to build it
 		/// </summary>
 		[CommandLine("-Rebuild")]
-		public bool bRebuild;
+		public bool bRebuild { get; set; }
 
 		/// <summary>
 		/// Whether to unify C++ code into larger files for faster compilation.
 		/// </summary>
 		[CommandLine("-DisableUnity", Value = "false")]
-		public bool bUseUnityBuild = true;
+		public bool bUseUnityBuild { get; set; } = true;
 
 		/// <summary>
 		/// Whether to force C++ source files to be combined into larger files for faster compilation.
 		/// </summary>
 		[CommandLine("-ForceUnity")]
-		public bool bForceUnityBuild = false;
+		public bool bForceUnityBuild { get; set; }
 
 		/// <summary>
 		/// Enables "include what you use" mode.
 		/// </summary>
 		[CommandLine("-IWYU")]
-		public bool bIWYU = false;
+		public bool bIWYU { get; set; }
 
 		/// <summary>
 		/// Intermediate environment. Determines if the intermediates end up in a different folder than normal.
@@ -176,6 +175,8 @@ namespace UnrealBuildTool
 		}
 		private UnrealIntermediateEnvironment? IntermediateEnvironmentOverride;
 
+		public OriginalTargetDescriptor Original { get; private set; }
+
 		/// <summary>
 		/// Constructor
 		/// </summary>
@@ -185,6 +186,7 @@ namespace UnrealBuildTool
 		/// <param name="Configuration">Configuration to build</param>
 		/// <param name="Architectures">Architectures to build for</param>
 		/// <param name="Arguments">Other command-line arguments for the target</param>
+		[SetsRequiredMembers]
 		public TargetDescriptor(FileReference? ProjectFile, string TargetName, UnrealTargetPlatform Platform, UnrealTargetConfiguration Configuration, UnrealArchitectures Architectures, CommandLineArguments? Arguments)
 		{
 			this.ProjectFile = ProjectFile;
@@ -259,11 +261,8 @@ namespace UnrealBuildTool
 				}
 			}
 			this.AdditionalArguments = new CommandLineArguments(AdditionalArguments.ToArray());
-		}
 
-		public TargetDescriptor Copy()
-		{
-			return (TargetDescriptor)MemberwiseClone();
+			Original = OriginalTargetDescriptor.From(this);
 		}
 
 		public static TargetDescriptor FromTargetInfo(TargetInfo Info)
@@ -286,25 +285,18 @@ namespace UnrealBuildTool
 			// apply the intermediate environment from the build configuration
 			if (BuildConfiguration.IntermediateEnvironment.HasValue)
 			{
-				TargetDescriptors.ForEach(x => x.IntermediateEnvironment = BuildConfiguration.IntermediateEnvironment.Value);
+				foreach (TargetDescriptor TargetDescriptor in TargetDescriptors)
+				{
+					TargetDescriptor.IntermediateEnvironment = BuildConfiguration.IntermediateEnvironment.Value;
+					// We have to remake the original to include the corrected IntermediateEnvironment.
+					// This should be included in the original as it's still before it's been handed over to the user.
+					TargetDescriptor.Original = OriginalTargetDescriptor.From(TargetDescriptor);
+				}
 			}
 
 			return TargetDescriptors;
 		}
 
-		/// <summary>
-		/// Parse a list of target descriptors from the command line
-		/// </summary>
-		/// <param name="Arguments">Command-line arguments</param>
-		/// <param name="Logger">Logger for output</param>
-		/// <returns>List of target descriptors</returns>
-		public static List<TargetDescriptor> ParseCommandLine(CommandLineArguments Arguments, ILogger Logger)
-		{
-			List<TargetDescriptor> TargetDescriptors = new List<TargetDescriptor>();
-			ParseCommandLine(Arguments, false, false, false, TargetDescriptors, Logger);
-			return TargetDescriptors;
-		}
-
 		/// <summary>
 		/// Parse a list of target descriptors from the command line
 		/// </summary>
@@ -612,86 +604,5 @@ namespace UnrealBuildTool
 			TelemetryService.Get().SetPrimaryTargetDetails(TargetDescriptors.FirstOrDefault());
 			TargetDescriptors.ForEach(x => TelemetryService.Get().AddEndpointsFromConfig(x.ProjectFile?.Directory, Logger));
 		}
-
-		/// <summary>
-		/// Try to parse the project file from the command line
-		/// </summary>
-		/// <param name="Arguments">The command line arguments</param>
-		/// <param name="ProjectFile">The project file that was parsed</param>
-		/// <returns>True if the project file was parsed, false otherwise</returns>
-		public static bool TryParseProjectFileArgument(CommandLineArguments Arguments, [NotNullWhen(true)] out FileReference? ProjectFile)
-		{
-			FileReference? ExplicitProjectFile;
-			if (Arguments.TryGetValue("-Project=", out ExplicitProjectFile))
-			{
-				ProjectFile = ExplicitProjectFile;
-				return true;
-			}
-
-			for (int Idx = 0; Idx < Arguments.Count; Idx++)
-			{
-				if (Arguments[Idx][0] != '-' && Arguments[Idx].EndsWith(".uproject", StringComparison.OrdinalIgnoreCase))
-				{
-					Arguments.MarkAsUsed(Idx);
-					ProjectFile = new FileReference(Arguments[Idx]);
-					return true;
-				}
-			}
-
-			if (Unreal.IsProjectInstalled())
-			{
-				ProjectFile = Unreal.GetInstalledProjectFile()!;
-				return true;
-			}
-
-			ProjectFile = null;
-			return false;
-		}
-
-		/// <summary>
-		/// Format this object for the debugger
-		/// </summary>
-		/// <returns>String representation of this target descriptor</returns>
-		public override string ToString()
-		{
-			StringBuilder Result = new StringBuilder();
-			Result.AppendFormat("{0} {1} {2}", Name, Platform, Configuration);
-			Result.AppendFormat($" -Architecture={Architectures}");
-			if (ProjectFile != null)
-			{
-				Result.AppendFormat(" -Project={0}", Utils.MakePathSafeToUseWithCommandLine(ProjectFile));
-			}
-			if (AdditionalArguments != null && AdditionalArguments.Count > 0)
-			{
-				Result.AppendFormat(" {0}", AdditionalArguments);
-			}
-			return Result.ToString();
-		}
-
-		public override int GetHashCode()
-		{
-#pragma warning disable RS1024
-			return String.GetHashCode(ProjectFile?.FullName) +
-				Name.GetHashCode() +
-				Platform.GetHashCode() +
-				Configuration.GetHashCode() +
-				Architectures.GetHashCode();
-#pragma warning restore RE1024
-		}
-
-		public override bool Equals(object? Obj)
-		{
-			TargetDescriptor? Other = Obj as TargetDescriptor;
-			if (Other != null)
-			{
-				return
-					ProjectFile == Other.ProjectFile &&
-					Name == Other.Name &&
-					Platform == Other.Platform &&
-					Configuration == Other.Configuration &&
-					Architectures.Architectures.SequenceEqual(Other.Architectures.Architectures);
-			}
-			return false;
-		}
 	}
 }
diff --git a/Engine/Source/Programs/UnrealBuildTool/Executors/ActionLogger.cs b/Engine/Source/Programs/UnrealBuildTool/Executors/ActionLogger.cs
--- a/Engine/Source/Programs/UnrealBuildTool/Executors/ActionLogger.cs
+++ b/Engine/Source/Programs/UnrealBuildTool/Executors/ActionLogger.cs
@@ -20,7 +20,22 @@ namespace UnrealBuildTool
 	/// <param name="PeakMemoryUsed">Peak Memory Used</param>
 	/// <param name="Action">Action</param>
 	/// <param name="AdditionalDescription">Additional description of action</param>
-	record ExecuteResults(List<string> LogLines, int ExitCode, TimeSpan ExecutionTime, TimeSpan ProcessorTime, ulong PeakMemoryUsed, LinkedAction Action, string? AdditionalDescription = null);
+	record ExecuteResults(List<string> LogLines, int ExitCode, TimeSpan ExecutionTime, TimeSpan ProcessorTime, ulong PeakMemoryUsed, LinkedAction Action, string? AdditionalDescription = null)
+	{
+		public string GetDescription()
+		{
+			string commandDescription = Action.CommandDescription ?? Action.CommandPath.GetFileNameWithoutExtension();
+			string status = Action.bShouldOutputStatusDescription || (LogLines.Count == 0) ? Action.StatusDescription : LogLines[0];
+			string description = $"{commandDescription} {status}".Trim();
+
+			if (!String.IsNullOrEmpty(AdditionalDescription))
+			{
+				description = $"{description} {AdditionalDescription}";
+			}
+
+			return description;
+		}
+	}
 
 	class ActionLogger : IDisposable
 	{
@@ -143,32 +158,19 @@ namespace UnrealBuildTool
 		/// <summary>
 		/// Log an action that has completed
 		/// </summary>
-		/// <param name="action">Action that has completed</param>
 		/// <param name="executeTaskResult">Results of the action</param>
 		/// <param name="totalActions">Number of total actions</param>
-		public void LogAction(LinkedAction action, ExecuteResults? executeTaskResult, int totalActions)
+		private void LogAction(ExecuteResults executeTaskResult, int totalActions)
 		{
-			List<string>? logLines = executeTaskResult?.LogLines ?? null;
-			int exitCode = executeTaskResult?.ExitCode ?? Int32.MaxValue;
-			TimeSpan executionTime = executeTaskResult?.ExecutionTime ?? TimeSpan.Zero;
-			TimeSpan processorTime = executeTaskResult?.ProcessorTime ?? TimeSpan.Zero;
-			ulong peakMemory = executeTaskResult?.PeakMemoryUsed ?? 0;
-			string? additionalDescription = executeTaskResult?.AdditionalDescription;
+			List<string>? logLines = executeTaskResult.LogLines;
+			int exitCode = executeTaskResult.ExitCode;
+			TimeSpan executionTime = executeTaskResult.ExecutionTime;
+			TimeSpan processorTime = executeTaskResult.ProcessorTime;
+			ulong peakMemory = executeTaskResult.PeakMemoryUsed;
 
 			// Write it to the log
-			string description = String.Empty;
-			if (action.bShouldOutputStatusDescription || (logLines != null && logLines.Count == 0))
-			{
-				description = $"{(action.CommandDescription ?? action.CommandPath.GetFileNameWithoutExtension())} {action.StatusDescription}".Trim();
-			}
-			else if (logLines != null && logLines.Count > 0)
-			{
-				description = $"{(action.CommandDescription ?? action.CommandPath.GetFileNameWithoutExtension())} {logLines[0]}".Trim();
-			}
-			if (!String.IsNullOrEmpty(additionalDescription))
-			{
-				description = $"{description} {additionalDescription}";
-			}
+			string description = executeTaskResult.GetDescription();
+			LinkedAction action = executeTaskResult.Action;
 
 			lock (ProgressWriter)
 			{
@@ -431,7 +433,7 @@ namespace UnrealBuildTool
 				foreach (int index in actionsToLog)
 				{
 					ExecuteResults result = Results[index]!;
-					LogAction(result.Action, result, Results.Length);
+					LogAction(result, Results.Length);
 				}
 			}
 		}
diff --git a/Engine/Source/Programs/UnrealBuildTool/Executors/UnrealBuildAccelerator/UBAExecutor.cs b/Engine/Source/Programs/UnrealBuildTool/Executors/UnrealBuildAccelerator/UBAExecutor.cs
--- a/Engine/Source/Programs/UnrealBuildTool/Executors/UnrealBuildAccelerator/UBAExecutor.cs
+++ b/Engine/Source/Programs/UnrealBuildTool/Executors/UnrealBuildAccelerator/UBAExecutor.cs
@@ -30,7 +30,6 @@ namespace UnrealBuildTool
 		readonly Task _initTask;
 		readonly CancellationTokenSource _cancellationSource;
 		readonly bool _bCompactOutput;
-		readonly IActionResultStore _resultStore;
 
 		IServer? _server;
 		IStorageServer? _storageServer;
@@ -171,8 +170,6 @@ namespace UnrealBuildTool
 			_cancellationSource = new();
 			_initTask = new(() => { Init(targetDescriptors, logger); }, TaskCreationOptions.LongRunning);
 			_initTask.Start();
-
-			_resultStore = new ContentAddressableActionResultStore(new FileSystemStorageProvider());
 		}
 
 		private void Init(IEnumerable<TargetDescriptor> targetDescriptors, ILogger logger)
@@ -603,6 +600,9 @@ namespace UnrealBuildTool
 			}
 			finally
 			{
+				_cancellationSource.Cancel();
+				_agentCoordinators.ForEach(ac => ac.Stop());
+
 				DisposeCacheClients();
 
 				lock (_sessionLock)
@@ -779,6 +779,9 @@ namespace UnrealBuildTool
 
 			using IScheduler scheduler = IScheduler.CreateScheduler(session, cache, NumParallelProcesses, UBAConfig.bForceBuildAllRemote);
 
+			// This must not outlive the cancellation token source being cancelled in order for it to correctly save results.
+			await using IActionResultStore resultStore = new ContentAddressableActionResultStore(new FileSystemStorageProvider(), logger);
+
 			int totalActions = 0;
 			int cacheTotalActions = 0;
 			int cacheHitActions = 0;
@@ -990,7 +993,7 @@ namespace UnrealBuildTool
 
 				ExecuteResults result = new(logLines, exitCode, executionTime, processorTime, peakMemory, action, additionalDescription);
 				actionLogger.AddActionToLog(result);
-				resultStoreTasks.Add(Task.Run(() => _resultStore.StoreResultAsync(action, ActionResult.From(result), _cancellationSource.Token)));
+				resultStoreTasks.Add(Task.Run(() => resultStore.StoreResultAsync(action, ActionResult.From(result), _cancellationSource.Token)));
 
 				return ProcessFinishedResponse.None;
 			});
@@ -1122,9 +1125,6 @@ namespace UnrealBuildTool
 					Console.CancelKeyPress -= cancelHandler;
 					cancelHandler = null;
 				}
-
-				_cancellationSource.Cancel();
-				_agentCoordinators.ForEach(ac => ac.Stop());
 			}
 		}
 
diff --git a/Engine/Source/Programs/UnrealBuildTool/Modes/ReportBuildResultsMode.cs b/Engine/Source/Programs/UnrealBuildTool/Modes/ReportBuildResultsMode.cs
new file mode 100644
--- /dev/null
+++ b/Engine/Source/Programs/UnrealBuildTool/Modes/ReportBuildResultsMode.cs
@@ -0,0 +1,51 @@
+// Copyright Epic Games, Inc. All Rights Reserved.
+
+using System;
+using System.Collections.Generic;
+using System.IO;
+using System.Text.Json;
+using System.Threading;
+using System.Threading.Tasks;
+using EpicGames.Core;
+using Microsoft.Extensions.Logging;
+using UnrealBuildTool.Actions;
+using UnrealBuildTool.Actions.ResultStore;
+using UnrealBuildTool.Storage.Impl;
+
+namespace UnrealBuildTool.Modes;
+
+[ToolMode("ReportBuildResults", ToolModeOptions.XmlConfig | ToolModeOptions.BuildPlatforms)]
+internal class ReportBuildResultsMode : ToolMode
+{
+	public override async Task<int> ExecuteAsync(CommandLineArguments arguments, ILogger logger)
+	{
+		BuildConfiguration buildConfiguration = new();
+		XmlConfig.ApplyTo(buildConfiguration);
+		arguments.ApplyTo(buildConfiguration);
+
+		List<TargetDescriptor> targetDescriptors = TargetDescriptor.ParseCommandLine(arguments, buildConfiguration, logger);
+
+		await using ContentAddressableActionResultStore resultStore = new(new FileSystemStorageProvider(), logger);
+
+		JsonSerializerOptions options = new JsonSerializerOptions()
+		{
+			WriteIndented = true
+		};
+
+		foreach (TargetDescriptor targetDescriptor in targetDescriptors)
+		{
+			ActionResultManifest? manifest = await resultStore.LoadManifestAsync(targetDescriptor.Original, CancellationToken.None);
+			if (manifest is null)
+			{
+				logger.LogInformation("Manifest for target '{Target}' is not available", targetDescriptor);
+				continue;
+			}
+
+			await using Stream stream = Console.OpenStandardOutput();
+			await JsonSerializer.SerializeAsync(stream, manifest, options);
+			Console.WriteLine();
+		}
+
+		return 0;
+	}
+}
