Fix issue CA1508 (Avoid dead conditional code) in UBT.

#Synced-CL 49202056

diff --git a/Engine/Source/Programs/UnrealBuildTool/Configuration/Ini/ConfigValueParser.cs b/Engine/Source/Programs/UnrealBuildTool/Configuration/Ini/ConfigValueParser.cs
--- a/Engine/Source/Programs/UnrealBuildTool/Configuration/Ini/ConfigValueParser.cs
+++ b/Engine/Source/Programs/UnrealBuildTool/Configuration/Ini/ConfigValueParser.cs
@@ -263,8 +263,8 @@ namespace UnrealBuildTool
 		{
 			if (TryParseTyped(typeof(T), Line, out object? ObjResult) && ObjResult != null)
 			{
-				Result = (T)ObjResult!;
-				return (Result != null);
+				Result = (T)ObjResult;
+				return true;
 			}
 
 			Result = default;
diff --git a/Engine/Source/Programs/UnrealBuildTool/Configuration/Rules/TargetRules.cs b/Engine/Source/Programs/UnrealBuildTool/Configuration/Rules/TargetRules.cs
--- a/Engine/Source/Programs/UnrealBuildTool/Configuration/Rules/TargetRules.cs
+++ b/Engine/Source/Programs/UnrealBuildTool/Configuration/Rules/TargetRules.cs
@@ -3186,7 +3186,7 @@ namespace UnrealBuildTool
 				IDictionary<string, string> dictionary = (IDictionary<string, string>)value;
 				return String.Join(',', dictionary.Select(x => $"{x.Key}:{x.Value}"));
 			}
-			return value?.ToString() ?? String.Empty;
+			return value.ToString() ?? String.Empty;
 		}
 
 		/// <summary>
diff --git a/Engine/Source/Programs/UnrealBuildTool/Configuration/UEBuildModule.cs b/Engine/Source/Programs/UnrealBuildTool/Configuration/UEBuildModule.cs
--- a/Engine/Source/Programs/UnrealBuildTool/Configuration/UEBuildModule.cs
+++ b/Engine/Source/Programs/UnrealBuildTool/Configuration/UEBuildModule.cs
@@ -1244,11 +1244,11 @@ namespace UnrealBuildTool
 				{
 					if (Rules.CppCompileWarningSettings.PluginModuleUnsupportedWarningLevel == WarningLevel.Warning)
 					{
-						Logger.LogWarning("Warning: Referenced plugin '{Plugin}' module '{Module}' not supported '{InvalidReason}' via {ReferenceChain}", Rules.Plugin?.File.GetFileName(), ThisRefName, InvalidReason, ReferenceChain);
+						Logger.LogWarning("Warning: Referenced plugin '{Plugin}' module '{Module}' not supported '{InvalidReason}' via {ReferenceChain}", Rules.Plugin.File.GetFileName(), ThisRefName, InvalidReason, ReferenceChain);
 					}
 					else if (Rules.CppCompileWarningSettings.PluginModuleUnsupportedWarningLevel == WarningLevel.Error)
 					{
-						Logger.LogError("Error: Referenced plugin '{Plugin}' module '{Module}' not supported '{InvalidReason}' via {ReferenceChain}", Rules.Plugin?.File.GetFileName(), ThisRefName, InvalidReason, ReferenceChain);
+						Logger.LogError("Error: Referenced plugin '{Plugin}' module '{Module}' not supported '{InvalidReason}' via {ReferenceChain}", Rules.Plugin.File.GetFileName(), ThisRefName, InvalidReason, ReferenceChain);
 					}
 				}
 			}
diff --git a/Engine/Source/Programs/UnrealBuildTool/Configuration/UEBuildTarget.cs b/Engine/Source/Programs/UnrealBuildTool/Configuration/UEBuildTarget.cs
--- a/Engine/Source/Programs/UnrealBuildTool/Configuration/UEBuildTarget.cs
+++ b/Engine/Source/Programs/UnrealBuildTool/Configuration/UEBuildTarget.cs
@@ -2937,12 +2937,7 @@ namespace UnrealBuildTool
 #if __VPROJECT_AVAILABLE__
 					if (!rerunCodeGen)
 					{
-						HashSet<string>? lastRunVNIModules = null;
-
-						if (lastRunVNIModules == null)
-						{
-							lastRunVNIModules = new(lastRunContext.Makefile.VNIModules.Select(m => m.ModuleName));
-						}
+						HashSet<string>? lastRunVNIModules = new(lastRunContext.Makefile.VNIModules.Select(m => m.ModuleName));
 
 						foreach (VNIModuleInfo module in makefile.VNIModules)
 						{
diff --git a/Engine/Source/Programs/UnrealBuildTool/Executors/SNDBS.cs b/Engine/Source/Programs/UnrealBuildTool/Executors/SNDBS.cs
--- a/Engine/Source/Programs/UnrealBuildTool/Executors/SNDBS.cs
+++ b/Engine/Source/Programs/UnrealBuildTool/Executors/SNDBS.cs
@@ -385,21 +385,16 @@ namespace UnrealBuildTool
 					DateTime startTimeUTC = DateTime.UtcNow;
 					Process LocalProcess = new Process();
 					LocalProcess.StartInfo = SnDbsStartInfo;
-					bool bShouldRedirectOuput = EventHandlerWrapper != null;
-					if (bShouldRedirectOuput)
-					{
-						SnDbsStartInfo.RedirectStandardError = true;
-						SnDbsStartInfo.RedirectStandardOutput = true;
-						LocalProcess.EnableRaisingEvents = true;
-						LocalProcess.OutputDataReceived += EventHandlerWrapper;
-						LocalProcess.ErrorDataReceived += EventHandlerWrapper;
-					}
+					SnDbsStartInfo.RedirectStandardError = true;
+					SnDbsStartInfo.RedirectStandardOutput = true;
+					LocalProcess.EnableRaisingEvents = true;
+					LocalProcess.OutputDataReceived += EventHandlerWrapper;
+					LocalProcess.ErrorDataReceived += EventHandlerWrapper;
+
 					LocalProcess.Start();
-					if (bShouldRedirectOuput)
-					{
-						LocalProcess.BeginOutputReadLine();
-						LocalProcess.BeginErrorReadLine();
-					}
+
+					LocalProcess.BeginOutputReadLine();
+					LocalProcess.BeginErrorReadLine();
 
 					Logger.LogInformation("Distributing {NumAction} action{ActionS} to SN-DBS",
 						NumActions,
diff --git a/Engine/Source/Programs/UnrealBuildTool/Modes/ServerMode.cs b/Engine/Source/Programs/UnrealBuildTool/Modes/ServerMode.cs
--- a/Engine/Source/Programs/UnrealBuildTool/Modes/ServerMode.cs
+++ b/Engine/Source/Programs/UnrealBuildTool/Modes/ServerMode.cs
@@ -188,14 +188,14 @@ namespace UnrealBuildTool
 			};
 			try
 			{
-				Command? MaybeCommand = JsonSerializer.Deserialize<Command>(CommandString, jsonOptions);
-				if (MaybeCommand is null)
-				{
-					Console.Error.WriteLine($"Failed to parse command {CommandString}");
-					return;
-				}
+				Command command = JsonSerializer.Deserialize<Command>(CommandString, jsonOptions);
+				// Unclear if 0 is a valid sequence number - if it is then there's no way to determine a blank JSON input from an intentional default(Command).
+				//if (command is null)
+				//{
+				//	Console.Error.WriteLine($"Failed to parse command {CommandString}");
+				//	return;
+				//}
 				Logger.LogInformation($"Got command {CommandString}");
-				Command command = MaybeCommand.Value;
 				switch (command.Type)
 				{
 					case CommandType.Quit:
diff --git a/Engine/Source/Programs/UnrealBuildTool/Modes/TestMode.cs b/Engine/Source/Programs/UnrealBuildTool/Modes/TestMode.cs
--- a/Engine/Source/Programs/UnrealBuildTool/Modes/TestMode.cs
+++ b/Engine/Source/Programs/UnrealBuildTool/Modes/TestMode.cs
@@ -215,7 +215,7 @@ namespace UnrealBuildTool
 						}
 					}
 					FieldInfo? TestMetadataField = FindFieldUpwards(TestModuleType, "TestMetadata", BindingFlags.Static | BindingFlags.NonPublic);
-					if (UpdateBuildGraphMetadata != null && TestMetadataField != null)
+					if (TestMetadataField != null)
 					{
 						TestModuleRules.Metadata MetadataValue = (TestModuleRules.Metadata)TestMetadataField.GetValue(null)!;
 						if (MetadataValue != null && !String.IsNullOrEmpty(MetadataValue.TestName))
diff --git a/Engine/Source/Programs/UnrealBuildTool/Platform/Android/UEDeployAndroid.cs b/Engine/Source/Programs/UnrealBuildTool/Platform/Android/UEDeployAndroid.cs
--- a/Engine/Source/Programs/UnrealBuildTool/Platform/Android/UEDeployAndroid.cs
+++ b/Engine/Source/Programs/UnrealBuildTool/Platform/Android/UEDeployAndroid.cs
@@ -3778,12 +3778,6 @@ namespace UnrealBuildTool
 								byte[] CompressedBuffer = CompressedStream.GetBuffer();
 								CompressedSize = (uint)CompressedStream.Length;
 
-								// calculate location after header and any additional alignment needed (only if uncompressed)
-								uint HeaderSize = (uint)(4 + 2 + 2 + 2 + 2 + 2 + 4 + 4 + 4 + 2 + 2 + FilenameLen + ExtraLen);
-								uint HeaderEnd = Location + HeaderSize;
-								uint DataStart = (uint)((HeaderEnd + 3) & ~3);
-								ushort ExtraAlign = (ushort)(Compression != 0 ? 0 : (DataStart - HeaderEnd));
-
 								// write local header
 								DestWriter.Write(Header);
 								DestWriter.Write(Version);
@@ -3795,13 +3789,9 @@ namespace UnrealBuildTool
 								DestWriter.Write(CompressedSize);
 								DestWriter.Write(UncompressedSize);
 								DestWriter.Write(FilenameLen);
-								DestWriter.Write((ushort)(ExtraLen + ExtraAlign));
+								DestWriter.Write(ExtraLen);
 								DestWriter.BaseStream.Write(Filename, 0, FilenameLen);
 								DestWriter.BaseStream.Write(Extra, 0, ExtraLen);
-								while (ExtraAlign-- > 0)
-								{
-									DestWriter.Write((byte)0);
-								}
 
 								// write compressed data
 								DestWriter.BaseStream.Write(CompressedBuffer, 0, (int)CompressedSize);
diff --git a/Engine/Source/Programs/UnrealBuildTool/Platform/Apple/AppleExports.cs b/Engine/Source/Programs/UnrealBuildTool/Platform/Apple/AppleExports.cs
--- a/Engine/Source/Programs/UnrealBuildTool/Platform/Apple/AppleExports.cs
+++ b/Engine/Source/Programs/UnrealBuildTool/Platform/Apple/AppleExports.cs
@@ -399,13 +399,13 @@ namespace UnrealBuildTool
 				}
 
 				_MobileProvisionDir = new DirectoryReference(Environment.GetEnvironmentVariable("HOME") + MobileProvisionDir);
-				Console.WriteLine(" Setting Mobile Provision Profile Dir: {0}", _MobileProvisionDir?.FullName ?? "Not set");
+				Console.WriteLine(" Setting Mobile Provision Profile Dir: {0}", _MobileProvisionDir.FullName ?? "Not set");
 				return _MobileProvisionDir!;
 			}
 			else
 			{
 				_MobileProvisionDir = new DirectoryReference(Environment.GetFolderPath(Environment.SpecialFolder.LocalApplicationData) + "/Apple Computer/MobileDevice/Provisioning Profiles/");
-				Console.WriteLine(" Setting Mobile Provision Profile Dir: {0}", _MobileProvisionDir?.FullName ?? "Not set");
+				Console.WriteLine(" Setting Mobile Provision Profile Dir: {0}", _MobileProvisionDir.FullName ?? "Not set");
 				return _MobileProvisionDir!;
 			}
 		}
diff --git a/Engine/Source/Programs/UnrealBuildTool/Platform/IOS/IOSCodeSigning.cs b/Engine/Source/Programs/UnrealBuildTool/Platform/IOS/IOSCodeSigning.cs
--- a/Engine/Source/Programs/UnrealBuildTool/Platform/IOS/IOSCodeSigning.cs
+++ b/Engine/Source/Programs/UnrealBuildTool/Platform/IOS/IOSCodeSigning.cs
@@ -436,11 +436,6 @@ namespace UnrealBuildTool
 					Console.WriteLine(ex.Message);
 				}
 
-				if (Info == null)
-				{
-					return false;
-				}
-
 				// Get the name of the bundle
 				Info.GetString("CFBundleIdentifier", out CFBundleIdentifier);
 				if (CFBundleIdentifier == null)
diff --git a/Engine/Source/Programs/UnrealBuildTool/Platform/IOS/PListHelper.cs b/Engine/Source/Programs/UnrealBuildTool/Platform/IOS/PListHelper.cs
--- a/Engine/Source/Programs/UnrealBuildTool/Platform/IOS/PListHelper.cs
+++ b/Engine/Source/Programs/UnrealBuildTool/Platform/IOS/PListHelper.cs
@@ -445,7 +445,7 @@ namespace UnrealBuildTool
 				if (ValueNode != null)
 				{
 					ValueNode.RemoveAll();
-					ValueNode?.ParentNode?.RemoveChild(ValueNode);
+					ValueNode.ParentNode?.RemoveChild(ValueNode);
 				}
 
 				//remove key
diff --git a/Engine/Source/Programs/UnrealBuildTool/Platform/Windows/UEBuildWindows.cs b/Engine/Source/Programs/UnrealBuildTool/Platform/Windows/UEBuildWindows.cs
--- a/Engine/Source/Programs/UnrealBuildTool/Platform/Windows/UEBuildWindows.cs
+++ b/Engine/Source/Programs/UnrealBuildTool/Platform/Windows/UEBuildWindows.cs
@@ -1561,8 +1561,8 @@ namespace UnrealBuildTool
 			IEnumerable<WindowsCompiler> PreferredCompilers = GetPreferredCompilers(TargetRules);
 			if (PreferredCompilers.Any())
 			{
-				WindowsCompiler? Compiler = PreferredCompilers.FirstOrDefault(x => MicrosoftPlatformSDK.HasValidCompiler(x, Architecture, Logger));
-				if (Compiler != null)
+				WindowsCompiler Compiler = PreferredCompilers.FirstOrDefault(x => MicrosoftPlatformSDK.HasValidCompiler(x, Architecture, Logger));
+				if (Compiler != default)
 				{
 					if (Compiler != PreferredCompilers.First())
 					{
@@ -1572,7 +1572,7 @@ namespace UnrealBuildTool
 					{
 						Logger.LogDebug("{Target} {Platform} {Architecture} Compiler set to preferred compiler {Compiler}", TargetRules?.Name ?? "Default", TargetRules?.Platform ?? UnrealTargetPlatform.Win64, Architecture, Compiler);
 					}
-					return (WindowsCompiler)Compiler;
+					return Compiler;
 				}
 				Logger.Log(bSkipWarning || MicrosoftPlatformSDK.IgnoreToolchainErrors ? LogLevel.Debug : LogLevel.Warning, "{Target} {Platform} {Architecture} No preferred compilers available [{Preferred}]", TargetRules?.Name ?? "Default", TargetRules?.Platform ?? UnrealTargetPlatform.Win64, Architecture, String.Join(',', PreferredCompilers));
 			}
diff --git a/Engine/Source/Programs/UnrealBuildTool/System/NativeProjects.cs b/Engine/Source/Programs/UnrealBuildTool/System/NativeProjects.cs
--- a/Engine/Source/Programs/UnrealBuildTool/System/NativeProjects.cs
+++ b/Engine/Source/Programs/UnrealBuildTool/System/NativeProjects.cs
@@ -43,7 +43,9 @@ namespace UnrealBuildTool
 			{
 				lock (LockObject)
 				{
+#pragma warning disable CA1508 // Avoid dead conditional code - required to check for null again after locking
 					if (CachedTargetNameToProjectFile == null)
+#pragma warning restore CA1508
 					{
 						using ITimelineEvent _ = Timeline.ScopeEvent("Find target files");
 						Dictionary<string, FileReference> TargetNameToProjectFile = [];
diff --git a/Engine/Source/Programs/UnrealBuildTool/System/SourceFileSearch.cs b/Engine/Source/Programs/UnrealBuildTool/System/SourceFileSearch.cs
--- a/Engine/Source/Programs/UnrealBuildTool/System/SourceFileSearch.cs
+++ b/Engine/Source/Programs/UnrealBuildTool/System/SourceFileSearch.cs
@@ -123,7 +123,7 @@ namespace UnrealBuildTool
 			{
 				lock (SearchedDirectories)
 				{
-					SearchedDirectories?.Add(Directory.Location);
+					SearchedDirectories.Add(Directory.Location);
 				}
 			}
 			else
