[UBA]
* Backported -UBASharedMemoryTempFile to be able to use on farm

#Synced-CL 49291772

diff --git a/Engine/Source/Programs/UnrealBuildTool/Executors/UnrealBuildAccelerator/UBAExecutor.cs b/Engine/Source/Programs/UnrealBuildTool/Executors/UnrealBuildAccelerator/UBAExecutor.cs
--- a/Engine/Source/Programs/UnrealBuildTool/Executors/UnrealBuildAccelerator/UBAExecutor.cs
+++ b/Engine/Source/Programs/UnrealBuildTool/Executors/UnrealBuildAccelerator/UBAExecutor.cs
@@ -283,6 +283,11 @@ namespace UnrealBuildTool
 				config.AddValue("Storage", "CompressionLevel", UBAConfig.CompressionLevel);
 			}
 
+			if (!String.IsNullOrEmpty(UBAConfig.SharedMemoryTempFile))
+			{
+				config.AddValue("Storage", "SharedMemoryTempFile", UBAConfig.SharedMemoryTempFile);
+			}
+
 			if (!UBAConfig.bCleanupOnDispose)
 			{
 				config.AddValue("Storage", "SkipCleanupOnShutdown", "true"); // For faster shutdown
