Fix UBT generating invalid VSCode launch configs for UBT on Linux/Mac

#13652

#Synced-CL 49326952

diff --git a/Engine/Source/Programs/UnrealBuildTool/ProjectFiles/VisualStudioCode/VSCodeProjectFileGenerator.cs b/Engine/Source/Programs/UnrealBuildTool/ProjectFiles/VisualStudioCode/VSCodeProjectFileGenerator.cs
--- a/Engine/Source/Programs/UnrealBuildTool/ProjectFiles/VisualStudioCode/VSCodeProjectFileGenerator.cs
+++ b/Engine/Source/Programs/UnrealBuildTool/ProjectFiles/VisualStudioCode/VSCodeProjectFileGenerator.cs
@@ -1599,7 +1599,10 @@ namespace UnrealBuildTool
 				PreLaunchTask = "UnrealBuildTool " + HostPlatform.ToString() + " Development Build";
 			}
 
-			FileReference RunUbtPath = FileReference.Combine(ProjectRoot, "Engine", "Build", "BatchFiles", "RunUBT.bat");
+			FileReference RunUbtPath;
+			string ScriptExt = HostPlatform == UnrealTargetPlatform.Win64 ? ".bat" : ".sh";
+			RunUbtPath = FileReference.Combine(ProjectRoot, "Engine", "Build", "BatchFiles", $"RunUBT{ScriptExt}");
+
 			WriteSingleCSharpLaunchConfig(
 				OutFile,
 				"Generate Project Files",
