Clean up the UBT project file by removing some redundant properties.

#Synced-CL 49326952

diff --git a/Engine/Source/Programs/UnrealBuildTool/UnrealBuildTool.csproj b/Engine/Source/Programs/UnrealBuildTool/UnrealBuildTool.csproj
--- a/Engine/Source/Programs/UnrealBuildTool/UnrealBuildTool.csproj
+++ b/Engine/Source/Programs/UnrealBuildTool/UnrealBuildTool.csproj
@@ -1,37 +1,19 @@
 ﻿<Project Sdk="Microsoft.NET.Sdk">
   <Import Project="..\Shared\UnrealEngine.csproj.props" />
-  
+
   <PropertyGroup>
     <TargetFramework>net8.0</TargetFramework>
-    <Configuration Condition=" '$(Configuration)' == '' ">Development</Configuration>
     <OutputType>Exe</OutputType>
     <AppendTargetFrameworkToOutputPath>false</AppendTargetFrameworkToOutputPath>
     <GenerateAssemblyInfo>false</GenerateAssemblyInfo>
     <GenerateTargetFrameworkAttribute>false</GenerateTargetFrameworkAttribute>
-    <Configurations>Debug;Release;Development;Analyze</Configurations>
     <ApplicationManifest>app.manifest</ApplicationManifest>
-    <AssemblyName>UnrealBuildTool</AssemblyName>
-    <RootNamespace>UnrealBuildTool</RootNamespace>
     <SatelliteResourceLanguages>en</SatelliteResourceLanguages> <!-- remove non english resource languages -->
     <AppendRuntimeIdentifierToOutputPath>false</AppendRuntimeIdentifierToOutputPath>
     <Nullable>enable</Nullable>
-    <RunAnalyzersDuringBuild>false</RunAnalyzersDuringBuild>
     <GenerateDocumentationFile>True</GenerateDocumentationFile>
   </PropertyGroup>
-  <PropertyGroup Condition=" '$(Configuration)|$(Platform)' == 'Development|AnyCPU' ">
-    <Optimize>True</Optimize>
-  </PropertyGroup>
-  <PropertyGroup Condition=" '$(Configuration)|$(Platform)' == 'Release|AnyCPU' ">
-    <Optimize>True</Optimize>
-  </PropertyGroup>
-  <PropertyGroup Condition=" '$(Configuration)|$(Platform)' == 'Debug|AnyCPU' ">
-    <DebugSymbols>true</DebugSymbols>
-  </PropertyGroup>
-  <PropertyGroup Condition=" '$(Configuration)|$(Platform)' == 'Analyze|AnyCPU' ">
-    <EnforceCodeStyleInBuild>true</EnforceCodeStyleInBuild>
-    <RunAnalyzersDuringBuild>True</RunAnalyzersDuringBuild>
-    <TreatWarningsAsErrors>false</TreatWarningsAsErrors>
-  </PropertyGroup>
+
   <ItemGroup>
     <Reference Include="Ionic.Zip.Reduced">
       <HintPath>..\..\..\Binaries\DotNET\Ionic.Zip.Reduced.dll</HintPath>
@@ -95,6 +77,7 @@
     <ProjectReference Include="..\Shared\EpicGames.UBA\EpicGames.UBA.csproj" />
     <ProjectReference Include="..\Shared\EpicGames.UHT\EpicGames.UHT.csproj" />
   </ItemGroup>
+
   <ItemGroup>
     <PackageReference Include="Microsoft.CodeAnalysis.CSharp" Version="4.13.0" />
     <PackageReference Include="Microsoft.Extensions.Logging.Console" Version="8.0.1" />
