Simplify hashing target descriptors - we cache the hashes, and builds typically have very few target descriptors, so we won't get too many allocations in normal usage anyway.

#Synced-CL 49329310

diff --git a/Engine/Source/Programs/UnrealBuildTool/Actions/ResultStore/ContentAddressableActionResultStore.cs b/Engine/Source/Programs/UnrealBuildTool/Actions/ResultStore/ContentAddressableActionResultStore.cs
--- a/Engine/Source/Programs/UnrealBuildTool/Actions/ResultStore/ContentAddressableActionResultStore.cs
+++ b/Engine/Source/Programs/UnrealBuildTool/Actions/ResultStore/ContentAddressableActionResultStore.cs
@@ -108,10 +108,10 @@ internal class ContentAddressableActionResultStore(IStorageProvider storageProvi
 
 	private IoHash GetTargetHash(OriginalTargetDescriptor targetDescriptor)
 	{
-		byte[] byteArray = ArrayPool<byte>.Shared.Rent(64 * 1024);
-		using MemoryStream memoryStream = new(byteArray);
+		const int byteEstimateForSmallTargetDescriptor = 2048;
+		using MemoryStream memoryStream = new MemoryStream(byteEstimateForSmallTargetDescriptor);
 		JsonSerializer.Serialize(memoryStream, targetDescriptor, _jsonSerializerOptions);
-		ReadOnlySpan<byte> usedStream = byteArray.AsSpan()[..(int)memoryStream.Position];
+		ReadOnlySpan<byte> usedStream = memoryStream.GetBuffer().AsSpan()[..(int)memoryStream.Position];
 		return IoHash.Compute(usedStream);
 	}
 
