Remove unused using in UBT to fix analysis warning.

#Synced-CL 49329310

diff --git a/Engine/Source/Programs/UnrealBuildTool/Actions/ResultStore/ContentAddressableActionResultStore.cs b/Engine/Source/Programs/UnrealBuildTool/Actions/ResultStore/ContentAddressableActionResultStore.cs
--- a/Engine/Source/Programs/UnrealBuildTool/Actions/ResultStore/ContentAddressableActionResultStore.cs
+++ b/Engine/Source/Programs/UnrealBuildTool/Actions/ResultStore/ContentAddressableActionResultStore.cs
@@ -1,7 +1,6 @@
 // Copyright Epic Games, Inc. All Rights Reserved.
 
 using System;
-using System.Buffers;
 using System.Collections.Concurrent;
 using System.Collections.Generic;
 using System.Collections.Immutable;
