Fix issue CA1310 (Specify StringComparison for correctness) in UBT.

#Synced-CL 49329310

diff --git a/Engine/Platforms/VisionOS/Source/Programs/UnrealBuildTool/UEDeployVisionOS.cs b/Engine/Platforms/VisionOS/Source/Programs/UnrealBuildTool/UEDeployVisionOS.cs
--- a/Engine/Platforms/VisionOS/Source/Programs/UnrealBuildTool/UEDeployVisionOS.cs
+++ b/Engine/Platforms/VisionOS/Source/Programs/UnrealBuildTool/UEDeployVisionOS.cs
@@ -132,11 +132,11 @@ namespace UnrealBuildTool
 			// write the iCloud container identifier, if present in the old file
 			if (!String.IsNullOrEmpty(OldPListData))
 			{
-				int index = OldPListData.IndexOf("ICloudContainerIdentifier");
+				int index = OldPListData.IndexOf("ICloudContainerIdentifier", StringComparison.Ordinal);
 				if (index > 0)
 				{
-					index = OldPListData.IndexOf("<string>", index) + 8;
-					int length = OldPListData.IndexOf("</string>", index) - index;
+					index = OldPListData.IndexOf("<string>", index, StringComparison.Ordinal) + 8;
+					int length = OldPListData.IndexOf("</string>", index, StringComparison.Ordinal) - index;
 					string ICloudContainerIdentifier = OldPListData.Substring(index, length);
 					Text.AppendLine("\t<key>ICloudContainerIdentifier</key>");
 					Text.AppendLine(String.Format("\t<string>{0}</string>", ICloudContainerIdentifier));
diff --git a/Engine/Source/Programs/UnrealBuildTool/Actions/ActionGraph.cs b/Engine/Source/Programs/UnrealBuildTool/Actions/ActionGraph.cs
--- a/Engine/Source/Programs/UnrealBuildTool/Actions/ActionGraph.cs
+++ b/Engine/Source/Programs/UnrealBuildTool/Actions/ActionGraph.cs
@@ -322,7 +322,7 @@ namespace UnrealBuildTool
 							return false;
 						}
 
-						if (item.Location.FullName.EndsWith(".i.PVS-Studio.log")) // Ignore PVS Studio intermediate items
+						if (item.Location.FullName.EndsWith(".i.PVS-Studio.log", StringComparison.Ordinal)) // Ignore PVS Studio intermediate items
 						{
 							return false;
 						}
diff --git a/Engine/Source/Programs/UnrealBuildTool/Configuration/Descriptors/PluginDescriptor.cs b/Engine/Source/Programs/UnrealBuildTool/Configuration/Descriptors/PluginDescriptor.cs
--- a/Engine/Source/Programs/UnrealBuildTool/Configuration/Descriptors/PluginDescriptor.cs
+++ b/Engine/Source/Programs/UnrealBuildTool/Configuration/Descriptors/PluginDescriptor.cs
@@ -293,7 +293,7 @@ namespace UnrealBuildTool
 
 			// Due to a difference in command line parsing between Windows and Mac, we shipped a few Mac samples containing
 			// a category name with escaped quotes. Remove them here to make sure we can list them in the right category.
-			if (Category != null && Category.Length >= 2 && Category.StartsWith("\"") && Category.EndsWith("\""))
+			if (Category != null && Category.Length >= 2 && Category.StartsWith("\"", StringComparison.Ordinal) && Category.EndsWith("\"", StringComparison.Ordinal))
 			{
 				Category = Category.Substring(1, Category.Length - 2);
 			}
diff --git a/Engine/Source/Programs/UnrealBuildTool/Configuration/Descriptors/TargetDescriptor.cs b/Engine/Source/Programs/UnrealBuildTool/Configuration/Descriptors/TargetDescriptor.cs
--- a/Engine/Source/Programs/UnrealBuildTool/Configuration/Descriptors/TargetDescriptor.cs
+++ b/Engine/Source/Programs/UnrealBuildTool/Configuration/Descriptors/TargetDescriptor.cs
@@ -27,7 +27,7 @@ namespace UnrealBuildTool
 
 		public static string GetTestedName(string Name)
 		{
-			if (Name.EndsWith(TEST_TARGETS_SUFFIX))
+			if (Name.EndsWith(TEST_TARGETS_SUFFIX, StringComparison.Ordinal))
 			{
 				return Name.Substring(0, Name.Length - TEST_TARGETS_SUFFIX.Length);
 			}
diff --git a/Engine/Source/Programs/UnrealBuildTool/Configuration/Ini/ConfigCache.cs b/Engine/Source/Programs/UnrealBuildTool/Configuration/Ini/ConfigCache.cs
--- a/Engine/Source/Programs/UnrealBuildTool/Configuration/Ini/ConfigCache.cs
+++ b/Engine/Source/Programs/UnrealBuildTool/Configuration/Ini/ConfigCache.cs
@@ -786,7 +786,7 @@ namespace UnrealBuildTool
 				{
 					LastIndexOfPrefix = Index;
 				}
-				else if (Lines[Index].StartsWith("["))
+				else if (Lines[Index].StartsWith("[", StringComparison.Ordinal))
 				{
 					NextSectionIndex = Index;
 					// we found another section, so break out
diff --git a/Engine/Source/Programs/UnrealBuildTool/Configuration/Ini/ConfigHierarchy.cs b/Engine/Source/Programs/UnrealBuildTool/Configuration/Ini/ConfigHierarchy.cs
--- a/Engine/Source/Programs/UnrealBuildTool/Configuration/Ini/ConfigHierarchy.cs
+++ b/Engine/Source/Programs/UnrealBuildTool/Configuration/Ini/ConfigHierarchy.cs
@@ -558,7 +558,7 @@ namespace UnrealBuildTool
 		/// <returns>True if the text was parsed, false otherwise</returns>
 		public static bool TryParse(string Text, out float Value)
 		{
-			if (Text.EndsWith("f"))
+			if (Text.EndsWith("f", StringComparison.Ordinal))
 			{
 				return Single.TryParse(Text.Substring(0, Text.Length - 1), out Value);
 			}
@@ -576,7 +576,7 @@ namespace UnrealBuildTool
 		/// <returns>True if the text was parsed, false otherwise</returns>
 		public static bool TryParse(string Text, out double Value)
 		{
-			if (Text.EndsWith("f"))
+			if (Text.EndsWith("f", StringComparison.Ordinal))
 			{
 				return Double.TryParse(Text.Substring(0, Text.Length - 1), out Value);
 			}
@@ -1432,7 +1432,7 @@ namespace UnrealBuildTool
 		public static Dictionary<string, string> GetStructKeyValuePairs(string Input)
 		{
 			// we expect parens around a properly encoded struct
-			if (!Input.StartsWith("(") || !Input.EndsWith(")"))
+			if (!Input.StartsWith("(", StringComparison.Ordinal) || !Input.EndsWith(")", StringComparison.Ordinal))
 			{
 				return new Dictionary<string, string>();
 			}
diff --git a/Engine/Source/Programs/UnrealBuildTool/Configuration/Rules/RulesAssembly.cs b/Engine/Source/Programs/UnrealBuildTool/Configuration/Rules/RulesAssembly.cs
--- a/Engine/Source/Programs/UnrealBuildTool/Configuration/Rules/RulesAssembly.cs
+++ b/Engine/Source/Programs/UnrealBuildTool/Configuration/Rules/RulesAssembly.cs
@@ -372,7 +372,7 @@ namespace UnrealBuildTool
 				}
 
 				string NameWithoutTarget = Name;
-				if (NameWithoutTarget.EndsWith("Target"))
+				if (NameWithoutTarget.EndsWith("Target", StringComparison.Ordinal))
 				{
 					NameWithoutTarget = NameWithoutTarget.Substring(0, NameWithoutTarget.Length - 6);
 				}
@@ -773,7 +773,7 @@ namespace UnrealBuildTool
 			// Until VNI fully supports the new VM, we need the ability to have both the old and new
 			// available in some rare cases.  If we are using the old VM and the target hasn't overridden
 			// the new VM define, then set the define based on the old VM flag.
-			if (!Rules.GlobalDefinitions.Any(x => x.StartsWith("WITH_VERSE_VM=")))
+			if (!Rules.GlobalDefinitions.Any(x => x.StartsWith("WITH_VERSE_VM=", StringComparison.Ordinal)))
 			{
 				Rules.GlobalDefinitions.Add($"WITH_VERSE_VM={(Rules.bUseVerseBPVM ? 0 : 1)}");
 			}
diff --git a/Engine/Source/Programs/UnrealBuildTool/Configuration/Rules/RulesDocumentation.cs b/Engine/Source/Programs/UnrealBuildTool/Configuration/Rules/RulesDocumentation.cs
--- a/Engine/Source/Programs/UnrealBuildTool/Configuration/Rules/RulesDocumentation.cs
+++ b/Engine/Source/Programs/UnrealBuildTool/Configuration/Rules/RulesDocumentation.cs
@@ -49,7 +49,7 @@ namespace UnrealBuildTool
 			// First read the fields
 			foreach (FieldInfo Field in RulesType.GetFields(BindingFlags.Instance | BindingFlags.SetProperty | BindingFlags.Public))
 			{
-				if (!Field.FieldType.IsClass || !Field.FieldType.Name.EndsWith("TargetRules"))
+				if (!Field.FieldType.IsClass || !Field.FieldType.Name.EndsWith("TargetRules", StringComparison.Ordinal))
 				{
 					List<string>? Lines;
 					if (TryGetXmlComment(InputDocumentation, Field, Logger, out Lines))
@@ -70,7 +70,7 @@ namespace UnrealBuildTool
 			// ...then read all the properties
 			foreach (PropertyInfo Property in RulesType.GetProperties(BindingFlags.Instance | BindingFlags.SetProperty | BindingFlags.Public))
 			{
-				if (!Property.PropertyType.IsClass || !Property.PropertyType.Name.EndsWith("TargetRules"))
+				if (!Property.PropertyType.IsClass || !Property.PropertyType.Name.EndsWith("TargetRules", StringComparison.Ordinal))
 				{
 					List<string>? Lines;
 					if (TryGetXmlComment(InputDocumentation, Property, Logger, out Lines))
@@ -247,7 +247,7 @@ namespace UnrealBuildTool
 				List<string> NewLines = new List<string>(Node.InnerText.Trim().Split('\n').Select(x => x.Trim()));
 				for (int Idx = NewLines.Count - 1; Idx > 0; Idx--)
 				{
-					if (NewLines[Idx - 1].Length > 0 && !NewLines[Idx].StartsWith("*") && !NewLines[Idx].StartsWith("-"))
+					if (NewLines[Idx - 1].Length > 0 && !NewLines[Idx].StartsWith("*", StringComparison.Ordinal) && !NewLines[Idx].StartsWith("-", StringComparison.Ordinal))
 					{
 						NewLines[Idx - 1] += " " + NewLines[Idx];
 						NewLines.RemoveAt(Idx);
@@ -288,7 +288,7 @@ namespace UnrealBuildTool
 			Writer.WriteLine("$ {0} ({1}): {2}", Field.Name, GetPrettyTypeName(Field.Type), Lines[0]);
 			for (int Idx = 1; Idx < Lines.Count; Idx++)
 			{
-				if (Lines[Idx].StartsWith("*") || Lines[Idx].StartsWith("-"))
+				if (Lines[Idx].StartsWith("*", StringComparison.Ordinal) || Lines[Idx].StartsWith("-", StringComparison.Ordinal))
 				{
 					Writer.WriteLine("        * {0}", Lines[Idx].Substring(1).TrimStart());
 				}
@@ -359,10 +359,10 @@ namespace UnrealBuildTool
 					Writer.WriteLine("      <dd>");
 					for (int Idx = 0; Idx < Lines.Count; Idx++)
 					{
-						if (Lines[Idx].StartsWith("*") || Lines[Idx].StartsWith("-"))
+						if (Lines[Idx].StartsWith("*", StringComparison.Ordinal) || Lines[Idx].StartsWith("-", StringComparison.Ordinal))
 						{
 							Writer.WriteLine("        <ul>");
-							for (; Idx < Lines.Count && (Lines[Idx].StartsWith("*") || Lines[Idx].StartsWith("-")); Idx++)
+							for (; Idx < Lines.Count && (Lines[Idx].StartsWith("*", StringComparison.Ordinal) || Lines[Idx].StartsWith("-", StringComparison.Ordinal)); Idx++)
 							{
 								Writer.WriteLine("          <li>{0}</li>", Lines[Idx].Substring(1).TrimStart());
 							}
diff --git a/Engine/Source/Programs/UnrealBuildTool/Configuration/Rules/TestModuleRules.cs b/Engine/Source/Programs/UnrealBuildTool/Configuration/Rules/TestModuleRules.cs
--- a/Engine/Source/Programs/UnrealBuildTool/Configuration/Rules/TestModuleRules.cs
+++ b/Engine/Source/Programs/UnrealBuildTool/Configuration/Rules/TestModuleRules.cs
@@ -387,10 +387,10 @@ namespace UnrealBuildTool
 
 		private static string TryGetBinariesPath(string ModuleDirectory)
 		{
-			int SourceFolderIndex = ModuleDirectory.IndexOf("Source");
+			int SourceFolderIndex = ModuleDirectory.IndexOf("Source", StringComparison.Ordinal);
 			if (SourceFolderIndex < 0)
 			{
-				int PluginFolderIndex = ModuleDirectory.IndexOf("Plugins");
+				int PluginFolderIndex = ModuleDirectory.IndexOf("Plugins", StringComparison.Ordinal);
 				if (PluginFolderIndex >= 0)
 				{
 					return ModuleDirectory.Substring(0, PluginFolderIndex) + "Binaries";
diff --git a/Engine/Source/Programs/UnrealBuildTool/Configuration/UEBuildBinary.cs b/Engine/Source/Programs/UnrealBuildTool/Configuration/UEBuildBinary.cs
--- a/Engine/Source/Programs/UnrealBuildTool/Configuration/UEBuildBinary.cs
+++ b/Engine/Source/Programs/UnrealBuildTool/Configuration/UEBuildBinary.cs
@@ -614,7 +614,7 @@ namespace UnrealBuildTool
 		public static FileReference GetAdditionalConsoleAppPath(FileReference BinaryPath)
 		{
 			DirectoryReference Directory = BinaryPath.Directory;
-			if (Directory.FullName.EndsWith(".app/Contents/MacOS"))
+			if (Directory.FullName.EndsWith(".app/Contents/MacOS", StringComparison.Ordinal))
 			{
 				Directory = Directory.ParentDirectory!.ParentDirectory!.ParentDirectory!;
 			}
diff --git a/Engine/Source/Programs/UnrealBuildTool/Configuration/UEBuildFramework.cs b/Engine/Source/Programs/UnrealBuildTool/Configuration/UEBuildFramework.cs
--- a/Engine/Source/Programs/UnrealBuildTool/Configuration/UEBuildFramework.cs
+++ b/Engine/Source/Programs/UnrealBuildTool/Configuration/UEBuildFramework.cs
@@ -95,7 +95,7 @@ namespace UnrealBuildTool
 			this.CopyBundledAssets = CopyBundledAssets;
 			this.bCopyFramework = bCopyFramework;
 			this.bLinkFramework = bLinkFramework;
-			if (this.ZipFile != null && this.ZipFile.FullName.EndsWith(".xcframework.zip"))
+			if (this.ZipFile != null && this.ZipFile.FullName.EndsWith(".xcframework.zip", System.StringComparison.Ordinal))
 			{
 				XCFrameworkVariants = LoadXCFrameworkVariantsFromZipFile(Logger);
 			}
@@ -117,7 +117,7 @@ namespace UnrealBuildTool
 			this.CopyBundledAssets = CopyBundledAssets;
 			this.bCopyFramework = bCopyFramework;
 			this.bLinkFramework = bLinkFramework;
-			if (this.FrameworkDirectory.FullName.EndsWith(".xcframework"))
+			if (this.FrameworkDirectory.FullName.EndsWith(".xcframework", System.StringComparison.Ordinal))
 			{
 				XCFrameworkVariants = LoadXCFrameworkVariants(Logger);
 			}
diff --git a/Engine/Source/Programs/UnrealBuildTool/Configuration/UEBuildModule.cs b/Engine/Source/Programs/UnrealBuildTool/Configuration/UEBuildModule.cs
--- a/Engine/Source/Programs/UnrealBuildTool/Configuration/UEBuildModule.cs
+++ b/Engine/Source/Programs/UnrealBuildTool/Configuration/UEBuildModule.cs
@@ -1268,7 +1268,7 @@ namespace UnrealBuildTool
 					LogDependencyNameList("Private Include Paths:", Rules.PrivateIncludePathModuleNames, Logger);
 				}
 
-				ReferenceStackParentModules ??= new(Comparer<UEBuildModule>.Create(static (a, b) => a.Name.CompareTo(b.Name)));
+				ReferenceStackParentModules ??= new(Comparer<UEBuildModule>.Create(static (a, b) => string.Compare(a.Name, b.Name, StringComparison.Ordinal)));
 
 				// Create all the dependency modules - pass through the reference stack so we can check for cycles
 				RecursivelyCreateModulesByName(Rules.PublicDependencyModuleNames, ref PublicDependencyModules, ref bDependsOnVerse, CreateModule, ReferenceChain, Logger);
diff --git a/Engine/Source/Programs/UnrealBuildTool/Configuration/UEBuildModuleCPP.cs b/Engine/Source/Programs/UnrealBuildTool/Configuration/UEBuildModuleCPP.cs
--- a/Engine/Source/Programs/UnrealBuildTool/Configuration/UEBuildModuleCPP.cs
+++ b/Engine/Source/Programs/UnrealBuildTool/Configuration/UEBuildModuleCPP.cs
@@ -745,7 +745,7 @@ namespace UnrealBuildTool
 					foreach (FileItem File in DirItem.EnumerateFiles())
 					{
 						string FileName = File.Name;
-						if (FileName.EndsWith(".gen.cpp"))
+						if (FileName.EndsWith(".gen.cpp", System.StringComparison.Ordinal))
 						{
 							string Key = Prefix + FileName.Substring(0, FileName.Length - ".gen.cpp".Length);
 							GeneratedFiles.Add(Key, File);
@@ -794,7 +794,7 @@ namespace UnrealBuildTool
 
 						foreach (string Name in GeneratedFiles.Keys)
 						{
-							if (!Name.StartsWith("UHT/"))
+							if (!Name.StartsWith("UHT/", StringComparison.Ordinal))
 							{
 								continue;
 							}
@@ -2957,7 +2957,7 @@ namespace UnrealBuildTool
 				if (SubDirectory.Name == "Tests" && Rules.Target.SelectedTestModuleNames.Count > 0 && IsModuleDirectory(SubDirectory.GetParentDirectoryItem()))
 				{
 					// If it's one of the selected test module names, make sure to include them exclusively
-					if (!Rules.Target.SelectedTestModuleNames.Any(STM => SubDirectory.FullName.EndsWith($"{Path.DirectorySeparatorChar}{STM}{Path.DirectorySeparatorChar}Tests")))
+					if (!Rules.Target.SelectedTestModuleNames.Any(STM => SubDirectory.FullName.EndsWith($"{Path.DirectorySeparatorChar}{STM}{Path.DirectorySeparatorChar}Tests", StringComparison.Ordinal)))
 					{
 						continue;
 					}
@@ -3053,7 +3053,7 @@ namespace UnrealBuildTool
 				}
 				// some text editors will leave temp files that start with . next to a, say, .cpp file, and we do not want to compile them,
 				// but skip files like .DS_Store that don't have an extension (so, no more .'s after the first)
-				else if (InputFile.Name.StartsWith(".") && InputFile.Name.IndexOf('.', 1) != -1)
+				else if (InputFile.Name.StartsWith(".", StringComparison.Ordinal) && InputFile.Name.IndexOf('.', 1) != -1)
 				{
 					if (KnownInputFileExtensions.Contains(InputFile.Location.GetExtension()))
 					{
@@ -3187,7 +3187,7 @@ namespace UnrealBuildTool
 
 			foreach (string FilePath in Directory.EnumerateFiles(MaybeVerseDirectory.FullName, "*.v*", SearchOption.AllDirectories))
 			{
-				if (FilePath.EndsWith(".verse") || FilePath.EndsWith(".vmodule"))
+				if (FilePath.EndsWith(".verse", System.StringComparison.Ordinal) || FilePath.EndsWith(".vmodule", System.StringComparison.Ordinal))
 				{
 					return true;
 				}
diff --git a/Engine/Source/Programs/UnrealBuildTool/Configuration/UEBuildTarget.cs b/Engine/Source/Programs/UnrealBuildTool/Configuration/UEBuildTarget.cs
--- a/Engine/Source/Programs/UnrealBuildTool/Configuration/UEBuildTarget.cs
+++ b/Engine/Source/Programs/UnrealBuildTool/Configuration/UEBuildTarget.cs
@@ -2140,7 +2140,7 @@ namespace UnrealBuildTool
 
 				List<string> SystemLibraryPaths = new List<string>();
 				SystemLibraryPaths.Add(Directory.GetCurrentDirectory());
-				SystemLibraryPaths.AddRange(Rules.PublicSystemLibraryPaths.Where(x => !x.StartsWith("$(")).Select(x => Path.GetFullPath(x.Replace('/', Path.DirectorySeparatorChar))));
+				SystemLibraryPaths.AddRange(Rules.PublicSystemLibraryPaths.Where(x => !x.StartsWith("$(", StringComparison.Ordinal)).Select(x => Path.GetFullPath(x.Replace('/', Path.DirectorySeparatorChar))));
 
 				// Get all the extensions to look for
 				List<string> LibraryExtensions = new List<string>();
@@ -2173,7 +2173,7 @@ namespace UnrealBuildTool
 				AllIncludePaths.AddRange(Rules.PublicSystemIncludePaths);
 
 				// Add all the include paths
-				foreach (string IncludePath in AllIncludePaths.Where(x => !x.StartsWith("$(")))
+				foreach (string IncludePath in AllIncludePaths.Where(x => !x.StartsWith("$(", StringComparison.Ordinal)))
 				{
 					if (Directory.Exists(IncludePath))
 					{
@@ -2443,7 +2443,7 @@ namespace UnrealBuildTool
 				{
 					if (Pair.Value == BuildProductType.Executable)
 					{
-						if (System.IO.Path.GetFileNameWithoutExtension(Pair.Key.FullName).EndsWith("-Cmd") && Receipt.LaunchCmd == null)
+						if (System.IO.Path.GetFileNameWithoutExtension(Pair.Key.FullName).EndsWith("-Cmd", System.StringComparison.Ordinal) && Receipt.LaunchCmd == null)
 						{
 							Receipt.LaunchCmd = Pair.Key;
 						}
@@ -3625,7 +3625,7 @@ namespace UnrealBuildTool
 		public DirectoryReference GetExecutableDir()
 		{
 			DirectoryReference ExeDir = Binaries[0].OutputDir;
-			if (Platform == UnrealTargetPlatform.Mac && ExeDir.FullName.EndsWith(".app/Contents/MacOS"))
+			if (Platform == UnrealTargetPlatform.Mac && ExeDir.FullName.EndsWith(".app/Contents/MacOS", System.StringComparison.Ordinal))
 			{
 				ExeDir = ExeDir.ParentDirectory!.ParentDirectory!.ParentDirectory!;
 			}
@@ -3993,7 +3993,7 @@ namespace UnrealBuildTool
 					{
 						continue;
 					}
-					if (current.Module.Rules.PublicDefinitions.Exists(x => x.StartsWith("PER_MODULE_INLINE_FILE")))
+					if (current.Module.Rules.PublicDefinitions.Exists(x => x.StartsWith("PER_MODULE_INLINE_FILE", StringComparison.Ordinal)))
 					{
 						((UEBuildModuleCPP)current.Source).bCreatePerModuleFile = true;
 						return;
@@ -5068,7 +5068,7 @@ namespace UnrealBuildTool
 
 					// Skip platform extension modules. We only care about the base modules, not the platform overrides.
 					// The platform overrides get applied at a later stage when we actually come to build the module.
-					if (!UEBuildPlatform.GetPlatformFolderNames().Any(Name => RulesType.Name.EndsWith("_" + Name)))
+					if (!UEBuildPlatform.GetPlatformFolderNames().Any(Name => RulesType.Name.EndsWith("_" + Name, StringComparison.Ordinal)))
 					{
 						if (ModuleFileName.ContainsAnyNames(ExcludeFolders, Unreal.EngineDirectory))
 						{
diff --git a/Engine/Source/Programs/UnrealBuildTool/Executors/Experimental/FASTBuild.cs b/Engine/Source/Programs/UnrealBuildTool/Executors/Experimental/FASTBuild.cs
--- a/Engine/Source/Programs/UnrealBuildTool/Executors/Experimental/FASTBuild.cs
+++ b/Engine/Source/Programs/UnrealBuildTool/Executors/Experimental/FASTBuild.cs
@@ -557,10 +557,10 @@ namespace UnrealBuildTool
 			string ResponseFilePath = "";
 			List<string> AllTokens = new List<string>();
 
-			int ResponseFileTokenIndex = Array.FindIndex(RawTokens, RawToken => RawToken.StartsWith("@\""));
+			int ResponseFileTokenIndex = Array.FindIndex(RawTokens, RawToken => RawToken.StartsWith("@\"", StringComparison.Ordinal));
 			if (ResponseFileTokenIndex == -1)
 			{
-				ResponseFileTokenIndex = Array.FindIndex(RawTokens, RawToken => RawToken.StartsWith("@"));
+				ResponseFileTokenIndex = Array.FindIndex(RawTokens, RawToken => RawToken.StartsWith("@", StringComparison.Ordinal));
 			}
 			if (ResponseFileTokenIndex > -1) //Response files are in 4.13 by default. Changing VCToolChain to not do this is probably better.
 			{
@@ -660,7 +660,7 @@ namespace UnrealBuildTool
 
 				// Defines can have escaped quotes and other strings inside them
 				// so we consume tokens until we've closed any open unescaped parentheses.
-				if ((Token.StartsWith("/D") || Token.StartsWith("-D")) && !QuotesOpened)
+				if ((Token.StartsWith("/D", StringComparison.Ordinal) || Token.StartsWith("-D", StringComparison.Ordinal)) && !QuotesOpened)
 				{
 					if (numQuotes == 0 || numQuotes == 2)
 					{
@@ -678,7 +678,7 @@ namespace UnrealBuildTool
 							{
 								PartialToken += " ";
 							}
-							else if (!NextToken.EndsWith("\\\"") && NextToken.EndsWith("\"")) //Looking for a token that ends with a non-escaped "
+							else if (!NextToken.EndsWith("\\\"", System.StringComparison.Ordinal) && NextToken.EndsWith("\"", System.StringComparison.Ordinal)) //Looking for a token that ends with a non-escaped "
 							{
 								ProcessedTokens.Add(PartialToken + " " + NextToken);
 								AddedToken = true;
@@ -793,7 +793,7 @@ namespace UnrealBuildTool
 					ProcessedTokens.RemoveAt(i);
 					i--;
 				}
-				else if (!Token.StartsWith("/") && !Token.StartsWith("-") && !Token.Contains(".framework", StringComparison.Ordinal))
+				else if (!Token.StartsWith("/", StringComparison.Ordinal) && !Token.StartsWith("-", StringComparison.Ordinal) && !Token.Contains(".framework", StringComparison.Ordinal))
 				{
 					ParsedCompilerOptions["InputFile"] = Token;
 					ProcessedTokens.RemoveAt(i);
diff --git a/Engine/Source/Programs/UnrealBuildTool/Executors/SNDBS.cs b/Engine/Source/Programs/UnrealBuildTool/Executors/SNDBS.cs
--- a/Engine/Source/Programs/UnrealBuildTool/Executors/SNDBS.cs
+++ b/Engine/Source/Programs/UnrealBuildTool/Executors/SNDBS.cs
@@ -207,7 +207,7 @@ namespace UnrealBuildTool
 			}
 
 			// Check the service is running
-			if (OperatingSystem.IsWindows() && !ServiceController.GetServices().Any(s => OperatingSystem.IsWindows() && s.ServiceName.StartsWith("SNDBS") && s.Status == ServiceControllerStatus.Running))
+			if (OperatingSystem.IsWindows() && !ServiceController.GetServices().Any(s => OperatingSystem.IsWindows() && s.ServiceName.StartsWith("SNDBS", StringComparison.Ordinal) && s.Status == ServiceControllerStatus.Running))
 			{
 				return false;
 			}
@@ -276,19 +276,22 @@ namespace UnrealBuildTool
 
 					if (a.PrerequisiteItems.Any())
 					{
-						Job["explicit_input_files"] = a.PrerequisiteItems.Where(i => !(i.AbsolutePath.EndsWith(".rsp") || i.AbsolutePath.EndsWith(".response"))).Select(i => new Dictionary<string, object>()
-						{
-							["filename"] = i.AbsolutePath
-						}).ToList();
+						Job["explicit_input_files"] = a.PrerequisiteItems
+							.Where(i => !(i.AbsolutePath.EndsWith(".rsp", System.StringComparison.Ordinal) || i.AbsolutePath.EndsWith(".response", System.StringComparison.Ordinal)))
+							.Select(i => new Dictionary<string, object>()
+								{
+									["filename"] = i.AbsolutePath
+								})
+							.ToList();
 					}
 
 					// Look for any prerequisite actions that produce .pch files and add any .cpp files they depend on
 					List<Dictionary<string, object>>? ExplicitInputFiles = Job["explicit_input_files"] as List<Dictionary<string, object>>;
 
 					ExplicitInputFiles?.AddRange(a.PrerequisiteActions
-						.Where(Prereq => Prereq.ProducedItems.Any(Produced => Produced.AbsolutePath.EndsWith(".gch") || Produced.AbsolutePath.EndsWith(".pch")))
+						.Where(Prereq => Prereq.ProducedItems.Any(Produced => Produced.AbsolutePath.EndsWith(".gch", System.StringComparison.Ordinal) || Produced.AbsolutePath.EndsWith(".pch", System.StringComparison.Ordinal)))
 						.SelectMany(Prereq => Prereq.PrerequisiteItems, (_, PrereqFile) => PrereqFile.AbsolutePath)
-						.Where(Path => Path.EndsWith(".cpp"))
+						.Where(Path => Path.EndsWith(".cpp", StringComparison.Ordinal))
 						.Select(Path => new Dictionary<string, object>()
 						{
 							["filename"] = Path
@@ -300,7 +303,7 @@ namespace UnrealBuildTool
 					//
 					// The file is tiny and harmless to include when not using clang-cl.
 					List<Dictionary<string, object>>? PrecompiledDependencies = ExplicitInputFiles?.Select(File => File["filename"].ToString())
-						.Where(Path => Path != null && Path.EndsWith(".h.pch"))
+						.Where(Path => Path != null && Path.EndsWith(".h.pch", System.StringComparison.Ordinal))
 						.Select(Path => Path!.Replace(".h.pch", ".cpp", StringComparison.Ordinal))
 						.Select(Path => new Dictionary<string, object>()
 						{
@@ -361,7 +364,7 @@ namespace UnrealBuildTool
 					if (Args.Data != null)
 					{
 						string Text = Args.Data;
-						if (Text.StartsWith(ProgressMarkupPrefix))
+						if (Text.StartsWith(ProgressMarkupPrefix, StringComparison.Ordinal))
 						{
 							Writer.Write(++NumCompletedActions, NumActions);
 
diff --git a/Engine/Source/Programs/UnrealBuildTool/Executors/UnrealBuildAccelerator/UBAAgentCoordinatorHorde.cs b/Engine/Source/Programs/UnrealBuildTool/Executors/UnrealBuildAccelerator/UBAAgentCoordinatorHorde.cs
--- a/Engine/Source/Programs/UnrealBuildTool/Executors/UnrealBuildAccelerator/UBAAgentCoordinatorHorde.cs
+++ b/Engine/Source/Programs/UnrealBuildTool/Executors/UnrealBuildAccelerator/UBAAgentCoordinatorHorde.cs
@@ -722,7 +722,7 @@ namespace UnrealBuildTool
 			IEnumerable<string> providers = Unreal.IsBuildMachine()
 				? [.. ubaConfig.BuildMachineProviders, .. ubaConfig.IniBuildMachineProviders]
 				: [.. ubaConfig.Providers, .. ubaConfig.IniProviders];
-			providers = providers.Distinct().Where(x => !String.IsNullOrEmpty(x) && x.StartsWith(ProviderPrefix));
+			providers = providers.Distinct().Where(x => !String.IsNullOrEmpty(x) && x.StartsWith(ProviderPrefix, StringComparison.Ordinal));
 
 			if (!providers.Any())
 			{
diff --git a/Engine/Source/Programs/UnrealBuildTool/Executors/UnrealBuildAccelerator/UBAExecutor.cs b/Engine/Source/Programs/UnrealBuildTool/Executors/UnrealBuildAccelerator/UBAExecutor.cs
--- a/Engine/Source/Programs/UnrealBuildTool/Executors/UnrealBuildAccelerator/UBAExecutor.cs
+++ b/Engine/Source/Programs/UnrealBuildTool/Executors/UnrealBuildAccelerator/UBAExecutor.cs
@@ -198,7 +198,7 @@ namespace UnrealBuildTool
 				{
 					providers = [ProviderPrefix];
 				}
-				foreach (string provider in providers.Distinct().Where(x => x.StartsWith(ProviderPrefix)))
+				foreach (string provider in providers.Distinct().Where(x => x.StartsWith(ProviderPrefix, StringComparison.Ordinal)))
 				{
 					foreach (TargetDescriptor targetDescriptor in targetDescriptors)
 					{
diff --git a/Engine/Source/Programs/UnrealBuildTool/Executors/UnrealBuildAccelerator/UnrealBuildAcceleratorHordeConfig.cs b/Engine/Source/Programs/UnrealBuildTool/Executors/UnrealBuildAccelerator/UnrealBuildAcceleratorHordeConfig.cs
--- a/Engine/Source/Programs/UnrealBuildTool/Executors/UnrealBuildAccelerator/UnrealBuildAcceleratorHordeConfig.cs
+++ b/Engine/Source/Programs/UnrealBuildTool/Executors/UnrealBuildAccelerator/UnrealBuildAcceleratorHordeConfig.cs
@@ -304,7 +304,7 @@ namespace UnrealBuildTool
 				if (attribute != null)
 				{
 					string keyName = attribute.KeyName ?? property.Name;
-					object? value = TryGetValue(propertyType, keyName) ?? (keyName.StartsWith("Uba") ? TryGetValue(propertyType, keyName[3..]) : null);
+					object? value = TryGetValue(propertyType, keyName) ?? (keyName.StartsWith("Uba", StringComparison.Ordinal) ? TryGetValue(propertyType, keyName[3..]) : null);
 					if (value != null)
 					{
 						property.SetValue(config, value);
diff --git a/Engine/Source/Programs/UnrealBuildTool/Executors/XGE.cs b/Engine/Source/Programs/UnrealBuildTool/Executors/XGE.cs
--- a/Engine/Source/Programs/UnrealBuildTool/Executors/XGE.cs
+++ b/Engine/Source/Programs/UnrealBuildTool/Executors/XGE.cs
@@ -729,7 +729,7 @@ namespace UnrealBuildTool
 					if (Args.Data != null)
 					{
 						string Text = Args.Data;
-						if (Text.StartsWith(ProgressMarkupPrefix))
+						if (Text.StartsWith(ProgressMarkupPrefix, StringComparison.Ordinal))
 						{
 							// Code below should not need to be tested for success but if some logging from XGE is wrong we just gracefully ignore it and accept that counting might end up wrong
 							int ActionIndex = -1;
@@ -758,7 +758,7 @@ namespace UnrealBuildTool
 
 							// Strip out anything that is just an XGE timer. Some programs don't output anything except the progress text.
 							Text = Args.Data.Substring(MarkupLength);
-							if (Text.StartsWith(" (") && Text.EndsWith(")"))
+							if (Text.StartsWith(" (", StringComparison.Ordinal) && Text.EndsWith(')'))
 							{
 								// Write the progress text with the next line of output if the current doesn't have any status.
 								ProgressText = Text.Trim();
diff --git a/Engine/Source/Programs/UnrealBuildTool/Modes/BuildMode.cs b/Engine/Source/Programs/UnrealBuildTool/Modes/BuildMode.cs
--- a/Engine/Source/Programs/UnrealBuildTool/Modes/BuildMode.cs
+++ b/Engine/Source/Programs/UnrealBuildTool/Modes/BuildMode.cs
@@ -215,7 +215,7 @@ namespace UnrealBuildTool
 						FileReference BaseLogFile = Log.OutputFile ?? new FileReference(BaseLogFileName);
 						FileReference RemoteLogFile = FileReference.Combine(BaseLogFile.Directory, BaseLogFile.GetFileNameWithoutExtension() + "_Remote.txt");
 
-						int RemoteIniIndex = TargetDesc.AdditionalArguments.FindIndex(x => x.ToLower().StartsWith("-remoteini="));
+						int RemoteIniIndex = TargetDesc.AdditionalArguments.FindIndex(x => x.ToLower().StartsWith("-remoteini=", StringComparison.Ordinal));
 						string? RemoteIni = RemoteIniIndex == -1 ? null : new string(TargetDesc.AdditionalArguments[RemoteIniIndex].Substring(11));
 						RemoteMac RemoteMac = new RemoteMac(TargetDesc.ProjectFile, Logger, RemoteIni);
 						if (!RemoteMac.Build(TargetDesc, RemoteLogFile, bSkipPreBuildTargets, Logger))
@@ -973,7 +973,7 @@ namespace UnrealBuildTool
 				Actions
 				.SelectMany(x => x.PrerequisiteItems)
 				.Distinct()
-				.Where(x => IsCode(x.Location) && !x.Name.StartsWith(Unity.ModulePrefix) && !x.HasExtension(".gen.cpp"));
+				.Where(x => IsCode(x.Location) && !x.Name.StartsWith(Unity.ModulePrefix, StringComparison.Ordinal) && !x.HasExtension(".gen.cpp"));
 
 			FileReference IncludeCache = FileReference.Combine(CacheDir, "Intermediate", "HeaderLookup.dat");
 			CppIncludeLookup IncludeLookup = new CppIncludeLookup(IncludeCache);
diff --git a/Engine/Source/Programs/UnrealBuildTool/Modes/FixIncludePathsMode.cs b/Engine/Source/Programs/UnrealBuildTool/Modes/FixIncludePathsMode.cs
--- a/Engine/Source/Programs/UnrealBuildTool/Modes/FixIncludePathsMode.cs
+++ b/Engine/Source/Programs/UnrealBuildTool/Modes/FixIncludePathsMode.cs
@@ -249,7 +249,7 @@ namespace UnrealBuildTool
 														string? FoundPreferredPath = PreferredPaths.FirstOrDefault(path => FullPath.Contains(path, StringComparison.Ordinal));
 														if (FoundPreferredPath != null)
 														{
-															int end = FullPath.LastIndexOf(FoundPreferredPath) + FoundPreferredPath.Length;
+															int end = FullPath.LastIndexOf(FoundPreferredPath, StringComparison.Ordinal) + FoundPreferredPath.Length;
 															PreferredInclude = FullPath.Substring(end);
 
 															// Is the current include a shortened version of the preferred include path?
diff --git a/Engine/Source/Programs/UnrealBuildTool/Modes/IWYUMode.cs b/Engine/Source/Programs/UnrealBuildTool/Modes/IWYUMode.cs
--- a/Engine/Source/Programs/UnrealBuildTool/Modes/IWYUMode.cs
+++ b/Engine/Source/Programs/UnrealBuildTool/Modes/IWYUMode.cs
@@ -688,7 +688,7 @@ namespace UnrealBuildTool
 
 					foreach (FileItem OutputItem in OutputItems)
 					{
-						if (OutputItem.Name.EndsWith(".iwyu"))
+						if (OutputItem.Name.EndsWith(".iwyu", StringComparison.Ordinal))
 						{
 							TocContent.Add(OutputItem.AbsolutePath);
 						}
@@ -800,7 +800,7 @@ namespace UnrealBuildTool
 							continue;
 						}
 
-						bool isHeader = !isCpp && info.File.EndsWith(".h");
+						bool isHeader = !isCpp && info.File.EndsWith(".h", StringComparison.Ordinal);
 						bool isGeneratedHeader = isHeader && (info.File.EndsWith(".generated.h", StringComparison.Ordinal) || info.File.EndsWith(".gen.h", StringComparison.Ordinal));
 
 						lock (Infos)
@@ -1104,7 +1104,9 @@ namespace UnrealBuildTool
 					}
 
 					// We don't want to remove these because they enable includes of windows headers.
-					if (Entry.Printable.StartsWith("Windows/AllowWindows") || Entry.Printable.StartsWith("Windows/HideWindows") || Entry.Printable.StartsWith("Windows/WindowsHWrapper.h"))
+					if (Entry.Printable.StartsWith("Windows/AllowWindows", StringComparison.Ordinal) ||
+						Entry.Printable.StartsWith("Windows/HideWindows", StringComparison.Ordinal) ||
+						Entry.Printable.StartsWith("Windows/WindowsHWrapper.h", StringComparison.Ordinal))
 					{
 						Info.Includes.Add(Entry);
 					}
@@ -1162,7 +1164,7 @@ namespace UnrealBuildTool
 						string fixedPath = value.File.Replace('/', '\\');
 						foreach (DirectoryReference? Dir in value.Module.PublicIncludePaths.Union(value.Module.PrivateIncludePaths).Union(value.Module.PublicSystemIncludePaths))
 						{
-							if (fixedPath.StartsWith(Dir.FullName))
+							if (fixedPath.StartsWith(Dir.FullName, StringComparison.Ordinal))
 							{
 								string printable = value.File.Substring(Dir.FullName.Length + 1);
 								IWYUIncludeEntry entry = new IWYUIncludeEntry() { Full = MissingInclude.Full, Printable = printable, System = false, Resolved = value, };
@@ -1592,13 +1594,13 @@ namespace UnrealBuildTool
 				{
 					// We never remove header with name matching cpp
 					string NameWithoutPath = Include.Printable;
-					int LastSlash = NameWithoutPath.LastIndexOf("/");
+					int LastSlash = NameWithoutPath.LastIndexOf('/');
 					if (LastSlash != -1)
 					{
 						NameWithoutPath = NameWithoutPath.Substring(LastSlash + 1);
 					}
 
-					if (IsCpp && Include.Full.EndsWith(".generated.h"))
+					if (IsCpp && Include.Full.EndsWith(".generated.h", StringComparison.Ordinal))
 					{
 						continue;
 					}
@@ -1696,7 +1698,7 @@ namespace UnrealBuildTool
 				HashSet<string> SeenIncludes = new();
 				foreach (IWYUIncludeEntry SeenInclude in Info.IncludesSeenInFile)
 				{
-					if (SeenInclude.Printable.EndsWith(".ipp")) // Abuse this to make sure we never remove .ipp includes
+					if (SeenInclude.Printable.EndsWith(".ipp", StringComparison.Ordinal)) // Abuse this to make sure we never remove .ipp includes
 					{
 						IncludesToAdd.Remove($"\"{SeenInclude.Printable}\"");
 						continue;
@@ -1931,7 +1933,7 @@ namespace UnrealBuildTool
 						NewLines.Add(OldLine);
 						++LineIndex;
 
-						if (!OldLine.TrimStart().StartsWith("#pragma"))
+						if (!OldLine.TrimStart().StartsWith("#pragma", StringComparison.Ordinal))
 						{
 							continue;
 						}
@@ -1982,7 +1984,7 @@ namespace UnrealBuildTool
 								{
 									IsInFirstIncludeBlock = false;
 								}
-								if (OldLineTrimmedStart.StartsWith("#if"))
+								if (OldLineTrimmedStart.StartsWith("#if", StringComparison.Ordinal))
 								{
 									// This logic is a bit shaky but handle the situations where file starts with #if and ends with #endif
 									if (LastSeenIncludeBeforeCode != -1)
@@ -2014,7 +2016,7 @@ namespace UnrealBuildTool
 							else if (LinesToAdd.Count != 0 && (!IsCpp || LastSeenIncludeBeforeCode != -1))
 							{
 								string LineToAdd = LinesToAdd.First();
-								if (LineToAdd.CompareTo(OldLine) < 0)
+								if (string.Compare(LineToAdd, OldLine, StringComparison.Ordinal) < 0)
 								{
 									NewLines.Add(LineToAdd);
 									LinesToAdd.Remove(LineToAdd);
@@ -2081,7 +2083,7 @@ namespace UnrealBuildTool
 				// If file is public, in engine and we have a deprecation tag set we will 
 				// add a deprecated include scope at the end of the file (unless scope already exists, then we'll add it inside that)
 				string EngineDir = Unreal.EngineDirectory.FullName.Replace('\\', '/');
-				if (!(IsPrivate || IsInternal) && Info.File.StartsWith(EngineDir) && !String.IsNullOrEmpty(HeaderDeprecationTag))
+				if (!(IsPrivate || IsInternal) && Info.File.StartsWith(EngineDir, StringComparison.Ordinal) && !String.IsNullOrEmpty(HeaderDeprecationTag))
 				{
 					Dictionary<string, string> PrintableToFull = new();
 					foreach (IWYUIncludeEntry Seen in Info.IncludesSeenInFile)
@@ -2404,7 +2406,7 @@ namespace UnrealBuildTool
 			Logger.LogInformation($"Generating transitive seen include lists...");
 			Parallel2.ForEach(Infos.Values, Info =>
 			{
-				if (!IsValidForUpdate(Info, ValidPaths, false) || Info.File.EndsWith(".h"))
+				if (!IsValidForUpdate(Info, ValidPaths, false) || Info.File.EndsWith(".h", StringComparison.Ordinal))
 				{
 					return;
 				}
diff --git a/Engine/Source/Programs/UnrealBuildTool/Modes/InlineGeneratedCpps.cs b/Engine/Source/Programs/UnrealBuildTool/Modes/InlineGeneratedCpps.cs
--- a/Engine/Source/Programs/UnrealBuildTool/Modes/InlineGeneratedCpps.cs
+++ b/Engine/Source/Programs/UnrealBuildTool/Modes/InlineGeneratedCpps.cs
@@ -159,7 +159,7 @@ namespace UnrealBuildTool
 
 										if (Line.Contains("//", StringComparison.InvariantCultureIgnoreCase))
 										{
-											Line = Line.Substring(0, Line.IndexOf("//"));
+											Line = Line.Substring(0, Line.IndexOf("//", StringComparison.Ordinal));
 										}
 
 										if (SameLineCommentRegex.IsMatch(Line))
diff --git a/Engine/Source/Programs/UnrealBuildTool/Platform/Android/AndroidPlatformSDK.cs b/Engine/Source/Programs/UnrealBuildTool/Platform/Android/AndroidPlatformSDK.cs
--- a/Engine/Source/Programs/UnrealBuildTool/Platform/Android/AndroidPlatformSDK.cs
+++ b/Engine/Source/Programs/UnrealBuildTool/Platform/Android/AndroidPlatformSDK.cs
@@ -100,7 +100,7 @@ namespace UnrealBuildTool
 								continue;
 							}
 
-							if (BashProfileContents[LineIndex].StartsWith("export " + kvp.Key + "="))
+							if (BashProfileContents[LineIndex].StartsWith("export " + kvp.Key + "=", StringComparison.Ordinal))
 							{
 								string PathVar = BashProfileContents[LineIndex].Split('=')[1].Replace("\"", "", StringComparison.Ordinal);
 								AndroidEnv.Add(kvp.Key, PathVar);
@@ -163,7 +163,7 @@ namespace UnrealBuildTool
 				string[] PropertyContents = File.ReadAllLines(SourcePropFilename);
 				foreach (string PropertyLine in PropertyContents)
 				{
-					if (PropertyLine.StartsWith("Pkg.Revision"))
+					if (PropertyLine.StartsWith("Pkg.Revision", StringComparison.Ordinal))
 					{
 						RevisionString = PropertyLine;
 						break;
@@ -448,11 +448,11 @@ namespace UnrealBuildTool
 			string[] JavaReleaseLines = File.ReadAllLines(JavaReleaseFile);
 			foreach (string LineContents in JavaReleaseLines)
 			{
-				if (LineContents.StartsWith("JAVA_VERSION="))
+				if (LineContents.StartsWith("JAVA_VERSION=", StringComparison.Ordinal))
 				{
 					// JAVA_VERSION="17.0.6"
-					int VersionStartIndex = LineContents.IndexOf("\"");
-					int VersionStopIndex = LineContents.IndexOf(".");
+					int VersionStartIndex = LineContents.IndexOf('\"', StringComparison.Ordinal);
+					int VersionStopIndex = LineContents.IndexOf('.', StringComparison.Ordinal);
 					_ = Int32.TryParse(LineContents.Substring(VersionStartIndex, VersionStopIndex - VersionStartIndex), out JavaVersion);
 					break;
 				}
diff --git a/Engine/Source/Programs/UnrealBuildTool/Platform/Android/AndroidToolChain.cs b/Engine/Source/Programs/UnrealBuildTool/Platform/Android/AndroidToolChain.cs
--- a/Engine/Source/Programs/UnrealBuildTool/Platform/Android/AndroidToolChain.cs
+++ b/Engine/Source/Programs/UnrealBuildTool/Platform/Android/AndroidToolChain.cs
@@ -699,10 +699,10 @@ namespace UnrealBuildTool
 				if (SourceFile.AbsolutePath.Equals(GetNativeGluePath(), StringComparison.Ordinal))
 				{
 					// Remove visibility settings for android native glue. Since it doesn't decorate with visibility attributes.
-					Arguments.RemoveAll(x => x.StartsWith("-fvisibility"));
+					Arguments.RemoveAll(x => x.StartsWith("-fvisibility", StringComparison.Ordinal));
 				}
 				// remove any PCH includes - mostly for the force-added .c files in Launch as those will attempt to have the PCH used that was made with .cpp language
-				Arguments.RemoveAll(x => x.StartsWith("-include-pch"));
+				Arguments.RemoveAll(x => x.StartsWith("-include-pch", StringComparison.Ordinal));
 			}
 
 			return TargetFile;
@@ -958,7 +958,7 @@ namespace UnrealBuildTool
 		{
 			// strip any absolute path
 			string Lib = Path.GetFileNameWithoutExtension(FullLib);
-			if (Lib.StartsWith("lib"))
+			if (Lib.StartsWith("lib", StringComparison.Ordinal))
 			{
 				Lib = Lib.Substring(3);
 			}
@@ -994,7 +994,7 @@ namespace UnrealBuildTool
 				if (ComboName.Key != Arch)
 				{
 					string ArchitectureName = ComboName.Key.ToString();
-					if (Lib.EndsWith(ComboName.Value) || Lib.EndsWith(ArchitectureName))
+					if (Lib.EndsWith(ComboName.Value, StringComparison.Ordinal) || Lib.EndsWith(ArchitectureName, StringComparison.Ordinal))
 					{
 						return true;
 					}
@@ -1264,7 +1264,7 @@ namespace UnrealBuildTool
 					{
 						if (String.IsNullOrEmpty(Path.GetDirectoryName(SystemLibrary)))
 						{
-							if (SystemLibrary.StartsWith("lib"))
+							if (SystemLibrary.StartsWith("lib", StringComparison.Ordinal))
 							{
 								AdditionalLibraries.Add(SystemLibrary);
 							}
@@ -1289,7 +1289,7 @@ namespace UnrealBuildTool
 					LinkAction.PrerequisiteItems.Add(FileItem.GetItemByFileReference(Library));
 
 					string Lib = Path.GetFileNameWithoutExtension(Library.FullName);
-					if (Lib.StartsWith("lib"))
+					if (Lib.StartsWith("lib", StringComparison.Ordinal))
 					{
 						AdditionalLibraries.Add(Lib);
 						if (!AdditionalLibraryPaths.Contains(AbsoluteLibraryPath))
@@ -1339,7 +1339,7 @@ namespace UnrealBuildTool
 				LinkResponseArguments += String.Format(" -Wl,--start-group");
 				foreach (string AdditionalLibrary in AdditionalLibraries)
 				{
-					if (AdditionalLibrary.StartsWith("lib"))
+					if (AdditionalLibrary.StartsWith("lib", StringComparison.Ordinal))
 					{
 						LinkResponseArguments += String.Format(" \"-l{0}\"", AdditionalLibrary.Substring(3));
 					}
@@ -1456,7 +1456,7 @@ namespace UnrealBuildTool
 			}
 
 			// fix the executable (replace the last -ar with -strip and keep any extension)
-			int ArIndex = StripExe.LastIndexOf("-ar");
+			int ArIndex = StripExe.LastIndexOf("-ar", StringComparison.Ordinal);
 			StripExe = StripExe.Substring(0, ArIndex) + "-strip" + StripExe.Substring(ArIndex + 3);
 			return StripExe;
 		}
@@ -1499,7 +1499,7 @@ namespace UnrealBuildTool
 							string[] PropertyContents = File.ReadAllLines(SourcePropFilename);
 							foreach (string PropertyLine in PropertyContents)
 							{
-								if (PropertyLine.StartsWith("Pkg.Revision="))
+								if (PropertyLine.StartsWith("Pkg.Revision=", StringComparison.Ordinal))
 								{
 									RevisionString = PropertyLine.Substring(13);
 									RevisionValue = GetRevisionValue(RevisionString);
diff --git a/Engine/Source/Programs/UnrealBuildTool/Platform/Android/UEDeployAndroid.cs b/Engine/Source/Programs/UnrealBuildTool/Platform/Android/UEDeployAndroid.cs
--- a/Engine/Source/Programs/UnrealBuildTool/Platform/Android/UEDeployAndroid.cs
+++ b/Engine/Source/Programs/UnrealBuildTool/Platform/Android/UEDeployAndroid.cs
@@ -917,7 +917,7 @@ namespace UnrealBuildTool
 			bool bMustWriteFile = !bFileExists;
 
 			string AppType = "";
-			if (CookFlavor.EndsWith("Client"))
+			if (CookFlavor.EndsWith("Client", StringComparison.Ordinal))
 			{
 				//				AppType = ".Client";		// should always be empty now; fix up the name in batch file instead
 			}
@@ -1185,7 +1185,7 @@ namespace UnrealBuildTool
 			{
 				// ex:   0x0000000000000001 (NEEDED)         Shared library: [libGLESv3.so]
 				int ScanStart = 0;
-				int SharedIndex = Output.IndexOf("Shared library: [", ScanStart);
+				int SharedIndex = Output.IndexOf("Shared library: [", ScanStart, StringComparison.Ordinal);
 				while (SharedIndex >= 0)
 				{
 					ScanStart = SharedIndex + 17;
@@ -1195,12 +1195,12 @@ namespace UnrealBuildTool
 						string SharedLib = Output.Substring(ScanStart, EndIndex - ScanStart);
 
 						// strip off lib (not needed for System.loadLibrary()
-						if (SharedLib.StartsWith("lib"))
+						if (SharedLib.StartsWith("lib", StringComparison.Ordinal))
 						{
 							SharedLib = SharedLib.Substring(3);
 						}
 						// strip off .so (not needed for System.loadLibrary()
-						if (SharedLib.EndsWith(".so"))
+						if (SharedLib.EndsWith(".so", StringComparison.Ordinal))
 						{
 							SharedLib = SharedLib.Substring(0, SharedLib.Length - 3);
 						}
@@ -1213,7 +1213,7 @@ namespace UnrealBuildTool
 						break;
 					}
 					ScanStart = EndIndex + 1;
-					SharedIndex = Output.IndexOf("Shared library: [", ScanStart);
+					SharedIndex = Output.IndexOf("Shared library: [", ScanStart, StringComparison.Ordinal);
 				}
 			}
 			return DependencyList;
@@ -2010,7 +2010,7 @@ namespace UnrealBuildTool
 
 			// find the key
 			string AppNameTag = "<string name=\"app_name\">";
-			int KeyIndex = Contents.IndexOf(AppNameTag);
+			int KeyIndex = Contents.IndexOf(AppNameTag, StringComparison.Ordinal);
 
 			// if doesn't exist, need to repackage
 			if (KeyIndex < 0)
@@ -2020,7 +2020,7 @@ namespace UnrealBuildTool
 
 			// get the current value
 			KeyIndex += AppNameTag.Length;
-			int TagEnd = Contents.IndexOf("</string>", KeyIndex);
+			int TagEnd = Contents.IndexOf("</string>", KeyIndex, StringComparison.Ordinal);
 			if (TagEnd < 0)
 			{
 				return true;
@@ -2750,11 +2750,11 @@ namespace UnrealBuildTool
 
 			//figure out the app type
 			string AppType = InTargetType == TargetType.Game ? "" : InTargetType.ToString();
-			if (CookFlavor.EndsWith("Client"))
+			if (CookFlavor.EndsWith("Client", StringComparison.Ordinal))
 			{
 				CookFlavor = CookFlavor.Substring(0, CookFlavor.Length - 6);
 			}
-			if (CookFlavor.EndsWith("Server"))
+			if (CookFlavor.EndsWith("Server", StringComparison.Ordinal))
 			{
 				CookFlavor = CookFlavor.Substring(0, CookFlavor.Length - 6);
 			}
@@ -3276,14 +3276,14 @@ namespace UnrealBuildTool
 				{
 					++LineIndex;
 
-					int StartIndex = Line.IndexOf("\"app_id\">");
+					int StartIndex = Line.IndexOf("\"app_id\">", StringComparison.Ordinal);
 					if (StartIndex < 0)
 					{
 						continue;
 					}
 
 					StartIndex += 9;
-					int EndIndex = Line.IndexOf("</string>");
+					int EndIndex = Line.IndexOf("</string>", StringComparison.Ordinal);
 					if (EndIndex < 0)
 					{
 						continue;
@@ -3292,7 +3292,7 @@ namespace UnrealBuildTool
 					string XmlAppId = Line.Substring(StartIndex, EndIndex - StartIndex);
 
 					//validate that the AppId matches the .ini value for the GooglePlay AppId, assuming it's valid
-					if (!bInvalidIniAppId && IniAppId.CompareTo(XmlAppId) != 0)
+					if (!bInvalidIniAppId && !IniAppId.Equals(XmlAppId, StringComparison.Ordinal))
 					{
 						Logger.LogInformation("Replacing Google Play AppID in GooglePlayAppID.xml with AndroidRuntimeSettings .ini value");
 
@@ -3459,7 +3459,7 @@ namespace UnrealBuildTool
 				bFound = false;
 				for (int DriveIndex = AllDrives.Length - 1; DriveIndex >= 0; DriveIndex--)
 				{
-					if (AllDrives[DriveIndex].Name.ToUpper().StartsWith(UnusedDriveLetter))
+					if (AllDrives[DriveIndex].Name.StartsWith(UnusedDriveLetter, StringComparison.OrdinalIgnoreCase))
 					{
 						bFound = true;
 						break;
@@ -3899,7 +3899,7 @@ namespace UnrealBuildTool
 			GradleProperties.AppendLine(String.Format("VERSION_DISPLAY_NAME={0}", VersionDisplayName));
 
 			string NDKPath = Environment.GetEnvironmentVariable("NDKROOT")!.Replace("\\", "/", StringComparison.Ordinal);
-			int NDKVersionIndex = NDKPath.LastIndexOf("/");
+			int NDKVersionIndex = NDKPath.LastIndexOf('/');
 			string NDKVersion = NDKVersionIndex > 0 ? NDKPath.Substring(NDKVersionIndex + 1) : DEFAULT_NDK_VERSION;
 			GradleProperties.AppendLine(String.Format("NDK_VERSION={0}", NDKVersion));
 			GradleProperties.AppendLine(String.Format("NDK_PATH={0}", NDKPath));
@@ -4373,7 +4373,7 @@ dependencies.implementation('androidx.tracing:tracing:1.2.0')
 				if (File.Exists(BuildSettingsCacheFile))
 				{
 					const string Key = "bEnableBundle=";
-					string? Line = File.ReadAllLines(BuildSettingsCacheFile).FirstOrDefault(Line => Line.StartsWith(Key));
+					string? Line = File.ReadAllLines(BuildSettingsCacheFile).FirstOrDefault(Line => Line.StartsWith(Key, StringComparison.Ordinal));
 					if (Line != null && !Line[Key.Length..].Equals(bEnableBundle ? "True" : "False", StringComparison.Ordinal))
 					{
 						bCreateFromScratch = true;
@@ -5061,7 +5061,7 @@ dependencies.implementation('androidx.tracing:tracing:1.2.0')
 				if (OutputDirectory != null && Path.GetRelativePath(DestApkDirectory, OutputDirectory) != OutputDirectory)
 				{
 					// only override if not UnrealGame (BluePrint project)
-					if (!Path.GetFileName(OutputPath).StartsWith("UnrealGame"))
+					if (!Path.GetFileName(OutputPath).StartsWith("UnrealGame", StringComparison.Ordinal))
 					{
 						DestApkDirectory = OutputDirectory;
 					}
@@ -5328,7 +5328,7 @@ popd
 								{
 									foreach (string Taskline in Output.Split(new[] { "\r\n", "\r", "\n" }, StringSplitOptions.None))
 									{
-										if (Taskline.StartsWith("java.exe"))
+										if (Taskline.StartsWith("java.exe", StringComparison.Ordinal))
 										{
 											string ProcessString = Taskline.Substring(8).TrimStart();
 											ProcessString = ProcessString.Substring(0, ProcessString.IndexOf(' ', StringComparison.Ordinal));
@@ -5368,7 +5368,7 @@ popd
 									{
 										foreach (string Taskline in Output.Split(new[] { "\r\n", "\r", "\n" }, StringSplitOptions.None))
 										{
-											if (Taskline.StartsWith("java.exe"))
+											if (Taskline.StartsWith("java.exe", StringComparison.Ordinal))
 											{
 												string ProcessString = Taskline.Substring(8).TrimStart();
 												ProcessString = ProcessString.Substring(0, ProcessString.IndexOf(' ', StringComparison.Ordinal));
@@ -5517,19 +5517,19 @@ popd
 						}
 
 						// ignore abi.gradle, we're going to generate a new one
-						if (Filename.EndsWith("abi.gradle"))
+						if (Filename.EndsWith("abi.gradle", StringComparison.Ordinal))
 						{
 							continue;
 						}
 
 						// ignore OBBData.java, we won't use it
-						if (Filename.EndsWith("OBBData.java"))
+						if (Filename.EndsWith("OBBData.java", StringComparison.Ordinal))
 						{
 							continue;
 						}
 
 						// deal with AndroidManifest.xml
-						if (Filename.EndsWith("AndroidManifest.xml"))
+						if (Filename.EndsWith("AndroidManifest.xml", StringComparison.Ordinal))
 						{
 							// only allowed to differ by versionCode
 							string[] SourceManifest = File.ReadAllLines(Filename);
@@ -5545,24 +5545,24 @@ popd
 										continue;
 									}
 
-									int SourceVersionIndex = SourceManifest[Index].IndexOf("android:versionCode=");
+									int SourceVersionIndex = SourceManifest[Index].IndexOf("android:versionCode=", StringComparison.Ordinal);
 									if (SourceVersionIndex < 0)
 									{
 										bDiffers = true;
 										break;
 									}
 
-									int DestVersionIndex = DestManifest[Index].IndexOf("android:versionCode=");
+									int DestVersionIndex = DestManifest[Index].IndexOf("android:versionCode=", StringComparison.Ordinal);
 									if (DestVersionIndex < 0)
 									{
 										bDiffers = true;
 										break;
 									}
 
-									int SourceVersionIndex2 = SourceManifest[Index].Substring(SourceVersionIndex + 22).IndexOf("\"");
+									int SourceVersionIndex2 = SourceManifest[Index].Substring(SourceVersionIndex + 22).IndexOf('\"', StringComparison.Ordinal);
 									string FixedSource = SourceManifest[Index].Substring(0, SourceVersionIndex + 21) + SourceManifest[Index].Substring(SourceVersionIndex + 22 + SourceVersionIndex2);
 
-									int DestVersionIndex2 = DestManifest[Index].Substring(DestVersionIndex + 22).IndexOf("\"");
+									int DestVersionIndex2 = DestManifest[Index].Substring(DestVersionIndex + 22).IndexOf('\"', StringComparison.Ordinal);
 									string FixedDest = SourceManifest[Index].Substring(0, DestVersionIndex + 21) + DestManifest[Index].Substring(DestVersionIndex + 22 + DestVersionIndex2);
 
 									if (FixedSource != FixedDest)
@@ -5584,7 +5584,7 @@ popd
 						}
 
 						// deal with buildAdditions.gradle
-						if (Filename.EndsWith("buildAdditions.gradle"))
+						if (Filename.EndsWith("buildAdditions.gradle", StringComparison.Ordinal))
 						{
 							// allow store filepath to differ
 							string[] SourceProperties = File.ReadAllLines(Filename);
@@ -5621,7 +5621,7 @@ popd
 						}
 
 						// deal with gradle.properties
-						if (Filename.EndsWith("gradle.properties"))
+						if (Filename.EndsWith("gradle.properties", StringComparison.Ordinal))
 						{
 							// allow STORE_VERSION and OUTPUT_FILENAME to differ
 							// only allowed to differ by versionCode
@@ -5638,32 +5638,32 @@ popd
 										continue;
 									}
 
-									if (SourceProperties[Index].StartsWith("STORE_VERSION=") && DestProperties[Index].StartsWith("STORE_VERSION="))
+									if (SourceProperties[Index].StartsWith("STORE_VERSION=", StringComparison.Ordinal) && DestProperties[Index].StartsWith("STORE_VERSION=", StringComparison.Ordinal))
 									{
 										continue;
 									}
 
-									if (SourceProperties[Index].StartsWith("STORE_FILE=") && DestProperties[Index].StartsWith("STORE_FILE="))
+									if (SourceProperties[Index].StartsWith("STORE_FILE=", StringComparison.Ordinal) && DestProperties[Index].StartsWith("STORE_FILE=", StringComparison.Ordinal))
 									{
 										continue;
 									}
 
-									if (SourceProperties[Index].StartsWith("DEBUG_STORE_FILE=") && DestProperties[Index].StartsWith("DEBUG_STORE_FILE="))
+									if (SourceProperties[Index].StartsWith("DEBUG_STORE_FILE=", StringComparison.Ordinal) && DestProperties[Index].StartsWith("DEBUG_STORE_FILE=", StringComparison.Ordinal))
 									{
 										continue;
 									}
 
-									if (SourceProperties[Index].StartsWith("OUTPUT_FILENAME=") && DestProperties[Index].StartsWith("OUTPUT_FILENAME="))
+									if (SourceProperties[Index].StartsWith("OUTPUT_FILENAME=", StringComparison.Ordinal) && DestProperties[Index].StartsWith("OUTPUT_FILENAME=", StringComparison.Ordinal))
 									{
 										continue;
 									}
 
-									if (SourceProperties[Index].StartsWith("OUTPUT_BUNDLEFILENAME=") && DestProperties[Index].StartsWith("OUTPUT_BUNDLEFILENAME="))
+									if (SourceProperties[Index].StartsWith("OUTPUT_BUNDLEFILENAME=", StringComparison.Ordinal) && DestProperties[Index].StartsWith("OUTPUT_BUNDLEFILENAME=", StringComparison.Ordinal))
 									{
 										continue;
 									}
 
-									if (SourceProperties[Index].StartsWith("OUTPUT_UNIVERSALFILENAME=") && DestProperties[Index].StartsWith("OUTPUT_UNIVERSALFILENAME="))
+									if (SourceProperties[Index].StartsWith("OUTPUT_UNIVERSALFILENAME=", StringComparison.Ordinal) && DestProperties[Index].StartsWith("OUTPUT_UNIVERSALFILENAME=", StringComparison.Ordinal))
 									{
 										continue;
 									}
@@ -5714,13 +5714,13 @@ popd
 					string[] ManifestContents = File.ReadAllLines(ManifestFilename);
 					for (int Index = 0; Index < ManifestContents.Length; Index++)
 					{
-						int ManifestVersionIndex = ManifestContents[Index].IndexOf("android:versionCode=");
+						int ManifestVersionIndex = ManifestContents[Index].IndexOf("android:versionCode=", StringComparison.Ordinal);
 						if (ManifestVersionIndex < 0)
 						{
 							continue;
 						}
 
-						int ManifestVersionIndex2 = ManifestContents[Index].Substring(ManifestVersionIndex + 22).IndexOf("\"");
+						int ManifestVersionIndex2 = ManifestContents[Index].Substring(ManifestVersionIndex + 22).IndexOf('\"', StringComparison.Ordinal);
 						ManifestContents[Index] = ManifestContents[Index].Substring(0, ManifestVersionIndex + 21) + BaseStoreVersion + ManifestContents[Index].Substring(ManifestVersionIndex + 22 + ManifestVersionIndex2);
 						break;
 					}
@@ -5988,11 +5988,11 @@ popd
 			//Logger.LogInformation("$$$$$$$$$$$$$$ PrepForUATPackageOrDeploy $$$$$$$$$$$$$$$$$");
 
 			TargetType Type = TargetType.Game;
-			if (CookFlavor.EndsWith("Client"))
+			if (CookFlavor.EndsWith("Client", StringComparison.Ordinal))
 			{
 				Type = TargetType.Client;
 			}
-			else if (CookFlavor.EndsWith("Server"))
+			else if (CookFlavor.EndsWith("Server", StringComparison.Ordinal))
 			{
 				Type = TargetType.Server;
 			}
@@ -6219,7 +6219,7 @@ import java.util.Collection;
 
 				if (!bJavaPerfettoMarkers)
 				{
-					if (SrcLine.StartsWith("Trace."))
+					if (SrcLine.StartsWith("Trace.", StringComparison.Ordinal))
 					{
 						TemplateSrc[LineIndex] = "";
 						Changed = true;
@@ -6239,13 +6239,13 @@ import java.util.Collection;
 							}
 							else
 							{
-								int TraceStartIndex = KVP.Value.IndexOf("Trace.");
+								int TraceStartIndex = KVP.Value.IndexOf("Trace.", StringComparison.Ordinal);
 								if (TraceStartIndex == 0 || (TraceStartIndex > 0 && KVP.Value[TraceStartIndex - 1] == '\n'))
 								{
 									string Addition = KVP.Value;
 									do
 									{
-										int TraceEndIndex = Addition.IndexOf("\n", TraceStartIndex);
+										int TraceEndIndex = Addition.IndexOf("\n", TraceStartIndex, StringComparison.Ordinal);
 										if (TraceEndIndex < 0)
 										{
 											Addition = Addition.Substring(0, TraceStartIndex);
@@ -6254,7 +6254,7 @@ import java.util.Collection;
 										{
 											Addition = Addition.Substring(0, TraceStartIndex) + Addition.Substring(TraceEndIndex + 1);
 										}
-										TraceStartIndex = Addition.IndexOf("Trace.", TraceStartIndex);
+										TraceStartIndex = Addition.IndexOf("Trace.", TraceStartIndex, StringComparison.Ordinal);
 									} while (TraceStartIndex >= 0);
 									SrcLine = SrcLine.Replace(KVP.Key, Addition, StringComparison.Ordinal);
 								}
@@ -6342,7 +6342,7 @@ import java.util.Collection;
 
 				if (!bJavaPerfettoMarkers)
 				{
-					if (SrcLine.StartsWith("Trace."))
+					if (SrcLine.StartsWith("Trace.", StringComparison.Ordinal))
 					{
 						TemplateSrc[LineIndex] = "";
 						Changed = true;
@@ -6362,13 +6362,13 @@ import java.util.Collection;
 							}
 							else
 							{
-								int TraceStartIndex = KVP.Value.IndexOf("Trace.");
+								int TraceStartIndex = KVP.Value.IndexOf("Trace.", StringComparison.Ordinal);
 								if (TraceStartIndex == 0 || (TraceStartIndex > 0 && KVP.Value[TraceStartIndex - 1] == '\n'))
 								{
 									string Addition = KVP.Value;
 									do
 									{
-										int TraceEndIndex = Addition.IndexOf("\n", TraceStartIndex);
+										int TraceEndIndex = Addition.IndexOf("\n", TraceStartIndex, StringComparison.Ordinal);
 										if (TraceEndIndex < 0)
 										{
 											Addition = Addition.Substring(0, TraceStartIndex);
@@ -6377,7 +6377,7 @@ import java.util.Collection;
 										{
 											Addition = Addition.Substring(0, TraceStartIndex) + Addition.Substring(TraceEndIndex + 1);
 										}
-										TraceStartIndex = Addition.IndexOf("Trace.", TraceStartIndex);
+										TraceStartIndex = Addition.IndexOf("Trace.", TraceStartIndex, StringComparison.Ordinal);
 									} while (TraceStartIndex >= 0);
 									SrcLine = SrcLine.Replace(KVP.Key, Addition, StringComparison.Ordinal);
 								}
@@ -6477,7 +6477,7 @@ import java.util.Collection;
 			{
 				string Trimmed = Line.Trim(' ', '\r');
 
-				if (Trimmed.StartsWith("repository "))
+				if (Trimmed.StartsWith("repository ", StringComparison.Ordinal))
 				{
 					string DirectoryPath = Trimmed.Substring(11).Trim(' ').TrimEnd('/', '\\');
 					DirectoryPath = DirectoryPath.Replace("$(ENGINEDIR)", EngineDir, StringComparison.Ordinal);
@@ -6485,7 +6485,7 @@ import java.util.Collection;
 					DirectoryPath = DirectoryPath.Replace('\\', Path.DirectorySeparatorChar).Replace('/', Path.DirectorySeparatorChar);
 					AARHandler.AddRepository(DirectoryPath, Logger);
 				}
-				else if (Trimmed.StartsWith("repositories "))
+				else if (Trimmed.StartsWith("repositories ", StringComparison.Ordinal))
 				{
 					string DirectoryPath = Trimmed.Substring(13).Trim(' ').TrimEnd('/', '\\');
 					DirectoryPath = DirectoryPath.Replace("$(ENGINEDIR)", EngineDir, StringComparison.Ordinal);
diff --git a/Engine/Source/Programs/UnrealBuildTool/Platform/Apple/AppleOnWindowsAppMaker.cs b/Engine/Source/Programs/UnrealBuildTool/Platform/Apple/AppleOnWindowsAppMaker.cs
--- a/Engine/Source/Programs/UnrealBuildTool/Platform/Apple/AppleOnWindowsAppMaker.cs
+++ b/Engine/Source/Programs/UnrealBuildTool/Platform/Apple/AppleOnWindowsAppMaker.cs
@@ -963,11 +963,11 @@ namespace UnrealBuildTool
 			// write the iCloud container identifier, if present in the old file
 			if (!String.IsNullOrEmpty(OldPListData))
 			{
-				int index = OldPListData.IndexOf("ICloudContainerIdentifier");
+				int index = OldPListData.IndexOf("ICloudContainerIdentifier", StringComparison.Ordinal);
 				if (index > 0)
 				{
-					index = OldPListData.IndexOf("<string>", index) + 8;
-					int length = OldPListData.IndexOf("</string>", index) - index;
+					index = OldPListData.IndexOf("<string>", index, StringComparison.Ordinal) + 8;
+					int length = OldPListData.IndexOf("</string>", index, StringComparison.Ordinal) - index;
 					string ICloudContainerIdentifier = OldPListData.Substring(index, length);
 					Text.AppendLine("\t<key>ICloudContainerIdentifier</key>");
 					Text.AppendLine(String.Format("\t<string>{0}</string>", ICloudContainerIdentifier));
@@ -1463,11 +1463,11 @@ namespace UnrealBuildTool
 			// write the iCloud container identifier, if present in the old file
 			if (!String.IsNullOrEmpty(OldPListData))
 			{
-				int index = OldPListData.IndexOf("ICloudContainerIdentifier");
+				int index = OldPListData.IndexOf("ICloudContainerIdentifier", StringComparison.Ordinal);
 				if (index > 0)
 				{
-					index = OldPListData.IndexOf("<string>", index) + 8;
-					int length = OldPListData.IndexOf("</string>", index) - index;
+					index = OldPListData.IndexOf("<string>", index, StringComparison.Ordinal) + 8;
+					int length = OldPListData.IndexOf("</string>", index, StringComparison.Ordinal) - index;
 					string ICloudContainerIdentifier = OldPListData.Substring(index, length);
 					Text.AppendLine("\t<key>ICloudContainerIdentifier</key>");
 					Text.AppendLine(String.Format("\t<string>{0}</string>", ICloudContainerIdentifier));
@@ -1507,7 +1507,7 @@ namespace UnrealBuildTool
 			}
 			else if (Target.Platform == UnrealTargetPlatform.IOS)
 			{
-				int Index = Binary.OutputFilePath.GetFileNameWithoutExtension().IndexOf("-");
+				int Index = Binary.OutputFilePath.GetFileNameWithoutExtension().IndexOf("-", StringComparison.Ordinal);
 				string OutputFile = Binary.OutputFilePath.GetFileNameWithoutExtension().Substring(0, Index > 0 ? Index : Binary.OutputFilePath.GetFileNameWithoutExtension().Length);
 				FileReference AssetFile = FileReference.Combine(Binary.OutputFilePath.Directory, "Payload", OutputFile + ".app", "Assets.car");
 				BuildProducts.Add(AssetFile, BuildProductType.RequiredResource);
@@ -1605,7 +1605,7 @@ namespace UnrealBuildTool
 						if (Framework.bCopyFramework)
 						{
 							string FrameworkDir = FrameworkDirectory.FullName;
-							if (FrameworkDir.EndsWith(".framework"))
+							if (FrameworkDir.EndsWith(".framework", StringComparison.Ordinal))
 							{
 								FrameworkDir = Path.GetDirectoryName(FrameworkDir)!;
 							}
diff --git a/Engine/Source/Programs/UnrealBuildTool/Platform/Apple/ApplePlatformSDK.cs b/Engine/Source/Programs/UnrealBuildTool/Platform/Apple/ApplePlatformSDK.cs
--- a/Engine/Source/Programs/UnrealBuildTool/Platform/Apple/ApplePlatformSDK.cs
+++ b/Engine/Source/Programs/UnrealBuildTool/Platform/Apple/ApplePlatformSDK.cs
@@ -216,7 +216,7 @@ namespace UnrealBuildTool
 				foreach (DirectoryReference SubDir in DirectoryReference.EnumerateDirectories(GetPlatformSDKDirectory(AppleOSName).ParentDirectory!))
 				{
 					string SubDirName = Path.GetFileNameWithoutExtension(SubDir.GetDirectoryName());
-					if (SubDirName.StartsWith(AppleOSName))
+					if (SubDirName.StartsWith(AppleOSName, StringComparison.Ordinal))
 					{
 						// get the SDK version from the directory name
 						string SDKString = SubDirName.Replace(AppleOSName, "", StringComparison.Ordinal);
diff --git a/Engine/Source/Programs/UnrealBuildTool/Platform/Apple/AppleProjectGenerator.cs b/Engine/Source/Programs/UnrealBuildTool/Platform/Apple/AppleProjectGenerator.cs
--- a/Engine/Source/Programs/UnrealBuildTool/Platform/Apple/AppleProjectGenerator.cs
+++ b/Engine/Source/Programs/UnrealBuildTool/Platform/Apple/AppleProjectGenerator.cs
@@ -120,21 +120,21 @@ namespace UnrealBuildTool
 				{
 					if (IsReadingIncludesSection)
 					{
-						if (Args.Data.StartsWith("End of search"))
+						if (Args.Data.StartsWith("End of search", StringComparison.Ordinal))
 						{
 							IsReadingIncludesSection = false;
 							XcrunProcess!.Kill();
 						}
 						else
 						{
-							if (!Args.Data.EndsWith("(framework directory)"))
+							if (!Args.Data.EndsWith("(framework directory)", StringComparison.Ordinal))
 							{
 								CachedIncludePaths[CurrentlyProcessedSDK].Add(Args.Data.Trim(' ', '"'));
 							}
 						}
 					}
 
-					if (Args.Data.StartsWith("#include <...>"))
+					if (Args.Data.StartsWith("#include <...>", StringComparison.Ordinal))
 					{
 						IsReadingIncludesSection = true;
 					}
diff --git a/Engine/Source/Programs/UnrealBuildTool/Platform/Apple/AppleToolChainSettings.cs b/Engine/Source/Programs/UnrealBuildTool/Platform/Apple/AppleToolChainSettings.cs
--- a/Engine/Source/Programs/UnrealBuildTool/Platform/Apple/AppleToolChainSettings.cs
+++ b/Engine/Source/Programs/UnrealBuildTool/Platform/Apple/AppleToolChainSettings.cs
@@ -156,7 +156,7 @@ namespace UnrealBuildTool
 					"Either way, you will need to enter your Mac password.");
 			}
 
-			if (bVerbose && !XcodeDeveloperDir.FullName.StartsWith("/Applications/Xcode.app"))
+			if (bVerbose && !XcodeDeveloperDir.FullName.StartsWith("/Applications/Xcode.app", StringComparison.Ordinal))
 			{
 				Log.TraceInformationOnce("Compiling with non-standard Xcode: {0}", XcodeDeveloperDir);
 			}
diff --git a/Engine/Source/Programs/UnrealBuildTool/Platform/Apple/RemoteMac.cs b/Engine/Source/Programs/UnrealBuildTool/Platform/Apple/RemoteMac.cs
--- a/Engine/Source/Programs/UnrealBuildTool/Platform/Apple/RemoteMac.cs
+++ b/Engine/Source/Programs/UnrealBuildTool/Platform/Apple/RemoteMac.cs
@@ -1395,7 +1395,7 @@ namespace UnrealBuildTool
 				string Path = Match.Value;
 				foreach (RemoteMapping Mapping in Mappings)
 				{
-					if (Path.StartsWith(Mapping.RemoteDirectory))
+					if (Path.StartsWith(Mapping.RemoteDirectory, StringComparison.Ordinal))
 					{
 						Path = Mapping.LocalDirectory + Path[Mapping.RemoteDirectory.Length..].Replace('/', '\\');
 						break;
diff --git a/Engine/Source/Programs/UnrealBuildTool/Platform/Clang/ClangSpecificFileAction.cs b/Engine/Source/Programs/UnrealBuildTool/Platform/Clang/ClangSpecificFileAction.cs
--- a/Engine/Source/Programs/UnrealBuildTool/Platform/Clang/ClangSpecificFileAction.cs
+++ b/Engine/Source/Programs/UnrealBuildTool/Platform/Clang/ClangSpecificFileAction.cs
@@ -67,7 +67,7 @@ namespace UnrealBuildTool
 			string DummyName = "SingleFile.cpp";
 			string UniqueDummyName = Filename;
 
-			int FileNameIndex = CommandArguments.IndexOf(DummyName);
+			int FileNameIndex = CommandArguments.IndexOf(DummyName, System.StringComparison.Ordinal);
 			string DummyPath = CommandArguments.Substring(2, FileNameIndex + DummyName.Length - 2);
 
 			List<string> NewRspLines = new();
diff --git a/Engine/Source/Programs/UnrealBuildTool/Platform/IOS/IOSCodeSigning.cs b/Engine/Source/Programs/UnrealBuildTool/Platform/IOS/IOSCodeSigning.cs
--- a/Engine/Source/Programs/UnrealBuildTool/Platform/IOS/IOSCodeSigning.cs
+++ b/Engine/Source/Programs/UnrealBuildTool/Platform/IOS/IOSCodeSigning.cs
@@ -40,7 +40,7 @@ namespace UnrealBuildTool
 			foreach (string Part in NameParts)
 			{
 				string CleanPart = Part.Trim();
-				if (CleanPart.StartsWith("CN="))
+				if (CleanPart.StartsWith("CN=", StringComparison.Ordinal))
 				{
 					CommonName = CleanPart.Substring(3);
 				}
@@ -113,16 +113,16 @@ namespace UnrealBuildTool
 					{
 						string header = "-----BEGIN CERTIFICATE-----\n";
 						string footer = "-----END CERTIFICATE-----";
-						int start = CertToolData.IndexOf(header);
+						int start = CertToolData.IndexOf(header, StringComparison.Ordinal);
 						while (start != -1)
 						{
 							start += header.Length;
-							int end = CertToolData.IndexOf(footer, start);
+							int end = CertToolData.IndexOf(footer, start, StringComparison.Ordinal);
 							string base64 = CertToolData.Substring(start, (end - start));
 							byte[] certData = Convert.FromBase64String(base64);
 							X509Certificate2 cert = new X509Certificate2(certData);
 							FoundCerts.Add(cert);
-							start = CertToolData.IndexOf(header, start);
+							start = CertToolData.IndexOf(header, start, StringComparison.Ordinal);
 						}
 					}
 				}
@@ -219,8 +219,8 @@ namespace UnrealBuildTool
 
 					if (CompanyName != "*")
 					{
-						CompanyName = CompanyName.Substring(0, CompanyName.LastIndexOf("."));
-						bPassesCompanyCheck = CFBundleIdentifier.StartsWith(CompanyName);
+						CompanyName = CompanyName.Substring(0, CompanyName.LastIndexOf('.'));
+						bPassesCompanyCheck = CFBundleIdentifier.StartsWith(CompanyName, StringComparison.Ordinal);
 					}
 					else
 					{
@@ -975,8 +975,8 @@ namespace UnrealBuildTool
 							string CompanyName = TestProvision.ApplicationIdentifier.Substring(TestProvision.ApplicationIdentifierPrefix.Length + 1);
 							if (CompanyName != "*")
 							{
-								CompanyName = CompanyName.Substring(0, CompanyName.LastIndexOf("."));
-								bPassesNameCheck = CFBundleIdentifier.StartsWith(CompanyName);
+								CompanyName = CompanyName.Substring(0, CompanyName.LastIndexOf('.'));
+								bPassesNameCheck = CFBundleIdentifier.StartsWith(CompanyName, StringComparison.Ordinal);
 							}
 						}
 					}
@@ -1159,7 +1159,7 @@ namespace UnrealBuildTool
 						string UbiquityKvstoreString;
 						XCentPList.GetString("com.apple.developer.ubiquity-kvstore-identifier", out UbiquityKvstoreString);
 
-						int DotPosition = UbiquityKvstoreString.LastIndexOf("*");
+						int DotPosition = UbiquityKvstoreString.LastIndexOf('*');
 						if (DotPosition >= 0)
 						{
 							string TeamPrefix = DotPosition > 1 ? UbiquityKvstoreString.Substring(0, DotPosition - 1) : TeamIdentifier;
@@ -1175,7 +1175,7 @@ namespace UnrealBuildTool
 
 						for (int i = 0; i < UbiquityContainerIdentifiersGroups.Count; i++)
 						{
-							int DotPosition = UbiquityContainerIdentifiersGroups[i].LastIndexOf("*");
+							int DotPosition = UbiquityContainerIdentifiersGroups[i].LastIndexOf('*');
 							if (DotPosition >= 0)
 							{
 								string TeamPrefix = DotPosition > 1 ? UbiquityContainerIdentifiersGroups[i].Substring(0, DotPosition - 1) : TeamIdentifier;
diff --git a/Engine/Source/Programs/UnrealBuildTool/Platform/IOS/IOSToolChain.cs b/Engine/Source/Programs/UnrealBuildTool/Platform/IOS/IOSToolChain.cs
--- a/Engine/Source/Programs/UnrealBuildTool/Platform/IOS/IOSToolChain.cs
+++ b/Engine/Source/Programs/UnrealBuildTool/Platform/IOS/IOSToolChain.cs
@@ -182,7 +182,7 @@ namespace UnrealBuildTool
 					{
 						// embedded frameworks have a framework inside of this directory, so we use this directory. regular frameworks need to go one up to point to the 
 						// directory containing the framework. -F gives a path to look for the -framework
-						if (FrameworkDirectory.FullName.EndsWith(".framework"))
+						if (FrameworkDirectory.FullName.EndsWith(".framework", StringComparison.Ordinal))
 						{
 							FrameworkDirectory = FrameworkDirectory.ParentDirectory;
 						}
@@ -313,7 +313,7 @@ namespace UnrealBuildTool
 
 					// embedded frameworks have a framework inside of this directory, so we use this directory. regular frameworks need to go one up to point to the 
 					// directory containing the framework. -F gives a path to look for the -framework
-					if (FrameworkDirectory.FullName.EndsWith(".framework"))
+					if (FrameworkDirectory.FullName.EndsWith(".framework", StringComparison.Ordinal))
 					{
 						FrameworkDirectory = FrameworkDirectory.ParentDirectory!;
 					}
@@ -849,7 +849,7 @@ namespace UnrealBuildTool
 
 					if (File.Exists(Image))
 					{
-						bUserImagesExist |= Image.StartsWith(BuildResourcesGraphicsDir);
+						bUserImagesExist |= Image.StartsWith(BuildResourcesGraphicsDir, StringComparison.Ordinal);
 
 						File.Copy(Image, Path.Combine(Dir, Images[Index][1], Images[Index][0]), true);
 						FileInfo DestFileInfo = new FileInfo(Path.Combine(Dir, Images[Index][1], Images[Index][0]));
@@ -907,7 +907,7 @@ namespace UnrealBuildTool
 					string Image = Path.Combine((Directory.Exists(Path.Combine(BuildDir, "Resources", "Graphics")) ? (BuildDir) : (Path.Combine(EngineDir, "Build", "IOS"))), "Resources", "Graphics", Images[Index][1]);
 					if (File.Exists(Image))
 					{
-						bUserImagesExist |= Image.StartsWith(BuildResourcesGraphicsDir);
+						bUserImagesExist |= Image.StartsWith(BuildResourcesGraphicsDir, StringComparison.Ordinal);
 
 						File.Copy(Image, Path.Combine(Dir, Images[Index][0]), true);
 						FileInfo DestFileInfo = new FileInfo(Path.Combine(Dir, Images[Index][0]));
diff --git a/Engine/Source/Programs/UnrealBuildTool/Platform/IOS/UEBuildIOS.cs b/Engine/Source/Programs/UnrealBuildTool/Platform/IOS/UEBuildIOS.cs
--- a/Engine/Source/Programs/UnrealBuildTool/Platform/IOS/UEBuildIOS.cs
+++ b/Engine/Source/Programs/UnrealBuildTool/Platform/IOS/UEBuildIOS.cs
@@ -361,7 +361,7 @@ namespace UnrealBuildTool
 						return "18.0";
 					case "IOS_Minimum":
 					default:
-						if (String.IsNullOrEmpty(MinimumIOSVersion) || MinimumIOSVersion.StartsWith("IOS_"))
+						if (String.IsNullOrEmpty(MinimumIOSVersion) || MinimumIOSVersion.StartsWith("IOS_", StringComparison.Ordinal))
 						{
 							return UEBuildPlatformSDK.GetSDKForPlatform("IOS")!.GetSoftwareInfo()!.Min!;
 						}
@@ -542,45 +542,45 @@ namespace UnrealBuildTool
 					Buffer.BlockCopy(AllBytes, (int)StartIndex, TextBytes, 0, (int)(EndIndex - StartIndex));
 
 					string AllText = Encoding.UTF8.GetString(TextBytes);
-					int idx = AllText.IndexOf("<key>UUID</key>");
+					int idx = AllText.IndexOf("<key>UUID</key>", StringComparison.Ordinal);
 					if (idx > 0)
 					{
-						idx = AllText.IndexOf("<string>", idx);
+						idx = AllText.IndexOf("<string>", idx, StringComparison.Ordinal);
 						if (idx > 0)
 						{
 							idx += "<string>".Length;
-							MobileProvisionUUID = AllText.Substring(idx, AllText.IndexOf("</string>", idx) - idx);
+							MobileProvisionUUID = AllText.Substring(idx, AllText.IndexOf("</string>", idx, StringComparison.Ordinal) - idx);
 						}
 					}
-					idx = AllText.IndexOf("<key>com.apple.developer.team-identifier</key>");
+					idx = AllText.IndexOf("<key>com.apple.developer.team-identifier</key>", StringComparison.Ordinal);
 					if (idx > 0)
 					{
-						idx = AllText.IndexOf("<string>", idx);
+						idx = AllText.IndexOf("<string>", idx, StringComparison.Ordinal);
 						if (idx > 0)
 						{
 							idx += "<string>".Length;
-							TeamUUID = AllText.Substring(idx, AllText.IndexOf("</string>", idx) - idx);
+							TeamUUID = AllText.Substring(idx, AllText.IndexOf("</string>", idx, StringComparison.Ordinal) - idx);
 						}
 					}
-					idx = AllText.IndexOf("<key>application-identifier</key>");
+					idx = AllText.IndexOf("<key>application-identifier</key>", StringComparison.Ordinal);
 					if (idx > 0)
 					{
-						idx = AllText.IndexOf("<string>", idx);
+						idx = AllText.IndexOf("<string>", idx, StringComparison.Ordinal);
 						if (idx > 0)
 						{
 							idx += "<string>".Length;
-							string FullID = AllText.Substring(idx, AllText.IndexOf("</string>", idx) - idx);
+							string FullID = AllText.Substring(idx, AllText.IndexOf("</string>", idx, StringComparison.Ordinal) - idx);
 							BundleIdentifier = FullID.Substring(FullID.IndexOf('.', StringComparison.Ordinal) + 1);
 						}
 					}
-					idx = AllText.IndexOf("<key>Name</key>");
+					idx = AllText.IndexOf("<key>Name</key>", StringComparison.Ordinal);
 					if (idx > 0)
 					{
-						idx = AllText.IndexOf("<string>", idx);
+						idx = AllText.IndexOf("<string>", idx, StringComparison.Ordinal);
 						if (idx > 0)
 						{
 							idx += "<string>".Length;
-							MobileProvisionName = AllText.Substring(idx, AllText.IndexOf("</string>", idx) - idx);
+							MobileProvisionName = AllText.Substring(idx, AllText.IndexOf("</string>", idx, StringComparison.Ordinal) - idx);
 						}
 					}
 				}
diff --git a/Engine/Source/Programs/UnrealBuildTool/Platform/IOS/UEDeployIOS.cs b/Engine/Source/Programs/UnrealBuildTool/Platform/IOS/UEDeployIOS.cs
--- a/Engine/Source/Programs/UnrealBuildTool/Platform/IOS/UEDeployIOS.cs
+++ b/Engine/Source/Programs/UnrealBuildTool/Platform/IOS/UEDeployIOS.cs
@@ -151,11 +151,11 @@ namespace UnrealBuildTool
 				string CFBundleVersion;
 				if (!Unreal.IsBuildMachine())
 				{
-					int Index = OldPList.IndexOf("CFBundleVersion");
+					int Index = OldPList.IndexOf("CFBundleVersion", StringComparison.Ordinal);
 					if (Index != -1)
 					{
-						int Start = OldPList.IndexOf("<string>", Index) + ("<string>").Length;
-						CFBundleVersion = OldPList.Substring(Start, OldPList.IndexOf("</string>", Index) - Start);
+						int Start = OldPList.IndexOf("<string>", Index, StringComparison.Ordinal) + ("<string>").Length;
+						CFBundleVersion = OldPList.Substring(Start, OldPList.IndexOf("</string>", Index, StringComparison.Ordinal) - Start);
 						CFBundleVersion = CalculateUpdatedMinorVersionString(CFBundleVersion);
 					}
 					else
diff --git a/Engine/Source/Programs/UnrealBuildTool/Platform/Linux/LinuxToolChain.cs b/Engine/Source/Programs/UnrealBuildTool/Platform/Linux/LinuxToolChain.cs
--- a/Engine/Source/Programs/UnrealBuildTool/Platform/Linux/LinuxToolChain.cs
+++ b/Engine/Source/Programs/UnrealBuildTool/Platform/Linux/LinuxToolChain.cs
@@ -1097,7 +1097,7 @@ namespace UnrealBuildTool
 					}
 
 					// On Windows, MakeRelativeTo can silently fail if the engine and the project are located on different drives
-					if (CrossCompiling() && RelativePath.StartsWith(Unreal.RootDirectory.FullName))
+					if (CrossCompiling() && RelativePath.StartsWith(Unreal.RootDirectory.FullName, StringComparison.Ordinal))
 					{
 						// do not replace directly, but take care to avoid potential double slashes or missed slashes
 						string PathFromRootDir = RelativePath.Replace(Unreal.RootDirectory.FullName, "", StringComparison.Ordinal);
@@ -1140,7 +1140,7 @@ namespace UnrealBuildTool
 				}
 
 				// On Windows, MakeRelativeTo can silently fail if the engine and the project are located on different drives
-				if (CrossCompiling() && RelativePath.StartsWith(Unreal.RootDirectory.FullName))
+				if (CrossCompiling() && RelativePath.StartsWith(Unreal.RootDirectory.FullName, StringComparison.Ordinal))
 				{
 					// do not replace directly, but take care to avoid potential double slashes or missed slashes
 					string PathFromRootDir = RelativePath.Replace(Unreal.RootDirectory.FullName, "", StringComparison.Ordinal);
@@ -1212,7 +1212,7 @@ namespace UnrealBuildTool
 					}
 
 					string LibName = Path.GetFileNameWithoutExtension(AdditionalLibrary);
-					if (LibName.StartsWith("lib"))
+					if (LibName.StartsWith("lib", StringComparison.Ordinal))
 					{
 						// Remove lib prefix
 						LibName = LibName[3..];
diff --git a/Engine/Source/Programs/UnrealBuildTool/Platform/Linux/UEBuildLinux.cs b/Engine/Source/Programs/UnrealBuildTool/Platform/Linux/UEBuildLinux.cs
--- a/Engine/Source/Programs/UnrealBuildTool/Platform/Linux/UEBuildLinux.cs
+++ b/Engine/Source/Programs/UnrealBuildTool/Platform/Linux/UEBuildLinux.cs
@@ -358,7 +358,7 @@ namespace UnrealBuildTool
 		/// <returns>True if the string matches the name of a build product, false otherwise</returns>
 		public override bool IsBuildProduct(string FileName, string[] NamePrefixes, string[] NameSuffixes)
 		{
-			if (FileName.StartsWith("lib"))
+			if (FileName.StartsWith("lib", StringComparison.Ordinal))
 			{
 				return IsBuildProductName(FileName, 3, FileName.Length - 3, NamePrefixes, NameSuffixes, ".a")
 					|| IsBuildProductName(FileName, 3, FileName.Length - 3, NamePrefixes, NameSuffixes, ".so")
diff --git a/Engine/Source/Programs/UnrealBuildTool/Platform/Mac/MacToolChain.cs b/Engine/Source/Programs/UnrealBuildTool/Platform/Mac/MacToolChain.cs
--- a/Engine/Source/Programs/UnrealBuildTool/Platform/Mac/MacToolChain.cs
+++ b/Engine/Source/Programs/UnrealBuildTool/Platform/Mac/MacToolChain.cs
@@ -222,7 +222,7 @@ namespace UnrealBuildTool
 
 		void AddFrameworkToLinkCommand(LinkEnvironment linkEnvironment, List<string> LinkArguments, string FrameworkName, string Arg = "-framework")
 		{
-			if (FrameworkName.EndsWith(".framework"))
+			if (FrameworkName.EndsWith(".framework", StringComparison.Ordinal))
 			{
 				FileReference frameworkPath = new FileReference(FrameworkName);
 				LinkArguments.Add($"-F \"{NormalizeCommandLinePath(frameworkPath.Directory, linkEnvironment.RootPaths)}\"");
@@ -321,7 +321,7 @@ namespace UnrealBuildTool
 			string ExeDir = Path.GetDirectoryName(ExeAbsolutePath)!;
 
 			// Only dylibs and frameworks, and only those that are outside of Engine/Binaries/Mac and Engine/Source/ThirdParty, and outside of the folder where the executable is need an additional RPATH entry
-			if ((Library.EndsWith("dylib") || Library.EndsWith(".framework"))
+			if ((Library.EndsWith("dylib", StringComparison.Ordinal) || Library.EndsWith(".framework", StringComparison.Ordinal))
 				&& !LibraryFullPath.Contains("/Engine/Source/ThirdParty/", StringComparison.Ordinal) && LibraryDir != ExeDir && !RPaths.Contains(LibraryDir))
 			{
 				// macOS gatekeeper erroneously complains about not seeing the CEF3 framework in the codesigned Launcher because it's only present in one of the folders specified in RPATHs.
@@ -344,7 +344,7 @@ namespace UnrealBuildTool
 					// To be safe leave the previous relative loader_path in place and add a relative engine executable path
 					if (!bIsBuildingAppBundle && LibraryDir.Contains("/Engine/Plugins/", StringComparison.Ordinal) && ExeAbsolutePath.EndsWith("dylib", StringComparison.Ordinal))
 					{
-						int EngineDirStrIndex = LibraryDir.IndexOf("Engine");
+						int EngineDirStrIndex = LibraryDir.IndexOf("Engine", StringComparison.Ordinal);
 						if (EngineDirStrIndex >= 0)
 						{
 							LinkArguments.Add("-rpath \"@executable_path/../../../../../../" + LibraryDir.Substring(EngineDirStrIndex) + "\"");
@@ -361,11 +361,11 @@ namespace UnrealBuildTool
 					// In packaged games dylibs are stored in Contents/UE subfolders, for example in GameName.app/Contents/UE/Engine/Binaries/ThirdParty/PhysX/Mac
 					string BundleUEDir = Path.GetFullPath(ExeDir + "/../../Contents/UE");
 					string? BundleLibraryDir = null;
-					if (LibraryDir.StartsWith(EngineDir + "/"))
+					if (LibraryDir.StartsWith(EngineDir + "/", StringComparison.Ordinal))
 					{
 						BundleLibraryDir = LibraryDir.Replace(EngineDir, BundleUEDir, StringComparison.Ordinal);
 					}
-					else if (ProjectDir != null && LibraryDir.StartsWith(ProjectDir + "/"))
+					else if (ProjectDir != null && LibraryDir.StartsWith(ProjectDir + "/", StringComparison.Ordinal))
 					{
 						string GameName = GetGameNameFromExecutablePath(ExeAbsolutePath);
 						BundleLibraryDir = LibraryDir.Replace(ProjectDir, BundleUEDir + "/" + GameName, StringComparison.Ordinal);
@@ -496,7 +496,7 @@ namespace UnrealBuildTool
 			{
 				AppName = DylibFilename.Substring(0, firstDash);
 			}
-			else if (DylibFilename.StartsWith("lib")) // This is expected to exist. if there is no dash involved this is the target lib itself and name is part of the name
+			else if (DylibFilename.StartsWith("lib", StringComparison.Ordinal)) // This is expected to exist. if there is no dash involved this is the target lib itself and name is part of the name
 			{
 				AppName = Path.GetFileNameWithoutExtension(DylibFilename).Substring(3);
 			}
@@ -732,7 +732,7 @@ namespace UnrealBuildTool
 				{
 					string MapFileBaseName = OutputFile.AbsolutePath;
 
-					int AppIdx = MapFileBaseName.IndexOf(".app/Contents/MacOS");
+					int AppIdx = MapFileBaseName.IndexOf(".app/Contents/MacOS", StringComparison.Ordinal);
 					if (AppIdx != -1)
 					{
 						MapFileBaseName = MapFileBaseName.Substring(0, AppIdx);
@@ -1028,7 +1028,8 @@ namespace UnrealBuildTool
 				DebugInfoFiles.Add(GenerateDebugInfo(Executable, BinaryLinkEnvironment, Graph));
 			}
 
-			if ((BinaryLinkEnvironment.bIsBuildingDLL && (Options & ClangToolChainOptions.OutputDylib) == 0) || (BinaryLinkEnvironment.bIsBuildingConsoleApplication && Executable.Name.EndsWith("-Cmd")))
+			if ((BinaryLinkEnvironment.bIsBuildingDLL && (Options & ClangToolChainOptions.OutputDylib) == 0) ||
+				(BinaryLinkEnvironment.bIsBuildingConsoleApplication && Executable.Name.EndsWith("-Cmd", StringComparison.Ordinal)))
 			{
 				return OutputFiles;
 			}
diff --git a/Engine/Source/Programs/UnrealBuildTool/Platform/Mac/UEDeployMac.cs b/Engine/Source/Programs/UnrealBuildTool/Platform/Mac/UEDeployMac.cs
--- a/Engine/Source/Programs/UnrealBuildTool/Platform/Mac/UEDeployMac.cs
+++ b/Engine/Source/Programs/UnrealBuildTool/Platform/Mac/UEDeployMac.cs
@@ -1,5 +1,6 @@
 // Copyright Epic Games, Inc. All Rights Reserved.
 
+using System;
 using System.IO;
 using EpicGames.Core;
 using Microsoft.Extensions.Logging;
@@ -50,15 +51,15 @@ namespace UnrealBuildTool
 			string BundleVersion = MacToolChain.LoadEngineDisplayVersion();
 			// duplicating some logic in MacToolchain for the BundleID
 			string[] ExeNameParts = ExeName.Split('-');
-			bool bBuildingEditor = ExeNameParts[0].EndsWith("Editor");
-			string FinalBundleID = bBuildingEditor ? $"com.epicgames.{ExeNameParts[0]}" : BundleIdentifier.Replace("[PROJECT_NAME]", ProjectName, System.StringComparison.Ordinal);
-			FinalBundleID = FinalBundleID.Replace("_", "", System.StringComparison.Ordinal);
+			bool bBuildingEditor = ExeNameParts[0].EndsWith("Editor", StringComparison.Ordinal);
+			string FinalBundleID = bBuildingEditor ? $"com.epicgames.{ExeNameParts[0]}" : BundleIdentifier.Replace("[PROJECT_NAME]", ProjectName, StringComparison.Ordinal);
+			FinalBundleID = FinalBundleID.Replace("_", "", StringComparison.Ordinal);
 
-			PListData = PListData.Replace("${EXECUTABLE_NAME}", ExeName, System.StringComparison.Ordinal).
-				Replace("${APP_NAME}", FinalBundleID, System.StringComparison.Ordinal).
-				Replace("${ICON_NAME}", GameName, System.StringComparison.Ordinal).
-				Replace("${MACOSX_DEPLOYMENT_TARGET}", MacToolChainSettings.MinMacDeploymentVersion(TargetType), System.StringComparison.Ordinal).
-				Replace("${BUNDLE_VERSION}", BundleVersion, System.StringComparison.Ordinal);
+			PListData = PListData.Replace("${EXECUTABLE_NAME}", ExeName, StringComparison.Ordinal).
+				Replace("${APP_NAME}", FinalBundleID, StringComparison.Ordinal).
+				Replace("${ICON_NAME}", GameName, StringComparison.Ordinal).
+				Replace("${MACOSX_DEPLOYMENT_TARGET}", MacToolChainSettings.MinMacDeploymentVersion(TargetType), StringComparison.Ordinal).
+				Replace("${BUNDLE_VERSION}", BundleVersion, StringComparison.Ordinal);
 
 			if (!Directory.Exists(IntermediateDirectory))
 			{
diff --git a/Engine/Source/Programs/UnrealBuildTool/Platform/Windows/PVSToolChain.cs b/Engine/Source/Programs/UnrealBuildTool/Platform/Windows/PVSToolChain.cs
--- a/Engine/Source/Programs/UnrealBuildTool/Platform/Windows/PVSToolChain.cs
+++ b/Engine/Source/Programs/UnrealBuildTool/Platform/Windows/PVSToolChain.cs
@@ -721,11 +721,11 @@ namespace UnrealBuildTool
 
 		public static bool ShouldCompileAsC(string compilerCommandLine, string sourceFileName)
 		{
-			int cFlagLastPosition = Math.Max(Math.Max(compilerCommandLine.LastIndexOf("/TC "), compilerCommandLine.LastIndexOf("/Tc ")),
-												Math.Max(compilerCommandLine.LastIndexOf("-TC "), compilerCommandLine.LastIndexOf("-Tc ")));
+			int cFlagLastPosition = Math.Max(Math.Max(compilerCommandLine.LastIndexOf("/TC ", StringComparison.Ordinal), compilerCommandLine.LastIndexOf("/Tc ", StringComparison.Ordinal)),
+												Math.Max(compilerCommandLine.LastIndexOf("-TC ", StringComparison.Ordinal), compilerCommandLine.LastIndexOf("-Tc ", StringComparison.Ordinal)));
 
-			int cppFlagLastPosition = Math.Max(Math.Max(compilerCommandLine.LastIndexOf("/TP "), compilerCommandLine.LastIndexOf("/Tp ")),
-												Math.Max(compilerCommandLine.LastIndexOf("-TP "), compilerCommandLine.LastIndexOf("-Tp ")));
+			int cppFlagLastPosition = Math.Max(Math.Max(compilerCommandLine.LastIndexOf("/TP ", StringComparison.Ordinal), compilerCommandLine.LastIndexOf("/Tp ", StringComparison.Ordinal)),
+												Math.Max(compilerCommandLine.LastIndexOf("-TP ", StringComparison.Ordinal), compilerCommandLine.LastIndexOf("-Tp ", StringComparison.Ordinal)));
 
 			bool compileAsCCode = cFlagLastPosition == cppFlagLastPosition
 				? Path.GetExtension(sourceFileName).Equals(".c", StringComparison.OrdinalIgnoreCase)
@@ -810,7 +810,7 @@ namespace UnrealBuildTool
 				configFileContents.AppendFormat("disable-ms-extensions={0}\n", disableMsExtensions ? "yes" : "no");
 			}
 
-			if (isVisualCppCompiler && preprocessAction.Arguments.Any(arg => arg.StartsWith("/await")))
+			if (isVisualCppCompiler && preprocessAction.Arguments.Any(arg => arg.StartsWith("/await", StringComparison.Ordinal)))
 			{
 				configFileContents.Append("msvc-await=yes\n");
 			}
@@ -1024,7 +1024,7 @@ namespace UnrealBuildTool
 				string arguments = $"-Input=\"{newInput.Location}\" -Output=\"{newOutput.Location}\" -Ignored=\"{ignored.Location}\" -RootPaths=\"{rootPaths.Location}\" -PrintLevel={printLevel}";
 				Action singleFinalizeAction = new(finalizeAction);
 				singleFinalizeAction.CommandDescription = "Process PVS-Studio Single File Results";
-				singleFinalizeAction.CommandArguments = $"{singleFinalizeAction.CommandArguments[..singleFinalizeAction.CommandArguments.IndexOf("-Input=")]}{arguments}";
+				singleFinalizeAction.CommandArguments = $"{singleFinalizeAction.CommandArguments[..singleFinalizeAction.CommandArguments.IndexOf("-Input=", StringComparison.Ordinal)]}{arguments}";
 				singleFinalizeAction.PrerequisiteItems.Clear();
 				singleFinalizeAction.PrerequisiteItems.Add(newInput);
 				singleFinalizeAction.PrerequisiteItems.Add(ignored);
diff --git a/Engine/Source/Programs/UnrealBuildTool/Platform/Windows/VCToolChain.cs b/Engine/Source/Programs/UnrealBuildTool/Platform/Windows/VCToolChain.cs
--- a/Engine/Source/Programs/UnrealBuildTool/Platform/Windows/VCToolChain.cs
+++ b/Engine/Source/Programs/UnrealBuildTool/Platform/Windows/VCToolChain.cs
@@ -316,7 +316,7 @@ namespace UnrealBuildTool
 		{
 			string PchThroughHeaderFilePath = PchThroughHeaderFile.Location.GetFileName();
 			string CreatePchFilePath = NormalizeCommandLinePath(CreatePchFile, RootPaths);
-			if (CreatePchFile.Name.EndsWith(".ifc"))
+			if (CreatePchFile.Name.EndsWith(".ifc", StringComparison.Ordinal))
 			{
 				Arguments.Add($"/exportHeader");
 				Arguments.Add($"/translateInclude");
@@ -336,7 +336,7 @@ namespace UnrealBuildTool
 		{
 			string PchThroughHeaderFilePath = NormalizeCommandLinePath(PchThroughHeaderFile, RootPaths);
 			string UsingPchFilePath = NormalizeCommandLinePath(UsingPchFile, RootPaths);
-			if (UsingPchFile.Name.EndsWith(".ifc"))
+			if (UsingPchFile.Name.EndsWith(".ifc", StringComparison.Ordinal))
 			{
 				Arguments.Add("/wd4127"); // with header units pragma warning disable is not propagated and this warning is very spammy so explicitly turned off here. Hoping ms will change so warning disable is propgated from HU
 				Arguments.Add($"/translateInclude");
@@ -3457,7 +3457,7 @@ namespace UnrealBuildTool
 			ModifyFinalLinkArguments(LinkEnvironment, Arguments, bBuildImportLibraryOnly);
 
 			// Add training data as a prerequisite for pgo optimize
-			if (EnvVars.Compiler.IsMSVC() && LinkEnvironment.bPGOOptimize && LinkEnvironment.PGOFilenamePrefix != null && LinkEnvironment.OutputFilePath.FullName.EndsWith(".exe"))
+			if (EnvVars.Compiler.IsMSVC() && LinkEnvironment.bPGOOptimize && LinkEnvironment.PGOFilenamePrefix != null && LinkEnvironment.OutputFilePath.FullName.EndsWith(".exe", StringComparison.Ordinal))
 			{
 				// training data files are copied to the output directory in ModifyFinalLinkArguments() -> PreparePGOFilesMsvc()
 				PrerequisiteItems.AddRange(DirectoryReference.EnumerateFiles(LinkEnvironment.OutputDirectory!, $"{LinkEnvironment.PGOFilenamePrefix}*.pgd").Select(FileItem.GetItemByFileReference));
@@ -3567,7 +3567,7 @@ namespace UnrealBuildTool
 
 		protected bool PreparePGOFilesMsvc(LinkEnvironment LinkEnvironment)
 		{
-			if (LinkEnvironment.bPGOOptimize && LinkEnvironment.OutputFilePath.FullName.EndsWith(".exe"))
+			if (LinkEnvironment.bPGOOptimize && LinkEnvironment.OutputFilePath.FullName.EndsWith(".exe", StringComparison.Ordinal))
 			{
 				if (!Directory.Exists(LinkEnvironment.PGODirectory))
 				{
@@ -3672,7 +3672,7 @@ namespace UnrealBuildTool
 		};
 		private IEnumerable<string> RemovePGOIgnoredSwitches(IEnumerable<string> SourceArguments)
 		{
-			return SourceArguments.Where(Argument => Argument.StartsWith("/") && !PGOIgnoredLinkerSwitch.Any(Switch => Argument.StartsWith(Switch, StringComparison.OrdinalIgnoreCase)));
+			return SourceArguments.Where(Argument => Argument.StartsWith('/') && !PGOIgnoredLinkerSwitch.Any(Switch => Argument.StartsWith(Switch, StringComparison.OrdinalIgnoreCase)));
 		}
 
 		protected virtual void AddPGOLinkArguments(LinkEnvironment LinkEnvironment, List<string> Arguments)
@@ -3883,7 +3883,7 @@ namespace UnrealBuildTool
 			}
 
 			// IMPLEMENT_MODULE_ is not required - it only exists to ensure developers add an IMPLEMENT_MODULE() declaration in code. These are always removed for PGO so that adding/removing a module won't invalidate PGC data.
-			Arguments.RemoveAll(Argument => Argument.StartsWith("/INCLUDE:IMPLEMENT_MODULE_"));
+			Arguments.RemoveAll(Argument => Argument.StartsWith("/INCLUDE:IMPLEMENT_MODULE_", StringComparison.Ordinal));
 
 			AddPGOLinkArguments(LinkEnvironment, Arguments);
 		}
diff --git a/Engine/Source/Programs/UnrealBuildTool/Preprocessor/Identifier.cs b/Engine/Source/Programs/UnrealBuildTool/Preprocessor/Identifier.cs
--- a/Engine/Source/Programs/UnrealBuildTool/Preprocessor/Identifier.cs
+++ b/Engine/Source/Programs/UnrealBuildTool/Preprocessor/Identifier.cs
@@ -59,7 +59,7 @@ namespace UnrealBuildTool
 		/// <returns>Value indicating which identifier should sort first</returns>
 		public int CompareTo(Identifier? other)
 		{
-			return other is null ? 1 : _name.CompareTo(other._name);
+			return other is null ? 1 : string.Compare(_name, other._name, StringComparison.Ordinal);
 		}
 
 		/// <summary>
diff --git a/Engine/Source/Programs/UnrealBuildTool/ProjectFiles/CMake/CMakefileGenerator.cs b/Engine/Source/Programs/UnrealBuildTool/ProjectFiles/CMake/CMakefileGenerator.cs
--- a/Engine/Source/Programs/UnrealBuildTool/ProjectFiles/CMake/CMakefileGenerator.cs
+++ b/Engine/Source/Programs/UnrealBuildTool/ProjectFiles/CMake/CMakefileGenerator.cs
@@ -134,7 +134,7 @@ namespace UnrealBuildTool
 
 		private void AppendCleanedPathToList(StringBuilder EngineFiles, StringBuilder ProjectFiles, string SourceFileRelativeToRoot, string FullName, string GameProjectPath, string UnrealRootPath, string GameRootPath)
 		{
-			if (!SourceFileRelativeToRoot.StartsWith("..") && !Path.IsPathRooted(SourceFileRelativeToRoot))
+			if (!SourceFileRelativeToRoot.StartsWith("..", StringComparison.Ordinal) && !Path.IsPathRooted(SourceFileRelativeToRoot))
 			{
 				EngineFiles.Append("\t\"" + UnrealRootPath + "/Engine/" + Utils.CleanDirectorySeparators(SourceFileRelativeToRoot, '/') + "\"\n");
 			}
@@ -290,8 +290,8 @@ namespace UnrealBuildTool
 
 					if (!PreprocessorDefinitions.Contains(Definition) &&
 						!PreprocessorDefinitions.Contains(AlternateDefinition) &&
-						!Definition.StartsWith("UE_ENGINE_DIRECTORY") &&
-						!Definition.StartsWith("ORIGINAL_FILE_NAME"))
+						!Definition.StartsWith("UE_ENGINE_DIRECTORY", StringComparison.Ordinal) &&
+						!Definition.StartsWith("ORIGINAL_FILE_NAME", StringComparison.Ordinal))
 					{
 						PreprocessorDefinitions.Add(Definition);
 					}
@@ -310,23 +310,23 @@ namespace UnrealBuildTool
 					// Exclude files/folders on a per-platform basis.
 					if (!IsPathExcludedOnPlatform(SourceFileRelativeToRoot))
 					{
-						if (SourceFileRelativeToRoot.EndsWith(".cpp"))
+						if (SourceFileRelativeToRoot.EndsWith(".cpp", StringComparison.Ordinal))
 						{
 							AppendCleanedPathToList(CMakeEngineSourceFilesList, CMakeProjectSourceFilesList, SourceFileRelativeToRoot, CurSourceFile.FullName, GameProjectPath, UnrealRootPath, CMakeGameRootPath);
 						}
-						else if (SourceFileRelativeToRoot.EndsWith(".h"))
+						else if (SourceFileRelativeToRoot.EndsWith(".h", StringComparison.Ordinal))
 						{
 							AppendCleanedPathToList(CMakeEngineHeaderFilesList, CMakeProjectHeaderFilesList, SourceFileRelativeToRoot, CurSourceFile.FullName, GameProjectPath, UnrealRootPath, CMakeGameRootPath);
 						}
-						else if (SourceFileRelativeToRoot.EndsWith(".cs"))
+						else if (SourceFileRelativeToRoot.EndsWith(".cs", StringComparison.Ordinal))
 						{
 							AppendCleanedPathToList(CMakeEngineCSFilesList, CMakeProjectCSFilesList, SourceFileRelativeToRoot, CurSourceFile.FullName, GameProjectPath, UnrealRootPath, CMakeGameRootPath);
 						}
-						else if (SourceFileRelativeToRoot.EndsWith(".usf") || SourceFileRelativeToRoot.EndsWith(".ush"))
+						else if (SourceFileRelativeToRoot.EndsWith(".usf", StringComparison.Ordinal) || SourceFileRelativeToRoot.EndsWith(".ush", StringComparison.Ordinal))
 						{
 							AppendCleanedPathToList(CMakeEngineShaderFilesList, CMakeProjectShaderFilesList, SourceFileRelativeToRoot, CurSourceFile.FullName, GameProjectPath, UnrealRootPath, CMakeGameRootPath);
 						}
-						else if (SourceFileRelativeToRoot.EndsWith(".ini"))
+						else if (SourceFileRelativeToRoot.EndsWith(".ini", StringComparison.Ordinal))
 						{
 							AppendCleanedPathToList(CMakeEngineConfigFilesList, CMakeProjectConfigFilesList, SourceFileRelativeToRoot, CurSourceFile.FullName, GameProjectPath, UnrealRootPath, CMakeGameRootPath);
 						}
@@ -341,7 +341,7 @@ namespace UnrealBuildTool
 
 			foreach (string PreprocessorDefinition in PreprocessorDefinitions)
 			{
-				int EqPos = PreprocessorDefinition.IndexOf("=");
+				int EqPos = PreprocessorDefinition.IndexOf('=', StringComparison.Ordinal);
 				if (EqPos >= 0)
 				{
 					string Key = PreprocessorDefinition.Substring(0, EqPos);
@@ -359,7 +359,7 @@ namespace UnrealBuildTool
 			foreach (FileReference CurSourceFile in EngineShaderFiles)
 			{
 				string SourceFileRelativeToRoot = CurSourceFile.MakeRelativeTo(Unreal.EngineDirectory);
-				if (SourceFileRelativeToRoot.EndsWith(".usf") || SourceFileRelativeToRoot.EndsWith(".ush"))
+				if (SourceFileRelativeToRoot.EndsWith(".usf", StringComparison.Ordinal) || SourceFileRelativeToRoot.EndsWith(".ush", StringComparison.Ordinal))
 				{
 					AppendCleanedPathToList(CMakeEngineShaderFilesList, CMakeProjectShaderFilesList, SourceFileRelativeToRoot, CurSourceFile.FullName, GameProjectPath, UnrealRootPath, CMakeGameRootPath);
 				}
@@ -370,7 +370,7 @@ namespace UnrealBuildTool
 			foreach (FileReference CurSourceFile in EngineConfigFiles)
 			{
 				string SourceFileRelativeToRoot = CurSourceFile.MakeRelativeTo(Unreal.EngineDirectory);
-				if (SourceFileRelativeToRoot.EndsWith(".ini"))
+				if (SourceFileRelativeToRoot.EndsWith(".ini", StringComparison.Ordinal))
 				{
 					AppendCleanedPathToList(CMakeEngineConfigFilesList, CMakeProjectConfigFilesList, SourceFileRelativeToRoot, CurSourceFile.FullName, GameProjectPath, UnrealRootPath, CMakeGameRootPath);
 				}
@@ -575,7 +575,11 @@ namespace UnrealBuildTool
 		{
 			if (TargetPlatform == UnrealTargetPlatform.IOS || TargetPlatform == UnrealTargetPlatform.TVOS)
 			{
-				if ((TargetName.StartsWith("UnrealGame") || (IsProjectBuild && TargetName.StartsWith(GameProjectName!)) || TargetName.StartsWith("QAGame")) && !TargetName.StartsWith("QAGameEditor"))
+				if (
+					(
+						TargetName.StartsWith("UnrealGame", StringComparison.Ordinal) ||
+						(IsProjectBuild && TargetName.StartsWith(GameProjectName!, StringComparison.Ordinal)) || TargetName.StartsWith("QAGame", StringComparison.Ordinal)
+					) && !TargetName.StartsWith("QAGameEditor", StringComparison.Ordinal))
 				{
 					return false;
 				}
@@ -586,19 +590,19 @@ namespace UnrealBuildTool
 			{
 				// Editor or game builds get all target configs
 				// The game project editor or game get all configs
-				if ((TargetName.StartsWith("UnrealEditor") && !TargetName.StartsWith("UnrealEditorServices")) ||
-					TargetName.StartsWith("UnrealGame") ||
-					(IsProjectBuild && TargetName.StartsWith(GameProjectName!)))
+				if ((TargetName.StartsWith("UnrealEditor", StringComparison.Ordinal) && !TargetName.StartsWith("UnrealEditorServices", StringComparison.Ordinal)) ||
+					TargetName.StartsWith("UnrealGame", StringComparison.Ordinal) ||
+					(IsProjectBuild && TargetName.StartsWith(GameProjectName!, StringComparison.Ordinal)))
 				{
 					return false;
 				}
 				// SCW & CRC are minimally included as just development builds
 				else if (TargetConfig == UnrealTargetConfiguration.Development &&
-					(TargetName.StartsWith("ShaderCompileWorker") || TargetName.StartsWith("CrashReportClient")))
+					(TargetName.StartsWith("ShaderCompileWorker", StringComparison.Ordinal) || TargetName.StartsWith("CrashReportClient", StringComparison.Ordinal)))
 				{
 					return false;
 				}
-				else if ((TargetName.StartsWith("QAGameEditor") && !TargetName.StartsWith("QAGameEditorServices")) || TargetName.StartsWith("QAGame"))
+				else if ((TargetName.StartsWith("QAGameEditor", StringComparison.Ordinal) && !TargetName.StartsWith("QAGameEditorServices", StringComparison.Ordinal)) || TargetName.StartsWith("QAGame", StringComparison.Ordinal))
 				{
 					return false;
 				}
@@ -614,8 +618,8 @@ namespace UnrealBuildTool
 			string FullPath;
 			// Check for paths outside of both the engine and the project
 			if (Path.IsPathRooted(IncludeDir) &&
-				!IncludeDir.StartsWith(FullProjectPath) &&
-				!IncludeDir.StartsWith(Unreal.RootDirectory.FullName))
+				!IncludeDir.StartsWith(FullProjectPath, StringComparison.Ordinal) &&
+				!IncludeDir.StartsWith(Unreal.RootDirectory.FullName, StringComparison.Ordinal))
 			{
 				// Full path to a folder outside of project
 				FullPath = IncludeDir;
@@ -623,7 +627,7 @@ namespace UnrealBuildTool
 			else
 			{
 				FullPath = Path.GetFullPath(Path.Combine(ProjectDir, IncludeDir));
-				if (!FullPath.StartsWith(Unreal.RootDirectory.FullName))
+				if (!FullPath.StartsWith(Unreal.RootDirectory.FullName, StringComparison.Ordinal))
 				{
 					FullPath = Utils.MakePathRelativeTo(FullPath, FullProjectPath);
 				}
diff --git a/Engine/Source/Programs/UnrealBuildTool/ProjectFiles/CodeLite/CodeLiteProject.cs b/Engine/Source/Programs/UnrealBuildTool/ProjectFiles/CodeLite/CodeLiteProject.cs
--- a/Engine/Source/Programs/UnrealBuildTool/ProjectFiles/CodeLite/CodeLiteProject.cs
+++ b/Engine/Source/Programs/UnrealBuildTool/ProjectFiles/CodeLite/CodeLiteProject.cs
@@ -118,7 +118,7 @@ namespace UnrealBuildTool
 						}
 						else
 						{
-							int Idx = Path.GetDirectoryName(CurrentFile.Reference.FullName)!.IndexOf(ProjectNameRaw) + ProjectNameRaw.Length;
+							int Idx = Path.GetDirectoryName(CurrentFile.Reference.FullName)!.IndexOf(ProjectNameRaw, StringComparison.Ordinal) + ProjectNameRaw.Length;
 							CurrentFilePath = Path.GetDirectoryName(CurrentFile.Reference.FullName)!.Substring(Idx);
 						}
 					}
@@ -127,7 +127,7 @@ namespace UnrealBuildTool
 						//
 						// We do not need all the editors subfolders to show the content. Find the correct programs subfolder.
 						//
-						int Idx = Path.GetDirectoryName(CurrentFile.Reference.FullName)!.IndexOf(TargetName) + TargetName.Length;
+						int Idx = Path.GetDirectoryName(CurrentFile.Reference.FullName)!.IndexOf(TargetName, StringComparison.Ordinal) + TargetName.Length;
 						CurrentFilePath = Path.GetDirectoryName(CurrentFile.Reference.FullName)!.Substring(Idx);
 					}
 
diff --git a/Engine/Source/Programs/UnrealBuildTool/ProjectFiles/Eddie/EddieProject.cs b/Engine/Source/Programs/UnrealBuildTool/ProjectFiles/Eddie/EddieProject.cs
--- a/Engine/Source/Programs/UnrealBuildTool/ProjectFiles/Eddie/EddieProject.cs
+++ b/Engine/Source/Programs/UnrealBuildTool/ProjectFiles/Eddie/EddieProject.cs
@@ -55,7 +55,7 @@ namespace UnrealBuildTool
 
 		public override SourceFile? AllocSourceFile(FileReference InitFilePath, DirectoryReference? InitProjectSubFolder)
 		{
-			if (InitFilePath.GetFileName().StartsWith("."))
+			if (InitFilePath.GetFileName().StartsWith('.'))
 			{
 				return null;
 			}
@@ -120,7 +120,9 @@ namespace UnrealBuildTool
 				EddieFolder? Group = FindFolderByRelativePath(ref Folders, RelativeSourceDirectory);
 				if (Group != null)
 				{
-					if (FileName.EndsWith(".build.cs") || FileName.EndsWith(".Build.cs") || FileName.EndsWith(".csproj"))
+					if (FileName.EndsWith(".build.cs", System.StringComparison.Ordinal) ||
+						FileName.EndsWith(".Build.cs", System.StringComparison.Ordinal) ||
+						FileName.EndsWith(".csproj", System.StringComparison.Ordinal))
 					{
 						Group.bIsModuleFolder = true;
 						Group.WorksetPath = ProjectFilePath.FullName + "." + Group.Name + ".wkst";
diff --git a/Engine/Source/Programs/UnrealBuildTool/ProjectFiles/KDevelop/KDevelopGenerator.cs b/Engine/Source/Programs/UnrealBuildTool/ProjectFiles/KDevelop/KDevelopGenerator.cs
--- a/Engine/Source/Programs/UnrealBuildTool/ProjectFiles/KDevelop/KDevelopGenerator.cs
+++ b/Engine/Source/Programs/UnrealBuildTool/ProjectFiles/KDevelop/KDevelopGenerator.cs
@@ -263,7 +263,7 @@ namespace UnrealBuildTool
 					string FullPath = "";
 
 					// need to test to see if this in the project souce tree
-					if (CurPath.StartsWith("/") && !CurPath.StartsWith(FullProjectPath))
+					if (CurPath.StartsWith('/') && !CurPath.StartsWith(FullProjectPath, StringComparison.Ordinal))
 					{
 						// Full path to a folder outside of project
 						FullPath = CurPath;
@@ -289,7 +289,7 @@ namespace UnrealBuildTool
 					string FullProjectPath = ProjectFileGenerator.PrimaryProjectPath.FullName;
 					string FullPath = "";
 
-					if (CurPath.StartsWith("/") && !CurPath.StartsWith(FullProjectPath))
+					if (CurPath.StartsWith('/') && !CurPath.StartsWith(FullProjectPath, StringComparison.Ordinal))
 					{
 						// Full path to a folder outside of project
 						FullPath = CurPath;
diff --git a/Engine/Source/Programs/UnrealBuildTool/ProjectFiles/ProjectFile.cs b/Engine/Source/Programs/UnrealBuildTool/ProjectFiles/ProjectFile.cs
--- a/Engine/Source/Programs/UnrealBuildTool/ProjectFiles/ProjectFile.cs
+++ b/Engine/Source/Programs/UnrealBuildTool/ProjectFiles/ProjectFile.cs
@@ -432,11 +432,11 @@ namespace UnrealBuildTool
 		public static string NormalizeProjectPath(string InputPath)
 		{
 			// If the path is rooted in an environment variable, leave it be.
-			if (InputPath.StartsWith("$("))
+			if (InputPath.StartsWith("$(", StringComparison.Ordinal))
 			{
 				return InputPath;
 			}
-			else if (InputPath.EndsWith("\\") || InputPath.EndsWith("/"))
+			else if (InputPath.EndsWith('\\') || InputPath.EndsWith('/'))
 			{
 				return NormalizeProjectPath(new DirectoryReference(InputPath));
 			}
diff --git a/Engine/Source/Programs/UnrealBuildTool/ProjectFiles/ProjectFileGenerator.cs b/Engine/Source/Programs/UnrealBuildTool/ProjectFiles/ProjectFileGenerator.cs
--- a/Engine/Source/Programs/UnrealBuildTool/ProjectFiles/ProjectFileGenerator.cs
+++ b/Engine/Source/Programs/UnrealBuildTool/ProjectFiles/ProjectFileGenerator.cs
@@ -1637,16 +1637,16 @@ namespace UnrealBuildTool
 
 			foreach (string CurArgument in Arguments)
 			{
-				if (CurArgument.StartsWith("-"))
+				if (CurArgument.StartsWith('-'))
 				{
-					if (CurArgument.StartsWith("-Platforms=", StringComparison.InvariantCultureIgnoreCase))
+					if (CurArgument.StartsWith("-Platforms=", StringComparison.OrdinalIgnoreCase))
 					{
 						// Parse the list... will be in Foo+Bar+New format
 						string PlatformList = CurArgument.Substring(11);
 						while (PlatformList.Length > 0)
 						{
 							string PlatformString = PlatformList;
-							int PlusIdx = PlatformList.IndexOf("+");
+							int PlusIdx = PlatformList.IndexOf('+', StringComparison.Ordinal);
 							if (PlusIdx != -1)
 							{
 								PlatformString = PlatformList.Substring(0, PlusIdx);
@@ -1975,7 +1975,7 @@ namespace UnrealBuildTool
 				string Search = CleanTargetFileName;
 				if (Search.StartsWith(LocalRoot, StringComparison.InvariantCultureIgnoreCase))
 				{
-					if (LocalRoot.EndsWith("\\") || LocalRoot.EndsWith("/"))
+					if (LocalRoot.EndsWith('\\') || LocalRoot.EndsWith('/'))
 					{
 						Search = Search.Substring(LocalRoot.Length - 1);
 					}
@@ -1990,7 +1990,7 @@ namespace UnrealBuildTool
 					string ProjectRoot = OnlyGameProject.Directory.FullName;
 					if (Search.StartsWith(ProjectRoot, StringComparison.InvariantCultureIgnoreCase))
 					{
-						if (ProjectRoot.EndsWith("\\") || ProjectRoot.EndsWith("/"))
+						if (ProjectRoot.EndsWith('\\') || ProjectRoot.EndsWith('/'))
 						{
 							Search = Search.Substring(ProjectRoot.Length - 1);
 						}
@@ -2141,7 +2141,7 @@ namespace UnrealBuildTool
 			Folder.SubFolders.AddRange(SubFoldersToAdd);
 
 			// Sort subfolders to be deterministic
-			Folder.SubFolders.Sort(static (a, b) => a.FolderName.CompareTo(b.FolderName));
+			Folder.SubFolders.Sort(static (a, b) => string.Compare(a.FolderName, b.FolderName, StringComparison.Ordinal));
 
 			// Sort child projects as well while we're at it... to make .sln deterministic
 			Folder.ChildProjects.Sort(static (a, b) => a.ProjectFilePath.CompareTo(b.ProjectFilePath));
@@ -2750,7 +2750,7 @@ namespace UnrealBuildTool
 					// Find all of the source files (and other files) and add them to the project
 					List<FileReference> FoundFiles = SourceFileSearch.FindModuleSourceFiles(CurModuleFile, SearchSubdirectories: SearchSubdirectories);
 					// remove any target files, they are technically not part of the module and are explicitly added when the project is created
-					FoundFiles.RemoveAll(f => f.FullName.EndsWith(".Target.cs"));
+					FoundFiles.RemoveAll(f => f.FullName.EndsWith(".Target.cs", StringComparison.Ordinal));
 
 					lock (ModuleToAdditionalPlugin)
 					{
diff --git a/Engine/Source/Programs/UnrealBuildTool/ProjectFiles/QMake/QMakefileGenerator.cs b/Engine/Source/Programs/UnrealBuildTool/ProjectFiles/QMake/QMakefileGenerator.cs
--- a/Engine/Source/Programs/UnrealBuildTool/ProjectFiles/QMake/QMakefileGenerator.cs
+++ b/Engine/Source/Programs/UnrealBuildTool/ProjectFiles/QMake/QMakefileGenerator.cs
@@ -66,7 +66,7 @@ namespace UnrealBuildTool
 		{
 			string FullProjectPath = ProjectFileGenerator.PrimaryProjectPath.FullName;
 			string FullPath;
-			if (IncludeDir.StartsWith("/") && !IncludeDir.StartsWith(FullProjectPath))
+			if (IncludeDir.StartsWith('/') && !IncludeDir.StartsWith(FullProjectPath, StringComparison.Ordinal))
 			{
 				// Full path to a fulder outside of project
 				FullPath = IncludeDir;
@@ -321,9 +321,9 @@ namespace UnrealBuildTool
 								}
 							}
 						}
-						if (SourceFileRelativeToRoot.EndsWith(".h"))
+						if (SourceFileRelativeToRoot.EndsWith(".h", StringComparison.Ordinal))
 						{
-							if (!SourceFileRelativeToRoot.StartsWith("..") && !Path.IsPathRooted(SourceFileRelativeToRoot))
+							if (!SourceFileRelativeToRoot.StartsWith("..", StringComparison.Ordinal) && !Path.IsPathRooted(SourceFileRelativeToRoot))
 							{
 								// SourceFileRelativeToRoot = "Engine/" + SourceFileRelativeToRoot;
 								QMakeHeaderFilesListBuilder.Append("\t\"" + $"$$unrealRootPath{Seperator}Engine{Seperator}" + SourceFileRelativeToRoot + "\" \\\n");
@@ -341,9 +341,9 @@ namespace UnrealBuildTool
 								}
 							}
 						}
-						if (SourceFileRelativeToRoot.EndsWith(".cs"))
+						if (SourceFileRelativeToRoot.EndsWith(".cs", StringComparison.Ordinal))
 						{
-							if (!SourceFileRelativeToRoot.StartsWith("..") && !Path.IsPathRooted(SourceFileRelativeToRoot))
+							if (!SourceFileRelativeToRoot.StartsWith("..", StringComparison.Ordinal) && !Path.IsPathRooted(SourceFileRelativeToRoot))
 							{
 								// SourceFileRelativeToRoot = "Engine/" + SourceFileRelativeToRoot;
 								QMakeConfigFilesListBuilder.Append("\t\"" + $"$$unrealRootPath{Seperator}Engine{Seperator}" + SourceFileRelativeToRoot + "\" \\\n");
diff --git a/Engine/Source/Programs/UnrealBuildTool/ProjectFiles/VisualStudioCode/VSCodeProjectFileGenerator.cs b/Engine/Source/Programs/UnrealBuildTool/ProjectFiles/VisualStudioCode/VSCodeProjectFileGenerator.cs
--- a/Engine/Source/Programs/UnrealBuildTool/ProjectFiles/VisualStudioCode/VSCodeProjectFileGenerator.cs
+++ b/Engine/Source/Programs/UnrealBuildTool/ProjectFiles/VisualStudioCode/VSCodeProjectFileGenerator.cs
@@ -290,7 +290,7 @@ namespace UnrealBuildTool
 			{
 				Projects = new List<ProjectFile>(AllProjectFiles);
 			}
-			Projects.Sort((A, B) => { return A.ProjectFilePath.GetFileName().CompareTo(B.ProjectFilePath.GetFileName()); });
+			Projects.Sort((A, B) => { return string.Compare(A.ProjectFilePath.GetFileName(), B.ProjectFilePath.GetFileName(), StringComparison.Ordinal); });
 
 			ProjectData ProjectData = GatherProjectData(Projects);
 
diff --git a/Engine/Source/Programs/UnrealBuildTool/ProjectFiles/Xcode/XcodeProject.File.cs b/Engine/Source/Programs/UnrealBuildTool/ProjectFiles/Xcode/XcodeProject.File.cs
--- a/Engine/Source/Programs/UnrealBuildTool/ProjectFiles/Xcode/XcodeProject.File.cs
+++ b/Engine/Source/Programs/UnrealBuildTool/ProjectFiles/Xcode/XcodeProject.File.cs
@@ -322,7 +322,7 @@ namespace UnrealBuildTool.XcodeProjectXcconfig
 				RelativePath = RelativePath,
 				FileKey = FileKey,
 				FileType = FileType,
-				SourceTree = SourceTree.StartsWith("<") ? $"\"{SourceTree}\"" : SourceTree,
+				SourceTree = SourceTree.StartsWith('<') ? $"\"{SourceTree}\"" : SourceTree,
 				GroupName = GroupName,
 			});
 
@@ -605,7 +605,7 @@ namespace UnrealBuildTool.XcodeProjectXcconfig
 						return null;
 					}
 
-					CurrentGroup = new XcodeFileGroup(Path.GetFileName(CurrentPath), CurrentPath, GetFileType(Path.GetExtension(CurrentPath)).StartsWith("folder"));
+					CurrentGroup = new XcodeFileGroup(Path.GetFileName(CurrentPath), CurrentPath, GetFileType(Path.GetExtension(CurrentPath)).StartsWith("folder", StringComparison.Ordinal));
 					if (Groups.Count == 0)
 					{
 						MainGroupGuid = CurrentGroup.GroupGuid;
diff --git a/Engine/Source/Programs/UnrealBuildTool/ProjectFiles/Xcode/XcodeProject.Scheme.cs b/Engine/Source/Programs/UnrealBuildTool/ProjectFiles/Xcode/XcodeProject.Scheme.cs
--- a/Engine/Source/Programs/UnrealBuildTool/ProjectFiles/Xcode/XcodeProject.Scheme.cs
+++ b/Engine/Source/Programs/UnrealBuildTool/ProjectFiles/Xcode/XcodeProject.Scheme.cs
@@ -50,8 +50,8 @@ namespace UnrealBuildTool.XcodeProjectXcconfig
 				if (FileReference.Exists(SchemeFilePath))
 				{
 					string OldContents = File.ReadAllText(SchemeFilePath.FullName);
-					int OldCommandLineArgumentsStart = OldContents.IndexOf("<CommandLineArguments>") + "<CommandLineArguments>".Length;
-					int OldCommandLineArgumentsEnd = OldContents.IndexOf("</CommandLineArguments>");
+					int OldCommandLineArgumentsStart = OldContents.IndexOf("<CommandLineArguments>", StringComparison.Ordinal) + "<CommandLineArguments>".Length;
+					int OldCommandLineArgumentsEnd = OldContents.IndexOf("</CommandLineArguments>", StringComparison.Ordinal);
 					if (OldCommandLineArgumentsStart != -1 && OldCommandLineArgumentsEnd != -1)
 					{
 						OldCommandLineArguments = OldContents.Substring(OldCommandLineArgumentsStart, OldCommandLineArgumentsEnd - OldCommandLineArgumentsStart);
diff --git a/Engine/Source/Programs/UnrealBuildTool/ProjectFiles/Xcode/XcodeProject.Utils.cs b/Engine/Source/Programs/UnrealBuildTool/ProjectFiles/Xcode/XcodeProject.Utils.cs
--- a/Engine/Source/Programs/UnrealBuildTool/ProjectFiles/Xcode/XcodeProject.Utils.cs
+++ b/Engine/Source/Programs/UnrealBuildTool/ProjectFiles/Xcode/XcodeProject.Utils.cs
@@ -126,7 +126,7 @@ namespace UnrealBuildTool.XcodeProjectXcconfig
 			string Identifier = Plist($"Print :{Key}", PlistItem.File.FullName);
 
 			// handle error
-			if (String.IsNullOrEmpty(Identifier) || Identifier.StartsWith("Print:"))
+			if (String.IsNullOrEmpty(Identifier) || Identifier.StartsWith("Print:", StringComparison.Ordinal))
 			{
 				if (PlistItem.Mode == MetadataMode.UsePremade)
 				{
@@ -173,7 +173,7 @@ namespace UnrealBuildTool.XcodeProjectXcconfig
 			string ExistingSetting = Plist($"Print {Entry}");
 
 			// Print errors start with Print
-			if (!ExistingSetting.StartsWith("Print:") && ExistingSetting != Value)
+			if (!ExistingSetting.StartsWith("Print:", StringComparison.Ordinal) && ExistingSetting != Value)
 			{
 				Plist($"Set {Entry} {Value}");
 
@@ -218,7 +218,7 @@ namespace UnrealBuildTool.XcodeProjectXcconfig
 		{
 			string ExistingPath = Plist($"Print {Entry}");
 			// skip of errors, or it's an absolute path
-			if (!ExistingPath.StartsWith("Print:") && !ExistingPath.StartsWith("/"))
+			if (!ExistingPath.StartsWith("Print:", StringComparison.Ordinal) && !ExistingPath.StartsWith('/'))
 			{
 				// fixup the path to be relative to new project instead of old
 				string FixedPath = Utils.CollapseRelativeDirectories(Path.Combine(RelativeToProject, ExistingPath));
diff --git a/Engine/Source/Programs/UnrealBuildTool/ProjectFiles/Xcode/XcodeProject.cs b/Engine/Source/Programs/UnrealBuildTool/ProjectFiles/Xcode/XcodeProject.cs
--- a/Engine/Source/Programs/UnrealBuildTool/ProjectFiles/Xcode/XcodeProject.cs
+++ b/Engine/Source/Programs/UnrealBuildTool/ProjectFiles/Xcode/XcodeProject.cs
@@ -1467,7 +1467,7 @@ namespace UnrealBuildTool.XcodeProjectXcconfig
 							// Get the source files
 							foreach (FileSystemInfo FSI in DI.EnumerateFileSystemInfos())
 							{
-								if (FSI.Name.StartsWith(".") || FSI.Name.EndsWith(".uappex"))
+								if (FSI.Name.StartsWith('.') || FSI.Name.EndsWith(".uappex", StringComparison.Ordinal))
 								{
 									continue;
 								}
@@ -1556,7 +1556,7 @@ namespace UnrealBuildTool.XcodeProjectXcconfig
 				{
 					// make it absolute if it was relative (to Engine/Source)
 					string FixedSource = Dylib.SourcePath!;
-					if (!FixedSource.StartsWith("/"))
+					if (!FixedSource.StartsWith('/'))
 					{
 						FixedSource = $"${{UE_ENGINE_DIR}}/Source/{FixedSource}";
 					}
@@ -1877,7 +1877,7 @@ namespace UnrealBuildTool.XcodeProjectXcconfig
 
 					if (Framework.ZipFile != null)
 					{
-						if (Framework.ZipFile.FullName.EndsWith(".embeddedframework.zip"))
+						if (Framework.ZipFile.FullName.EndsWith(".embeddedframework.zip", StringComparison.Ordinal))
 						{
 							// foo.embeddedframework.zip would have foo.framework inside it, which is what we want to install
 							FinalFrameworkDir = DirectoryReference.Combine(Framework.ZipOutputDirectory!, Framework.ZipFile.GetFileNameWithoutAnyExtensions() + ".framework");
@@ -2403,7 +2403,10 @@ namespace UnrealBuildTool.XcodeProjectXcconfig
 
 				foreach (FileSystemInfo FSI in ExtensionPath.EnumerateFileSystemInfos())
 				{
-					if (FSI.Name.EndsWith(".h") || FSI.Name.StartsWith(".") || FSI.Name.EndsWith(".plist") || FSI.Name.EndsWith(".uappex"))
+					if (FSI.Name.EndsWith(".h", StringComparison.Ordinal) ||
+						FSI.Name.StartsWith('.') ||
+						FSI.Name.EndsWith(".plist", StringComparison.Ordinal) ||
+						FSI.Name.EndsWith(".uappex", StringComparison.Ordinal))
 					{
 						continue;
 					}
@@ -3096,7 +3099,7 @@ namespace UnrealBuildTool.XcodeProjectXcconfig
 		/// <returns>The newly allocated source file object</returns>
 		public override SourceFile? AllocSourceFile(FileReference InitFilePath, DirectoryReference? InitProjectSubFolder)
 		{
-			if (InitFilePath.GetFileName().StartsWith("."))
+			if (InitFilePath.GetFileName().StartsWith('.'))
 			{
 				return null;
 			}
diff --git a/Engine/Source/Programs/UnrealBuildTool/ProjectFiles/Xcode/XcodeProjectFileGenerator.cs b/Engine/Source/Programs/UnrealBuildTool/ProjectFiles/Xcode/XcodeProjectFileGenerator.cs
--- a/Engine/Source/Programs/UnrealBuildTool/ProjectFiles/Xcode/XcodeProjectFileGenerator.cs
+++ b/Engine/Source/Programs/UnrealBuildTool/ProjectFiles/Xcode/XcodeProjectFileGenerator.cs
@@ -393,7 +393,7 @@ namespace UnrealBuildTool
 						return TypeA.CompareTo(TypeB);
 					}
 
-					return TargetA.Name.CompareTo(TargetB.Name);
+					return string.Compare(TargetA.Name, TargetB.Name, StringComparison.Ordinal);
 				});
 
 				foreach (XcodeProjectXcconfig.XcodeProjectFile XcodeProject in BuildableProjects)
@@ -593,7 +593,7 @@ namespace UnrealBuildTool
 							// walk over CPP and External modules looking for dylibs that need to be copied into the .app directly (not in a framework)
 							foreach (UEBuildModule Module in Binary.Modules)
 							{
-								Dylibs.AddRange(Module.Rules.RuntimeDependencies.Inner.Where(x => x.Path.StartsWith("$(BinaryOutputDir)") && Path.GetExtension(x.SourcePath) == ".dylib"));
+								Dylibs.AddRange(Module.Rules.RuntimeDependencies.Inner.Where(x => x.Path.StartsWith("$(BinaryOutputDir)", StringComparison.Ordinal) && Path.GetExtension(x.SourcePath) == ".dylib"));
 								Bundles.AddRange(Module.Rules.AdditionalBundleResources.Select(x => new UEBuildBundleResource(x)));
 							}
 						}
diff --git a/Engine/Source/Programs/UnrealBuildTool/ProjectFiles/Xcode/XcodeProjectLegacy.cs b/Engine/Source/Programs/UnrealBuildTool/ProjectFiles/Xcode/XcodeProjectLegacy.cs
--- a/Engine/Source/Programs/UnrealBuildTool/ProjectFiles/Xcode/XcodeProjectLegacy.cs
+++ b/Engine/Source/Programs/UnrealBuildTool/ProjectFiles/Xcode/XcodeProjectLegacy.cs
@@ -237,7 +237,7 @@ namespace UnrealBuildTool.XcodeProjectLegacy
 				XcodeFileGroup CurrentGroup;
 				if (!CurrentSubGroups.ContainsKey(CurrentPath))
 				{
-					CurrentGroup = new XcodeFileGroup(Path.GetFileName(CurrentPath), CurrentPath, CurrentPath.EndsWith(".xcassets"));
+					CurrentGroup = new XcodeFileGroup(Path.GetFileName(CurrentPath), CurrentPath, CurrentPath.EndsWith(".xcassets", StringComparison.Ordinal));
 					CurrentSubGroups.Add(CurrentPath, CurrentGroup);
 				}
 				else
@@ -274,7 +274,7 @@ namespace UnrealBuildTool.XcodeProjectLegacy
 		/// <returns>The newly allocated source file object</returns>
 		public override SourceFile? AllocSourceFile(FileReference InitFilePath, DirectoryReference? InitProjectSubFolder)
 		{
-			if (InitFilePath.GetFileName().StartsWith("."))
+			if (InitFilePath.GetFileName().StartsWith('.'))
 			{
 				return null;
 			}
@@ -286,7 +286,7 @@ namespace UnrealBuildTool.XcodeProjectLegacy
 		/// </summary>
 		private void GenerateSectionsWithSourceFiles(StringBuilder PBXBuildFileSection, StringBuilder PBXFileReferenceSection, StringBuilder PBXSourcesBuildPhaseSection, string TargetAppGuid, string TargetName, bool bIsAppBundle)
 		{
-			SourceFiles.Sort((x, y) => { return x.Reference.FullName.CompareTo(y.Reference.FullName); });
+			SourceFiles.Sort((x, y) => { return x.Reference.CompareTo(y.Reference); });
 
 			foreach (XcodeSourceFile SourceFile in SourceFiles.OfType<XcodeSourceFile>())
 			{
@@ -368,7 +368,7 @@ namespace UnrealBuildTool.XcodeProjectLegacy
 							DirectoryInfo ResourceDir = new System.IO.DirectoryInfo(Path.Combine(DI.FullName, "Resources"));
 							foreach (FileSystemInfo FSI in ResourceDir.EnumerateFileSystemInfos())
 							{
-								if (FSI.Name.StartsWith("."))
+								if (FSI.Name.StartsWith('.'))
 								{
 									continue;
 								}
@@ -1578,8 +1578,8 @@ namespace UnrealBuildTool.XcodeProjectLegacy
 			if (FileReference.Exists(SchemeFilePath))
 			{
 				string OldContents = File.ReadAllText(SchemeFilePath.FullName);
-				int OldCommandLineArgumentsStart = OldContents.IndexOf("<CommandLineArguments>") + "<CommandLineArguments>".Length;
-				int OldCommandLineArgumentsEnd = OldContents.IndexOf("</CommandLineArguments>");
+				int OldCommandLineArgumentsStart = OldContents.IndexOf("<CommandLineArguments>", StringComparison.Ordinal) + "<CommandLineArguments>".Length;
+				int OldCommandLineArgumentsEnd = OldContents.IndexOf("</CommandLineArguments>", StringComparison.Ordinal);
 				if (OldCommandLineArgumentsStart != -1 && OldCommandLineArgumentsEnd != -1)
 				{
 					OldCommandLineArguments = OldContents.Substring(OldCommandLineArgumentsStart, OldCommandLineArgumentsEnd - OldCommandLineArgumentsStart);
diff --git a/Engine/Source/Programs/UnrealBuildTool/System/BuildHostPlatform.cs b/Engine/Source/Programs/UnrealBuildTool/System/BuildHostPlatform.cs
--- a/Engine/Source/Programs/UnrealBuildTool/System/BuildHostPlatform.cs
+++ b/Engine/Source/Programs/UnrealBuildTool/System/BuildHostPlatform.cs
@@ -264,10 +264,10 @@ namespace UnrealBuildTool
 			for (string? line = process.StandardOutput.ReadLine(); line != null; line = process.StandardOutput.ReadLine())
 			{
 				line = line.Trim();
-				if (line.EndsWith(".dylib"))
+				if (line.EndsWith(".dylib", StringComparison.Ordinal))
 				{
 					const int sharingModeLength = 6;
-					int smStart = line.IndexOf("SM=");
+					int smStart = line.IndexOf("SM=", StringComparison.Ordinal);
 					if (smStart != -1)
 					{
 						int pathStart = smStart + sharingModeLength;
diff --git a/Engine/Source/Programs/UnrealBuildTool/System/CommandLine.cs b/Engine/Source/Programs/UnrealBuildTool/System/CommandLine.cs
--- a/Engine/Source/Programs/UnrealBuildTool/System/CommandLine.cs
+++ b/Engine/Source/Programs/UnrealBuildTool/System/CommandLine.cs
@@ -108,7 +108,7 @@ namespace UnrealBuildTool
 						}
 						else
 						{
-							if (NonNullableType(FieldInfo.FieldType) != typeof(bool) && Attribute.Value == null && !Prefix.EndsWith("=") && !Prefix.EndsWith(":"))
+							if (NonNullableType(FieldInfo.FieldType) != typeof(bool) && Attribute.Value == null && !Prefix.EndsWith('=') && !Prefix.EndsWith(':'))
 							{
 								Prefix += "=";
 							}
diff --git a/Engine/Source/Programs/UnrealBuildTool/System/CppIncludeParser.cs b/Engine/Source/Programs/UnrealBuildTool/System/CppIncludeParser.cs
--- a/Engine/Source/Programs/UnrealBuildTool/System/CppIncludeParser.cs
+++ b/Engine/Source/Programs/UnrealBuildTool/System/CppIncludeParser.cs
@@ -53,7 +53,7 @@ namespace UnrealBuildTool
 
 						string IncludeFile = Token.ToString();
 
-						if (!IncludeFile.EndsWith(".h\"") && !IncludeFile.EndsWith(".h>"))
+						if (!IncludeFile.EndsWith(".h\"", StringComparison.Ordinal) && !IncludeFile.EndsWith(".h>", StringComparison.Ordinal))
 						{
 							break;
 						}
diff --git a/Engine/Source/Programs/UnrealBuildTool/System/CppRootPaths.cs b/Engine/Source/Programs/UnrealBuildTool/System/CppRootPaths.cs
--- a/Engine/Source/Programs/UnrealBuildTool/System/CppRootPaths.cs
+++ b/Engine/Source/Programs/UnrealBuildTool/System/CppRootPaths.cs
@@ -144,7 +144,7 @@ namespace UnrealBuildTool
 				if (key == CppRootPathFolder.Project)
 				{
 					string projectName = value.GetDirectoryName();
-					if (!projectName.StartsWith("Clang")) // Lol, there is a project called clang.. which collides with the compiler clang causing two roots to have the same path
+					if (!projectName.StartsWith("Clang", StringComparison.Ordinal)) // Lol, there is a project called clang.. which collides with the compiler clang causing two roots to have the same path
 					{
 						AddFolderName(CppRootPathFolder.Project, value.GetDirectoryName());
 					}
diff --git a/Engine/Source/Programs/UnrealBuildTool/System/EncryptionAndSigning.cs b/Engine/Source/Programs/UnrealBuildTool/System/EncryptionAndSigning.cs
--- a/Engine/Source/Programs/UnrealBuildTool/System/EncryptionAndSigning.cs
+++ b/Engine/Source/Programs/UnrealBuildTool/System/EncryptionAndSigning.cs
@@ -217,7 +217,7 @@ namespace UnrealBuildTool
 		/// </summary>
 		private static string ProcessSigningKeyInputStrings(string InString)
 		{
-			if (InString.StartsWith("0x"))
+			if (InString.StartsWith("0x", StringComparison.Ordinal))
 			{
 				InString = InString.Substring(2);
 			}
@@ -520,7 +520,7 @@ namespace UnrealBuildTool
 		/// </summary>
 		private static byte[] ParseHexStringToByteArray(string InString, int InMinimumLength)
 		{
-			if (InString.StartsWith("0x"))
+			if (InString.StartsWith("0x", StringComparison.Ordinal))
 			{
 				InString = InString.Substring(2);
 			}
diff --git a/Engine/Source/Programs/UnrealBuildTool/System/HotReload.cs b/Engine/Source/Programs/UnrealBuildTool/System/HotReload.cs
--- a/Engine/Source/Programs/UnrealBuildTool/System/HotReload.cs
+++ b/Engine/Source/Programs/UnrealBuildTool/System/HotReload.cs
@@ -1256,7 +1256,7 @@ namespace UnrealBuildTool
 					// Find the argument for the metadata file
 					const string InputArgument = "-Input=";
 
-					int InputIdx = Arguments.IndexOf(InputArgument);
+					int InputIdx = Arguments.IndexOf(InputArgument, StringComparison.Ordinal);
 					if (InputIdx == -1)
 					{
 						throw new Exception("Missing -Input= argument to WriteMetadata command when patching action graph.");
diff --git a/Engine/Source/Programs/UnrealBuildTool/System/InstalledPlatformInfo.cs b/Engine/Source/Programs/UnrealBuildTool/System/InstalledPlatformInfo.cs
--- a/Engine/Source/Programs/UnrealBuildTool/System/InstalledPlatformInfo.cs
+++ b/Engine/Source/Programs/UnrealBuildTool/System/InstalledPlatformInfo.cs
@@ -234,7 +234,7 @@ namespace UnrealBuildTool
 		private static bool ParseSubValue(string TrimmedLine, string Match, out string Result)
 		{
 			Result = String.Empty;
-			int MatchIndex = TrimmedLine.IndexOf(Match);
+			int MatchIndex = TrimmedLine.IndexOf(Match, StringComparison.Ordinal);
 			if (MatchIndex < 0)
 			{
 				return false;
@@ -249,7 +249,7 @@ namespace UnrealBuildTool
 			// get everything up to the first comma and trim any new whitespace
 			Result = TrimmedLine.Split(',')[0];
 			Result = Result.Trim();
-			if (Result.StartsWith("\""))
+			if (Result.StartsWith('\"'))
 			{
 				// Remove quotes
 				int QuoteEnd = Result.LastIndexOf('\"');
diff --git a/Engine/Source/Programs/UnrealBuildTool/System/SourceFileMetadataCache.cs b/Engine/Source/Programs/UnrealBuildTool/System/SourceFileMetadataCache.cs
--- a/Engine/Source/Programs/UnrealBuildTool/System/SourceFileMetadataCache.cs
+++ b/Engine/Source/Programs/UnrealBuildTool/System/SourceFileMetadataCache.cs
@@ -398,7 +398,7 @@ namespace UnrealBuildTool
 				}
 				else
 				{
-					int HeaderUnitIndex = Line.IndexOf("HEADER_UNIT_");
+					int HeaderUnitIndex = Line.IndexOf("HEADER_UNIT_", StringComparison.Ordinal);
 					if (HeaderUnitIndex != -1)
 					{
 						ReadOnlySpan<char> Span = Line.AsSpan(HeaderUnitIndex + "HEADER_UNIT_".Length);
diff --git a/Engine/Source/Programs/UnrealBuildTool/System/SourceFileSearch.cs b/Engine/Source/Programs/UnrealBuildTool/System/SourceFileSearch.cs
--- a/Engine/Source/Programs/UnrealBuildTool/System/SourceFileSearch.cs
+++ b/Engine/Source/Programs/UnrealBuildTool/System/SourceFileSearch.cs
@@ -1,5 +1,6 @@
 // Copyright Epic Games, Inc. All Rights Reserved.
 
+using System;
 using System.Collections.Generic;
 using System.IO;
 using System.Linq;
@@ -140,7 +141,7 @@ namespace UnrealBuildTool
 			foreach (FileItem File in Directory.EnumerateFiles())
 			{
 				// Ignore Mac resource fork files on non-HFS partitions
-				if (File.Name.StartsWith("._"))
+				if (File.Name.StartsWith("._", StringComparison.Ordinal))
 				{
 					continue;
 				}
@@ -183,7 +184,7 @@ namespace UnrealBuildTool
 					continue;
 				}
 
-				if (SubDirectory.Name.StartsWith("__External"))
+				if (SubDirectory.Name.StartsWith("__External", StringComparison.Ordinal))
 				{
 					continue;
 				}
diff --git a/Engine/Source/Programs/UnrealBuildTool/System/SourceFileWorkingSet.cs b/Engine/Source/Programs/UnrealBuildTool/System/SourceFileWorkingSet.cs
--- a/Engine/Source/Programs/UnrealBuildTool/System/SourceFileWorkingSet.cs
+++ b/Engine/Source/Programs/UnrealBuildTool/System/SourceFileWorkingSet.cs
@@ -249,7 +249,7 @@ namespace UnrealBuildTool
 					MaxIdx--;
 				}
 
-				int ArrowIdx = Args.Data.IndexOf(" -> ", MinIdx, MaxIdx - MinIdx);
+				int ArrowIdx = Args.Data.IndexOf(" -> ", MinIdx, MaxIdx - MinIdx, StringComparison.Ordinal);
 				if (ArrowIdx == -1)
 				{
 					AddPath(Args.Data.Substring(MinIdx, MaxIdx - MinIdx));
@@ -280,7 +280,7 @@ namespace UnrealBuildTool
 		/// <param name="Path">Path to be added</param>
 		void AddPath(string Path)
 		{
-			if (Path.EndsWith("/"))
+			if (Path.EndsWith('/'))
 			{
 				Directories.Add(DirectoryReference.Combine(RootDir, Path));
 			}
diff --git a/Engine/Source/Programs/UnrealBuildTool/System/TargetMakefile.cs b/Engine/Source/Programs/UnrealBuildTool/System/TargetMakefile.cs
--- a/Engine/Source/Programs/UnrealBuildTool/System/TargetMakefile.cs
+++ b/Engine/Source/Programs/UnrealBuildTool/System/TargetMakefile.cs
@@ -984,7 +984,7 @@ namespace UnrealBuildTool
 
 				// First collect file timestamps by caching directories. This is much faster than letting each file do a GetFileAttribute
 				List<FileItem> fileItems = new(Makefile.InternalDependencies);
-				fileItems.Sort(static (a, b) => a.Location.FullName.CompareTo(b.Location.FullName));
+				fileItems.Sort(static (a, b) => a.Location.CompareTo(b.Location));
 				List<DirectoryItem> dirs = new();
 				DirectoryItem? lastDir = null;
 				foreach (FileItem item in fileItems)
diff --git a/Engine/Source/Programs/UnrealBuildTool/System/Unity.cs b/Engine/Source/Programs/UnrealBuildTool/System/Unity.cs
--- a/Engine/Source/Programs/UnrealBuildTool/System/Unity.cs
+++ b/Engine/Source/Programs/UnrealBuildTool/System/Unity.cs
@@ -221,8 +221,8 @@ namespace UnrealBuildTool
 					// declared in the header but are required to link. If they are placed later in the list, you will see
 					// compile errors because the templated function is instantiated but is defined later in the same translation unit
 					// which results in 'error C2908: explicit specialization; '*****' has already been instantiated'
-					bool bAIsGenerated = A.AbsolutePath.EndsWith(".gen.cpp") || CompileEnvironment.FileMatchesExtraGeneratedCPPTypes(A.AbsolutePath);
-					bool bBIsGenerated = B.AbsolutePath.EndsWith(".gen.cpp") || CompileEnvironment.FileMatchesExtraGeneratedCPPTypes(B.AbsolutePath);
+					bool bAIsGenerated = A.AbsolutePath.EndsWith(".gen.cpp", StringComparison.Ordinal) || CompileEnvironment.FileMatchesExtraGeneratedCPPTypes(A.AbsolutePath);
+					bool bBIsGenerated = B.AbsolutePath.EndsWith(".gen.cpp", StringComparison.Ordinal) || CompileEnvironment.FileMatchesExtraGeneratedCPPTypes(B.AbsolutePath);
 					if (bAIsGenerated != bBIsGenerated)
 					{
 						return bAIsGenerated && !bBIsGenerated ? -1 : 1;
@@ -385,7 +385,7 @@ namespace UnrealBuildTool
 			foreach (FileItem CPPFile in CPPFiles)
 			{
 				bool bHeaderInWorkingSet = false;
-				if (CPPFile.Name.EndsWith(".gen.cpp"))
+				if (CPPFile.Name.EndsWith(".gen.cpp", StringComparison.Ordinal))
 				{
 					// Ignore these files, does not need to be in adapt 
 				}
diff --git a/Engine/Source/Programs/UnrealBuildTool/System/UnrealPluginLanguage.cs b/Engine/Source/Programs/UnrealBuildTool/System/UnrealPluginLanguage.cs
--- a/Engine/Source/Programs/UnrealBuildTool/System/UnrealPluginLanguage.cs
+++ b/Engine/Source/Programs/UnrealBuildTool/System/UnrealPluginLanguage.cs
@@ -756,7 +756,7 @@ namespace UnrealBuildTool
 		private string ExpandVariables(UPLContext Context, string InputString)
 		{
 			string Result = InputString;
-			for (int Idx = Result.IndexOf("$B("); Idx != -1; Idx = Result.IndexOf("$B(", Idx))
+			for (int Idx = Result.IndexOf("$B(", StringComparison.Ordinal); Idx != -1; Idx = Result.IndexOf("$B(", Idx, StringComparison.Ordinal))
 			{
 				// Find the end of the variable name
 				int EndIdx = Result.IndexOf(')', Idx + 3);
@@ -782,7 +782,7 @@ namespace UnrealBuildTool
 				// Replace the variable, or skip past it
 				Result = Result.Substring(0, Idx) + Value.ToString().ToLower() + Result.Substring(EndIdx + 1);
 			}
-			for (int Idx = Result.IndexOf("$I("); Idx != -1; Idx = Result.IndexOf("$I(", Idx))
+			for (int Idx = Result.IndexOf("$I(", StringComparison.Ordinal); Idx != -1; Idx = Result.IndexOf("$I(", Idx, StringComparison.Ordinal))
 			{
 				// Find the end of the variable name
 				int EndIdx = Result.IndexOf(')', Idx + 3);
@@ -808,7 +808,7 @@ namespace UnrealBuildTool
 				// Replace the variable, or skip past it
 				Result = Result.Substring(0, Idx) + Value.ToString() + Result.Substring(EndIdx + 1);
 			}
-			for (int Idx = Result.IndexOf("$S("); Idx != -1; Idx = Result.IndexOf("$S(", Idx))
+			for (int Idx = Result.IndexOf("$S(", StringComparison.Ordinal); Idx != -1; Idx = Result.IndexOf("$S(", Idx, StringComparison.Ordinal))
 			{
 				// Find the end of the variable name
 				int EndIdx = Result.IndexOf(')', Idx + 3);
@@ -834,7 +834,7 @@ namespace UnrealBuildTool
 				// Replace the variable, or skip past it
 				Result = Result.Substring(0, Idx) + Value + Result.Substring(EndIdx + 1);
 			}
-			for (int Idx = Result.IndexOf("$E("); Idx != -1; Idx = Result.IndexOf("$E(", Idx))
+			for (int Idx = Result.IndexOf("$E(", StringComparison.Ordinal); Idx != -1; Idx = Result.IndexOf("$E(", Idx, StringComparison.Ordinal))
 			{
 				// Find the end of the variable name
 				int EndIdx = Result.IndexOf(')', Idx + 3);
@@ -870,7 +870,7 @@ namespace UnrealBuildTool
 			foreach (XAttribute Attrib in Node.Attributes())
 			{
 				string AttribStr = Attrib.ToString();
-				if (AttribStr.StartsWith("__line="))
+				if (AttribStr.StartsWith("__line=", StringComparison.Ordinal))
 				{
 					LineAttrib = Attrib;
 				}
@@ -1054,7 +1054,7 @@ namespace UnrealBuildTool
 
 		XNamespace? TrimNamespaceAliasFromName(ref string Name)
 		{
-			int Index = Name.IndexOf(":");
+			int Index = Name.IndexOf(':', StringComparison.Ordinal);
 			if (Index >= 0)
 			{
 				XNamespace? Namespace;
@@ -1348,7 +1348,7 @@ namespace UnrealBuildTool
 							string? Value = GetAttribute(CurrentContext, Node, "value");
 							if (Tag != null && Name != null && Value != null)
 							{
-								if (Tag.StartsWith("$"))
+								if (Tag.StartsWith('$'))
 								{
 									XElement? Target = CurrentElement;
 									if (Tag.Length > 1)
@@ -1385,7 +1385,7 @@ namespace UnrealBuildTool
 							string? Name = GetAttribute(CurrentContext, Node, "name");
 							if (Tag != null && Name != null)
 							{
-								if (Tag.StartsWith("$"))
+								if (Tag.StartsWith('$'))
 								{
 									XElement? Target = CurrentElement;
 									if (Tag.Length > 1)
@@ -1634,7 +1634,7 @@ namespace UnrealBuildTool
 										continue;
 									}
 								}
-								if (Tag.StartsWith("$"))
+								if (Tag.StartsWith('$'))
 								{
 									XElement? Target = CurrentElement;
 									if (Tag.Length > 1)
@@ -1682,7 +1682,7 @@ namespace UnrealBuildTool
 							bool bOnce = StringToBool(GetAttribute(CurrentContext, Node, "once", true, false));
 							if (Tag != null)
 							{
-								if (Tag.StartsWith("$"))
+								if (Tag.StartsWith('$'))
 								{
 									XElement? Target = CurrentElement;
 									if (Tag.Length > 1)
@@ -1989,7 +1989,7 @@ namespace UnrealBuildTool
 							string Find = GetAttribute(CurrentContext, Node, "find", true, false, "");
 							if (Result != null)
 							{
-								CurrentContext.BoolVariables[Result] = Source.StartsWith(Find);
+								CurrentContext.BoolVariables[Result] = Source.StartsWith(Find, StringComparison.Ordinal);
 							}
 						}
 						break;
@@ -2001,7 +2001,7 @@ namespace UnrealBuildTool
 							string Find = GetAttribute(CurrentContext, Node, "find", true, false, "");
 							if (Result != null)
 							{
-								CurrentContext.BoolVariables[Result] = Source.EndsWith(Find);
+								CurrentContext.BoolVariables[Result] = Source.EndsWith(Find, StringComparison.Ordinal);
 							}
 						}
 						break;
@@ -2258,7 +2258,7 @@ namespace UnrealBuildTool
 							string Find = GetAttribute(CurrentContext, Node, "find", true, false, "");
 							if (Result != null)
 							{
-								CurrentContext.IntVariables[Result] = Source.IndexOf(Find);
+								CurrentContext.IntVariables[Result] = Source.IndexOf(Find, StringComparison.Ordinal);
 							}
 						}
 						break;
@@ -2311,7 +2311,7 @@ namespace UnrealBuildTool
 							if (Result != null)
 							{
 								XElement? Element = CurrentElement;
-								if (Tag.StartsWith("$"))
+								if (Tag.StartsWith('$'))
 								{
 									if (Tag.Length > 1)
 									{
@@ -2338,7 +2338,7 @@ namespace UnrealBuildTool
 							if (Result != null && Tag != null && Name != null)
 							{
 								XElement? Element = CurrentElement;
-								if (Tag.StartsWith("$"))
+								if (Tag.StartsWith('$'))
 								{
 									if (Tag.Length > 1)
 									{
@@ -2368,7 +2368,7 @@ namespace UnrealBuildTool
 							if (Result != null)
 							{
 								XElement? Element = CurrentElement;
-								if (Tag.StartsWith("$"))
+								if (Tag.StartsWith('$'))
 								{
 									if (Tag.Length > 1)
 									{
@@ -2519,7 +2519,7 @@ namespace UnrealBuildTool
 							if (Result != null)
 							{
 								string Value = DefaultVal;
-								int Index = Source.IndexOf(Find);
+								int Index = Source.IndexOf(Find, StringComparison.Ordinal);
 								if (Index != -1)
 								{
 									Index += Find.Length;
@@ -3214,16 +3214,16 @@ namespace UnrealBuildTool
 		/// </summary>
 		private ParseAction GetActionForLine(ref string Line)
 		{
-			if (String.IsNullOrEmpty(Line) || Line.StartsWith(";") || Line.StartsWith("//"))
+			if (String.IsNullOrEmpty(Line) || Line.StartsWith(';') || Line.StartsWith("//", StringComparison.Ordinal))
 			{
 				return ParseAction.None;
 			}
-			else if (Line.StartsWith("-"))
+			else if (Line.StartsWith('-'))
 			{
 				Line = Line.Substring(1).TrimStart();
 				return ParseAction.Remove;
 			}
-			else if (Line.StartsWith("+"))
+			else if (Line.StartsWith('+'))
 			{
 				Line = Line.Substring(1).TrimStart();
 				return ParseAction.Add;
@@ -3277,14 +3277,14 @@ namespace UnrealBuildTool
 					string TrimmedLine = Line.Trim();
 					// Multiline value support
 					bool bWasMultiLine = bMultiLine;
-					bMultiLine = TrimmedLine.EndsWith("\\");
+					bMultiLine = TrimmedLine.EndsWith('\\');
 					if (bMultiLine)
 					{
 						TrimmedLine = TrimmedLine.Substring(0, TrimmedLine.Length - 1).TrimEnd();
 					}
 					if (!bWasMultiLine)
 					{
-						if (TrimmedLine.StartsWith("["))
+						if (TrimmedLine.StartsWith('['))
 						{
 							CurrentSection = FindOrAddSection(TrimmedLine, Filename, LineIndex);
 							LastAction = ParseAction.None;
@@ -3374,7 +3374,7 @@ namespace UnrealBuildTool
 				throw new IniParsingException("Empty key when parsing {0}, line {1}: {2}", Filename, LineIndex, TrimmedLine);
 			}
 			Value = TrimmedLine.Substring(AssignIndex + 1).Trim();
-			if (Value.StartsWith("\""))
+			if (Value.StartsWith('\"'))
 			{
 				// Remove quotes
 				int QuoteEnd = Value.LastIndexOf('\"');
diff --git a/Engine/Source/Programs/UnrealBuildTool/System/Utils.cs b/Engine/Source/Programs/UnrealBuildTool/System/Utils.cs
--- a/Engine/Source/Programs/UnrealBuildTool/System/Utils.cs
+++ b/Engine/Source/Programs/UnrealBuildTool/System/Utils.cs
@@ -64,12 +64,12 @@ namespace UnrealBuildTool
 			{
 				// Convoluted way of stripping first and last character and '(' in the case of $(ENV) to get to ENV
 				string EnvironmentVariable = M.ToString();
-				if (EnvironmentVariable.StartsWith("$") && EnvironmentVariable.EndsWith(")"))
+				if (EnvironmentVariable.StartsWith('$') && EnvironmentVariable.EndsWith(')'))
 				{
 					EnvironmentVariable = EnvironmentVariable.Substring(1, EnvironmentVariable.Length - 2).Replace("(", "", StringComparison.Ordinal);
 				}
 
-				if (EnvironmentVariable.StartsWith("%") && EnvironmentVariable.EndsWith("%"))
+				if (EnvironmentVariable.StartsWith('%') && EnvironmentVariable.EndsWith('%'))
 				{
 					EnvironmentVariable = EnvironmentVariable.Substring(1, EnvironmentVariable.Length - 2);
 				}
@@ -679,13 +679,13 @@ namespace UnrealBuildTool
 				}
 
 				// Consider empty paths or paths which start with .. or /.. as invalid
-				if (LocalString.StartsWith("..") || LocalString.StartsWith(ParentDir))
+				if (LocalString.StartsWith("..", StringComparison.Ordinal) || LocalString.StartsWith(ParentDir, StringComparison.Ordinal))
 				{
 					return InPath;
 				}
 
 				// If there are no "/.."s left then we're done
-				int Index = LocalString.IndexOf(ParentDir);
+				int Index = LocalString.IndexOf(ParentDir, StringComparison.Ordinal);
 				if (Index == -1)
 				{
 					break;
@@ -695,7 +695,7 @@ namespace UnrealBuildTool
 				for (; ; )
 				{
 					// Find the previous slash
-					PreviousSeparatorIndex = Math.Max(0, LocalString.LastIndexOf("/", PreviousSeparatorIndex - 1));
+					PreviousSeparatorIndex = Math.Max(0, LocalString.LastIndexOf('/', PreviousSeparatorIndex - 1));
 
 					// Stop if we've hit the start of the string
 					if (PreviousSeparatorIndex == 0)
@@ -711,7 +711,7 @@ namespace UnrealBuildTool
 				}
 
 				// If we're attempting to remove the drive letter, that's illegal
-				int Colon = LocalString.IndexOf(":", PreviousSeparatorIndex);
+				int Colon = LocalString.IndexOf(':', PreviousSeparatorIndex);
 				if (Colon >= 0 && Colon < Index)
 				{
 					return InPath;
@@ -752,7 +752,7 @@ namespace UnrealBuildTool
 			{
 				AbsolutePath = Path.GetFullPath(SourcePath);
 			}
-			bool SourcePathEndsWithDirectorySeparator = AbsolutePath.EndsWith(Path.DirectorySeparatorChar.ToString()) || AbsolutePath.EndsWith(Path.AltDirectorySeparatorChar.ToString());
+			bool SourcePathEndsWithDirectorySeparator = AbsolutePath.EndsWith(Path.DirectorySeparatorChar) || AbsolutePath.EndsWith(Path.AltDirectorySeparatorChar);
 			if (AlwaysTreatSourceAsDirectory && !SourcePathEndsWithDirectorySeparator)
 			{
 				AbsolutePath += Path.DirectorySeparatorChar;
@@ -768,7 +768,7 @@ namespace UnrealBuildTool
 
 			// Make sure the directory has a trailing directory separator so that the relative directory that
 			// MakeRelativeUri creates doesn't include our directory -- only the directories beneath it!
-			if (!AbsoluteRelativeDirectory.EndsWith(Path.DirectorySeparatorChar.ToString()) && !AbsoluteRelativeDirectory.EndsWith(Path.AltDirectorySeparatorChar.ToString()))
+			if (!AbsoluteRelativeDirectory.EndsWith(Path.DirectorySeparatorChar) && !AbsoluteRelativeDirectory.EndsWith(Path.AltDirectorySeparatorChar))
 			{
 				AbsoluteRelativeDirectory += Path.DirectorySeparatorChar;
 			}
@@ -781,7 +781,7 @@ namespace UnrealBuildTool
 			string RelativePath = Uri.UnescapeDataString(UriRelativePath.ToString()).Replace('/', Path.DirectorySeparatorChar);
 
 			// If we added a directory separator character earlier on, remove it now
-			if (!SourcePathEndsWithDirectorySeparator && AlwaysTreatSourceAsDirectory && RelativePath.EndsWith(Path.DirectorySeparatorChar.ToString()))
+			if (!SourcePathEndsWithDirectorySeparator && AlwaysTreatSourceAsDirectory && RelativePath.EndsWith(Path.DirectorySeparatorChar))
 			{
 				RelativePath = RelativePath.Substring(0, RelativePath.Length - 1);
 			}
@@ -927,11 +927,11 @@ namespace UnrealBuildTool
 		public static bool IsFileUnderDirectory(string FilePath, string Directory)
 		{
 			string DirectoryPathPlusSeparator = Path.GetFullPath(Directory);
-			if (!DirectoryPathPlusSeparator.EndsWith(Path.DirectorySeparatorChar.ToString()))
+			if (!DirectoryPathPlusSeparator.EndsWith(Path.DirectorySeparatorChar))
 			{
 				DirectoryPathPlusSeparator += Path.DirectorySeparatorChar;
 			}
-			return Path.GetFullPath(FilePath).StartsWith(DirectoryPathPlusSeparator, StringComparison.InvariantCultureIgnoreCase);
+			return Path.GetFullPath(FilePath).StartsWith(DirectoryPathPlusSeparator, StringComparison.OrdinalIgnoreCase);
 		}
 
 		/// <summary>
@@ -1022,7 +1022,7 @@ namespace UnrealBuildTool
 				// will show up as it's own logical "cpu".  Including the socket number allows us to count multi-processor
 				// system cores correctly
 				string Output = RunLocalProcessAndReturnStdOut("lscpu", "-p='SOCKET,CPU'");
-				List<string> CPUs = Output.Split("\n").Where(x => !x.StartsWith("#")).ToList();
+				List<string> CPUs = Output.Split("\n").Where(x => !x.StartsWith('#')).ToList();
 
 				return CPUs.Count;
 			}
@@ -1093,7 +1093,7 @@ namespace UnrealBuildTool
 				// using the HashSet ensures that those duplicates are removed and we only count the first "cpu" found
 				// for each "core".  Including the socket number allows us to count multi-processor system cores correctly
 				string Output = RunLocalProcessAndReturnStdOut("lscpu", "-p='SOCKET,CORE'");
-				HashSet<string> CPUs = Output.Split("\n").Where(x => !x.StartsWith("#")).ToHashSet();
+				HashSet<string> CPUs = Output.Split("\n").Where(x => !x.StartsWith('#')).ToHashSet();
 
 				return CPUs.Count;
 			}
@@ -1514,7 +1514,7 @@ namespace UnrealBuildTool
 				List<FileReference> TargetProjectFiles = new List<FileReference>();
 				for (int Idx = 0; Idx < Arguments.Count; Idx++)
 				{
-					if (!Arguments[Idx].StartsWith("-") && Arguments[Idx].EndsWith(".uproject", StringComparison.OrdinalIgnoreCase))
+					if (!Arguments[Idx].StartsWith('-') && Arguments[Idx].EndsWith(".uproject", StringComparison.OrdinalIgnoreCase))
 					{
 						CandidateProjectPath = Arguments[Idx];
 
diff --git a/Engine/Source/Programs/UnrealBuildTool/ToolChain/IWYUToolChain.cs b/Engine/Source/Programs/UnrealBuildTool/ToolChain/IWYUToolChain.cs
--- a/Engine/Source/Programs/UnrealBuildTool/ToolChain/IWYUToolChain.cs
+++ b/Engine/Source/Programs/UnrealBuildTool/ToolChain/IWYUToolChain.cs
@@ -264,11 +264,11 @@ namespace UnrealBuildTool
 					}
 				}
 
-				if (SourceFile.Name != "SingleFile.cpp" && !SourceFile.Name.EndsWith(".gen.cpp"))
+				if (SourceFile.Name != "SingleFile.cpp" && !SourceFile.Name.EndsWith(".gen.cpp", StringComparison.Ordinal))
 				{
 					foreach (string include in CompileEnvironment.MetadataCache.GetHeaderIncludes(SourceFile))
 					{
-						if (include.EndsWith(".ipp"))
+						if (include.EndsWith(".ipp", StringComparison.Ordinal))
 						{
 							Arguments.Add($"-Xiwyu --check_also=*/{include.Replace('\\', '/')}");
 						}
