Add missing Configuration property when creating OriginalTargetDescription.

#Synced-CL 49329310

diff --git a/Engine/Source/Programs/UnrealBuildTool/Configuration/Descriptors/OriginalTargetDescriptor.cs b/Engine/Source/Programs/UnrealBuildTool/Configuration/Descriptors/OriginalTargetDescriptor.cs
--- a/Engine/Source/Programs/UnrealBuildTool/Configuration/Descriptors/OriginalTargetDescriptor.cs
+++ b/Engine/Source/Programs/UnrealBuildTool/Configuration/Descriptors/OriginalTargetDescriptor.cs
@@ -108,6 +108,7 @@ internal record class OriginalTargetDescriptor
 			ProjectFile = targetDescriptor.ProjectFile,
 			Name = targetDescriptor.Name,
 			Platform = targetDescriptor.Platform,
+			Configuration = targetDescriptor.Configuration,
 			Architectures = new(targetDescriptor.Architectures),
 			AdditionalArguments = targetDescriptor.AdditionalArguments,
 			IsTestsTarget = targetDescriptor.IsTestsTarget,
