UBT: Move VCSharpProjectFile to its own file.

#Synced-CL 49366587

diff --git a/Engine/Source/Programs/UnrealBuildTool/ProjectFiles/VisualStudio/VCProject.cs b/Engine/Source/Programs/UnrealBuildTool/ProjectFiles/VisualStudio/VCProject.cs
--- a/Engine/Source/Programs/UnrealBuildTool/ProjectFiles/VisualStudio/VCProject.cs
+++ b/Engine/Source/Programs/UnrealBuildTool/ProjectFiles/VisualStudio/VCProject.cs
@@ -8,7 +8,6 @@ using System.IO;
 using System.Linq;
 using System.Security;
 using System.Text;
-using System.Text.RegularExpressions;
 using System.Xml;
 using System.Xml.Linq;
 using EpicGames.Core;
@@ -2584,224 +2583,4 @@ namespace UnrealBuildTool
 			}
 		}
 	}
-
-	/// <summary>
-	/// A Visual C# project.
-	/// </summary>
-	class VCSharpProjectFile : MSBuildProjectFile
-	{
-		/// <summary>
-		/// This is the GUID that Visual Studio uses to identify a C# project file in the solution
-		/// </summary>
-		public override string ProjectTypeGUID => "{FAE04EC0-301F-11D3-BF4B-00C04F79EFBC}";
-
-		/// <summary>
-		/// Platforms that this project supports
-		/// </summary>
-		public HashSet<string> Platforms = new HashSet<string>();
-
-		/// <summary>
-		/// Configurations that this project supports
-		/// </summary>
-		public HashSet<string> Configurations = new HashSet<string>();
-
-		/// <summary>
-		/// Constructs a new project file object
-		/// </summary>
-		/// <param name="InitFilePath">The path to the project file on disk</param>
-		/// <param name="Logger">Logger for output</param>
-		/// <param name="BaseDir">The base directory for files within this project - if not specified, InitFilePath.Directory will be used</param>
-		public VCSharpProjectFile(FileReference InitFilePath, ILogger Logger, DirectoryReference? BaseDir = null)
-			: base(InitFilePath, BaseDir ?? InitFilePath.Directory, true)
-		{
-			try
-			{
-				XmlDocument Document = new XmlDocument();
-				Document.Load(InitFilePath.FullName);
-
-				// Check the root element is the right type
-				if (Document.DocumentElement?.Name != "Project")
-				{
-					throw new BuildException("Unexpected root element '{0}' in project file", Document.DocumentElement?.Name);
-				}
-
-				// Parse all the configurations and platforms
-				// Parse the basic structure of the document, updating properties and recursing into other referenced projects as we go
-				if (!IsDotnetCoreOrDotnet5Plus())
-				{
-					foreach (XmlElement Element in Document.DocumentElement.ChildNodes.OfType<XmlElement>())
-					{
-						if (Element.Name == "PropertyGroup")
-						{
-							string Condition = Element.GetAttribute("Condition");
-							if (!String.IsNullOrEmpty(Condition))
-							{
-								Match Match = Regex.Match(Condition, "^\\s*'\\$\\(Configuration\\)\\|\\$\\(Platform\\)'\\s*==\\s*'(.+)\\|(.+)'\\s*$");
-								if (Match.Success && Match.Groups.Count == 3)
-								{
-									Configurations.Add(Match.Groups[1].Value);
-									Platforms.Add(Match.Groups[2].Value);
-								}
-								else
-								{
-									Logger.LogWarning("Unable to parse configuration/platform from condition '{InitFilePath}': {Condition}", InitFilePath, Condition);
-								}
-							}
-						}
-					}
-				}
-				else
-				{
-					foreach (string c in GetProjectProperty("Configurations").Split(';'))
-					{
-						Configurations.Add(c);
-					}
-					bool ConfigurationsFound = Configurations.Any();
-
-					if (!ConfigurationsFound)
-					{
-						foreach (XmlElement PropertyGroup in Document.DocumentElement.ChildNodes.OfType<XmlElement>()
-							.Where(element => element.Name == "PropertyGroup"))
-						{
-							XmlNodeList ConfigNodeList = PropertyGroup.GetElementsByTagName("Configurations");
-							// if this property group does not set configurations we do not care about it
-							if (ConfigNodeList.Count == 0)
-							{
-								continue;
-							}
-
-							if (PropertyGroup.HasAttribute("Condition"))
-							{
-								string Condition = PropertyGroup.GetAttribute("Condition");
-								Logger.LogWarning("Unable to parse configuration from property group with condition '{InitFilePath}': {Condition}. UBT Requires you to set the configuration without conditionals.", InitFilePath, Condition);
-								continue;
-							}
-							string[]? ParsedConfigurations = ConfigNodeList[0]?.FirstChild?.Value?.Split(';');
-							if (ParsedConfigurations != null)
-							{
-								foreach (string c in ParsedConfigurations)
-								{
-									Configurations.Add(c);
-								}
-							}
-
-							// platforms change meaning quite a bit in .net core but typically you do not specify this and its derived from the build instead
-							// for most intents it is just Any CPU from .net framework
-							Platforms.Add("AnyCPU");
-
-							ConfigurationsFound = true;
-							break;
-						}
-					}
-
-					// dotnet does not require you to specify configurations or platforms, if you do not debug and release are the defaults
-					if (!ConfigurationsFound)
-					{
-						Configurations.Add("Debug");
-						Configurations.Add("Release");
-						Platforms.Add("AnyCPU");
-					}
-				}
-			}
-			catch (Exception Ex)
-			{
-				Logger.LogWarning("Unable to parse {Path}: {Ex}", InitFilePath, Ex.ToString());
-			}
-		}
-
-		/// <summary>
-		/// Extract information from the csproj file based on the supplied configuration
-		/// </summary>
-		public IDotnetProjectInfo GetProjectInfo(UnrealTargetConfiguration InConfiguration)
-		{
-			if (CachedProjectInfo.TryGetValue(InConfiguration, out IDotnetProjectInfo? value))
-			{
-				return value;
-			}
-
-			Dictionary<string, string> Properties = [];
-			Properties.Add("Platform", "AnyCPU");
-			Properties.Add("Configuration", InConfiguration.ToString());
-			Properties.Add("EngineDirectory", Unreal.EngineDirectory.FullName);
-
-			DotnetProjectInfo Info = DotnetProjectInfo.Read(ProjectFilePath, Properties);
-			CachedProjectInfo.Add(InConfiguration, Info);
-			return Info;
-		}
-
-		/// <summary>
-		/// Determine if this project is a .NET Core or .NET 5+ project
-		/// </summary>
-		public bool IsDotnetCoreOrDotnet5Plus()
-		{
-			IDotnetProjectInfo Info = GetProjectInfo(UnrealTargetConfiguration.Debug);
-			return Info.IsDotnetCoreOrDotnet5Plus();
-		}
-
-		/// <summary>
-		/// Gets a property from the project
-		/// </summary>
-		public string GetProjectProperty(string property)
-		{
-			IDotnetProjectInfo Info = GetProjectInfo(UnrealTargetConfiguration.Debug)!;
-			Info.Properties.TryGetValue(property, out string? value);
-			return value ?? String.Empty;
-		}
-
-		/// <inheritdoc/>
-		public override MSBuildProjectContext? GetMatchingProjectContext(TargetType SolutionTarget, UnrealTargetConfiguration SolutionConfiguration, UnrealTargetPlatform SolutionPlatform, PlatformProjectGeneratorCollection PlatformProjectGenerators, UnrealArch? Architecture, ILogger Logger)
-		{
-			// Find the matching platform name
-			string ProjectPlatformName;
-			if (Platforms.Contains("x64"))
-			{
-				ProjectPlatformName = "x64";
-			}
-			else
-			{
-				ProjectPlatformName = "Any CPU";
-			}
-
-			// Find the matching configuration
-			string ProjectConfigurationName;
-			if (Configurations.Contains(SolutionConfiguration.ToString()))
-			{
-				ProjectConfigurationName = SolutionConfiguration.ToString();
-			}
-			else if (Configurations.Contains("Development"))
-			{
-				ProjectConfigurationName = "Development";
-			}
-			else
-			{
-				ProjectConfigurationName = "Release";
-			}
-
-			// Figure out whether to build it by default
-			bool bBuildByDefault = ShouldBuildByDefaultForSolutionTargets;
-			if (SolutionTarget == TargetType.Game || SolutionTarget == TargetType.Editor)
-			{
-				bBuildByDefault = true;
-			}
-
-			// Create the context
-			return new MSBuildProjectContext(ProjectConfigurationName, ProjectPlatformName) { bBuildByDefault = bBuildByDefault };
-		}
-
-		/// <summary>
-		/// Basic csproj file support. Generates C# library project with one build config.
-		/// </summary>
-		/// <param name="InPlatforms">Not used.</param>
-		/// <param name="InConfigurations">Not Used.</param>
-		/// <param name="PlatformProjectGenerators">Set of platform project generators</param>
-		/// <param name="Logger"></param>
-		/// <returns>true if the opration was successful, false otherwise</returns>
-		public override bool WriteProjectFile(List<UnrealTargetPlatform> InPlatforms, List<UnrealTargetConfiguration> InConfigurations, PlatformProjectGeneratorCollection PlatformProjectGenerators, ILogger Logger)
-		{
-			throw new BuildException("Support for writing C# projects from UnrealBuildTool has been removed.");
-		}
-
-		/// Cache of parsed info about this project
-		protected readonly Dictionary<UnrealTargetConfiguration, IDotnetProjectInfo> CachedProjectInfo = [];
-	}
 }
diff --git a/Engine/Source/Programs/UnrealBuildTool/ProjectFiles/VisualStudio/VCSharpProjectFile.cs b/Engine/Source/Programs/UnrealBuildTool/ProjectFiles/VisualStudio/VCSharpProjectFile.cs
new file mode 100644
--- /dev/null
+++ b/Engine/Source/Programs/UnrealBuildTool/ProjectFiles/VisualStudio/VCSharpProjectFile.cs
@@ -0,0 +1,232 @@
+// Copyright Epic Games, Inc. All Rights Reserved.
+
+using System;
+using System.Collections.Generic;
+using System.Linq;
+using System.Text.RegularExpressions;
+using System.Xml;
+using EpicGames.Core;
+using Microsoft.Extensions.Logging;
+using UnrealBuildBase;
+
+namespace UnrealBuildTool;
+
+/// <summary>
+/// Represents an existing C# project file, created by using MSBuild to parse the file from disk.
+/// </summary>
+class VCSharpProjectFile : MSBuildProjectFile
+{
+	/// <summary>
+	/// This is the GUID that Visual Studio uses to identify a C# project file in the solution
+	/// </summary>
+	public override string ProjectTypeGUID => "{FAE04EC0-301F-11D3-BF4B-00C04F79EFBC}";
+
+	/// <summary>
+	/// Platforms that this project supports
+	/// </summary>
+	public HashSet<string> Platforms = new HashSet<string>();
+
+	/// <summary>
+	/// Configurations that this project supports
+	/// </summary>
+	public HashSet<string> Configurations = new HashSet<string>();
+
+	/// <summary>
+	/// Constructs a new project file object
+	/// </summary>
+	/// <param name="InitFilePath">The path to the project file on disk</param>
+	/// <param name="Logger">Logger for output</param>
+	/// <param name="BaseDir">The base directory for files within this project - if not specified, InitFilePath.Directory will be used</param>
+	public VCSharpProjectFile(FileReference InitFilePath, ILogger Logger, DirectoryReference? BaseDir = null)
+		: base(InitFilePath, BaseDir ?? InitFilePath.Directory, true)
+	{
+		try
+		{
+			XmlDocument Document = new XmlDocument();
+			Document.Load(InitFilePath.FullName);
+
+			// Check the root element is the right type
+			if (Document.DocumentElement?.Name != "Project")
+			{
+				throw new BuildException("Unexpected root element '{0}' in project file", Document.DocumentElement?.Name);
+			}
+
+			// Parse all the configurations and platforms
+			// Parse the basic structure of the document, updating properties and recursing into other referenced projects as we go
+			if (!IsDotnetCoreOrDotnet5Plus())
+			{
+				foreach (XmlElement Element in Document.DocumentElement.ChildNodes.OfType<XmlElement>())
+				{
+					if (Element.Name == "PropertyGroup")
+					{
+						string Condition = Element.GetAttribute("Condition");
+						if (!String.IsNullOrEmpty(Condition))
+						{
+							Match Match = Regex.Match(Condition, "^\\s*'\\$\\(Configuration\\)\\|\\$\\(Platform\\)'\\s*==\\s*'(.+)\\|(.+)'\\s*$");
+							if (Match.Success && Match.Groups.Count == 3)
+							{
+								Configurations.Add(Match.Groups[1].Value);
+								Platforms.Add(Match.Groups[2].Value);
+							}
+							else
+							{
+								Logger.LogWarning("Unable to parse configuration/platform from condition '{InitFilePath}': {Condition}", InitFilePath, Condition);
+							}
+						}
+					}
+				}
+			}
+			else
+			{
+				foreach (string c in GetProjectProperty("Configurations").Split(';'))
+				{
+					Configurations.Add(c);
+				}
+				bool ConfigurationsFound = Configurations.Any();
+
+				if (!ConfigurationsFound)
+				{
+					foreach (XmlElement PropertyGroup in Document.DocumentElement.ChildNodes.OfType<XmlElement>()
+						.Where(element => element.Name == "PropertyGroup"))
+					{
+						XmlNodeList ConfigNodeList = PropertyGroup.GetElementsByTagName("Configurations");
+						// if this property group does not set configurations we do not care about it
+						if (ConfigNodeList.Count == 0)
+						{
+							continue;
+						}
+
+						if (PropertyGroup.HasAttribute("Condition"))
+						{
+							string Condition = PropertyGroup.GetAttribute("Condition");
+							Logger.LogWarning("Unable to parse configuration from property group with condition '{InitFilePath}': {Condition}. UBT Requires you to set the configuration without conditionals.", InitFilePath, Condition);
+							continue;
+						}
+						string[]? ParsedConfigurations = ConfigNodeList[0]?.FirstChild?.Value?.Split(';');
+						if (ParsedConfigurations != null)
+						{
+							foreach (string c in ParsedConfigurations)
+							{
+								Configurations.Add(c);
+							}
+						}
+
+						// platforms change meaning quite a bit in .net core but typically you do not specify this and its derived from the build instead
+						// for most intents it is just Any CPU from .net framework
+						Platforms.Add("AnyCPU");
+
+						ConfigurationsFound = true;
+						break;
+					}
+				}
+
+				// dotnet does not require you to specify configurations or platforms, if you do not debug and release are the defaults
+				if (!ConfigurationsFound)
+				{
+					Configurations.Add("Debug");
+					Configurations.Add("Release");
+					Platforms.Add("AnyCPU");
+				}
+			}
+		}
+		catch (Exception Ex)
+		{
+			Logger.LogWarning("Unable to parse {Path}: {Ex}", InitFilePath, Ex.ToString());
+		}
+	}
+
+	/// <summary>
+	/// Extract information from the csproj file based on the supplied configuration
+	/// </summary>
+	public IDotnetProjectInfo GetProjectInfo(UnrealTargetConfiguration InConfiguration)
+	{
+		if (CachedProjectInfo.TryGetValue(InConfiguration, out IDotnetProjectInfo? value))
+		{
+			return value;
+		}
+
+		Dictionary<string, string> Properties = [];
+		Properties.Add("Platform", "AnyCPU");
+		Properties.Add("Configuration", InConfiguration.ToString());
+		Properties.Add("EngineDirectory", Unreal.EngineDirectory.FullName);
+
+		DotnetProjectInfo Info = DotnetProjectInfo.Read(ProjectFilePath, Properties);
+		CachedProjectInfo.Add(InConfiguration, Info);
+		return Info;
+	}
+
+	/// <summary>
+	/// Determine if this project is a .NET Core or .NET 5+ project
+	/// </summary>
+	public bool IsDotnetCoreOrDotnet5Plus()
+	{
+		IDotnetProjectInfo Info = GetProjectInfo(UnrealTargetConfiguration.Debug);
+		return Info.IsDotnetCoreOrDotnet5Plus();
+	}
+
+	/// <summary>
+	/// Gets a property from the project
+	/// </summary>
+	public string GetProjectProperty(string property)
+	{
+		IDotnetProjectInfo Info = GetProjectInfo(UnrealTargetConfiguration.Debug)!;
+		Info.Properties.TryGetValue(property, out string? value);
+		return value ?? String.Empty;
+	}
+
+	/// <inheritdoc/>
+	public override MSBuildProjectContext? GetMatchingProjectContext(TargetType SolutionTarget, UnrealTargetConfiguration SolutionConfiguration, UnrealTargetPlatform SolutionPlatform, PlatformProjectGeneratorCollection PlatformProjectGenerators, UnrealArch? Architecture, ILogger Logger)
+	{
+		// Find the matching platform name
+		string ProjectPlatformName;
+		if (Platforms.Contains("x64"))
+		{
+			ProjectPlatformName = "x64";
+		}
+		else
+		{
+			ProjectPlatformName = "Any CPU";
+		}
+
+		// Find the matching configuration
+		string ProjectConfigurationName;
+		if (Configurations.Contains(SolutionConfiguration.ToString()))
+		{
+			ProjectConfigurationName = SolutionConfiguration.ToString();
+		}
+		else if (Configurations.Contains("Development"))
+		{
+			ProjectConfigurationName = "Development";
+		}
+		else
+		{
+			ProjectConfigurationName = "Release";
+		}
+
+		// Figure out whether to build it by default
+		bool bBuildByDefault = ShouldBuildByDefaultForSolutionTargets;
+		if (SolutionTarget == TargetType.Game || SolutionTarget == TargetType.Editor)
+		{
+			bBuildByDefault = true;
+		}
+
+		// Create the context
+		return new MSBuildProjectContext(ProjectConfigurationName, ProjectPlatformName) { bBuildByDefault = bBuildByDefault };
+	}
+
+	/// <summary>
+	/// Basic csproj file support. Generates C# library project with one build config.
+	/// </summary>
+	/// <param name="InPlatforms">Not used.</param>
+	/// <param name="InConfigurations">Not Used.</param>
+	/// <param name="PlatformProjectGenerators">Set of platform project generators</param>
+	/// <param name="Logger"></param>
+	/// <returns>true if the opration was successful, false otherwise</returns>
+	public override bool WriteProjectFile(List<UnrealTargetPlatform> InPlatforms, List<UnrealTargetConfiguration> InConfigurations, PlatformProjectGeneratorCollection PlatformProjectGenerators, ILogger Logger)
+	{
+		throw new BuildException("Support for writing C# projects from UnrealBuildTool has been removed.");
+	}
+
+	/// Cache of parsed info about this project
+	protected readonly Dictionary<UnrealTargetConfiguration, IDotnetProjectInfo> CachedProjectInfo = [];
+}
