[UBT/UBA]
* Added config option -UBAMaxRacingPercent=<percent> which can be used to have local machine race remotes for active actions when there is less work scheduled locally. The percentage value controls maximum utilization until it stops filling up with racing processes

#Synced-CL 49373559

diff --git a/Engine/Source/Programs/UnrealBuildTool/Executors/UnrealBuildAccelerator/UBAExecutor.cs b/Engine/Source/Programs/UnrealBuildTool/Executors/UnrealBuildAccelerator/UBAExecutor.cs
--- a/Engine/Source/Programs/UnrealBuildTool/Executors/UnrealBuildAccelerator/UBAExecutor.cs
+++ b/Engine/Source/Programs/UnrealBuildTool/Executors/UnrealBuildAccelerator/UBAExecutor.cs
@@ -308,6 +308,11 @@ namespace UnrealBuildTool
 				config.AddValue("Session", "AllowMemoryMaps", "false");
 			}
 
+			if (UBAConfig.MaxRacingPercent > 0)
+			{
+				config.AddValue("Scheduler", "MaxRacingPercent", UBAConfig.MaxRacingPercent.ToString());
+			}
+
 			EpicGames.UBA.ILogger ubaLogger = EpicGames.UBA.ILogger.GlobalLogger!;
 			_server = IServer.CreateServer(UBAConfig.MaxWorkers, UBAConfig.SendSize, ubaLogger, UBAConfig.bUseQuic);
 			_storageServer = IStorageServer.CreateStorageServer(_server, ubaLogger, new StorageServerCreateInfo(_rootDirRef.FullName, ((ulong)UBAConfig.StoreCapacityGb) * 1000 * 1000 * 1000, !UBAConfig.bStoreRaw, UBAConfig.Zone));
diff --git a/Engine/Source/Programs/UnrealBuildTool/Executors/UnrealBuildAccelerator/UnrealBuildAcceleratorConfig.cs b/Engine/Source/Programs/UnrealBuildTool/Executors/UnrealBuildAccelerator/UnrealBuildAcceleratorConfig.cs
--- a/Engine/Source/Programs/UnrealBuildTool/Executors/UnrealBuildAccelerator/UnrealBuildAcceleratorConfig.cs
+++ b/Engine/Source/Programs/UnrealBuildTool/Executors/UnrealBuildAccelerator/UnrealBuildAcceleratorConfig.cs
@@ -324,6 +324,17 @@ namespace UnrealBuildTool
 		[CommandLine("-UBASharedMemoryTempFile")]
 		public string SharedMemoryTempFile { get; set; } = String.Empty;
 
+		/// <summary>
+		/// With this enabled local machine will start local processes to race remotes when utilization goes under percent number
+		/// So if local machine can run 20 actions and only run 10, that means that it is 50% utilized. If MaxRacingPercent
+		/// is set to 70%, local machine will look at what actions remotes are running and start running them too until it reaches
+		/// 70% utilization
+		/// </summary>
+		[XmlConfigFile(Category = "UnrealBuildAccelerator", Name = "MaxRacingPercent")]
+		[CommandLine("-UBAMaxRacingPercent")]
+		public int MaxRacingPercent { get; set; } = 0;
+
+
 		/// <summary>
 		/// Max number of cache download tasks that can execute in parallel
 		/// </summary>
