UBT: Catch and log a bug when moving files within the same directory that only happens on Mac OS, logging more information in the hopes that we can understand it better. Swallow the exception so we don't continue to affect build health.

#Synced-CL 49438742

diff --git a/Engine/Source/Programs/UnrealBuildTool/Storage/Impl/FileSystemStorageProvider.cs b/Engine/Source/Programs/UnrealBuildTool/Storage/Impl/FileSystemStorageProvider.cs
--- a/Engine/Source/Programs/UnrealBuildTool/Storage/Impl/FileSystemStorageProvider.cs
+++ b/Engine/Source/Programs/UnrealBuildTool/Storage/Impl/FileSystemStorageProvider.cs
@@ -4,6 +4,7 @@ using System;
 using System.IO;
 using System.Threading;
 using EpicGames.Core;
+using Microsoft.Extensions.Logging;
 using UnrealBuildBase;
 
 namespace UnrealBuildTool.Storage.Impl
@@ -34,6 +35,25 @@ namespace UnrealBuildTool.Storage.Impl
 					{
 						FileReference.Move(tempLocation, finalLocation, overwrite: true);
 					}
+					// Try to handle an odd bug where, on Mac OS specifically, the temp file seems to go missing *during* the move call.
+					catch (FileNotFoundException ex) when (OperatingSystem.IsMacOS())
+					{
+						bool tempFileExists = FileReference.Exists(tempLocation);
+						// Check to see if the destination file exists, in case that proves to be a pattern.
+						bool destFileExists = FileReference.Exists(finalLocation);
+						// This log also will give us a timestamp, to see whether there's a pattern as to when in the compile this happens
+						// (right at the start or right at the end would be conspicuous).
+						Log.Logger.LogInformation(
+							ex,
+							"Caught the macOS FileNotFoundException issue! Temp file path: {TempFilePath}, temp file exists: {TempFileExists}, dest file path: {DestFilePath}, dest file exists: {DestFileExists}",
+							tempLocation,
+							tempFileExists,
+							finalLocation,
+							destFileExists);
+						// Deleting a nonexistent file is idempotent.
+						FileReference.Delete(tempLocation);
+						// Don't throw, as it's disrupting build health.
+					}
 					catch
 					{
 						FileReference.Delete(tempLocation);
