UBT: Fix logic bug in Mac build introduced in CL 49280486 by reverting an incorrect change.

#Synced-CL 49438742

diff --git a/Engine/Source/Programs/UnrealBuildTool/Platform/Mac/MacToolChain.cs b/Engine/Source/Programs/UnrealBuildTool/Platform/Mac/MacToolChain.cs
--- a/Engine/Source/Programs/UnrealBuildTool/Platform/Mac/MacToolChain.cs
+++ b/Engine/Source/Programs/UnrealBuildTool/Platform/Mac/MacToolChain.cs
@@ -908,14 +908,17 @@ namespace UnrealBuildTool
 					if (!String.IsNullOrEmpty(DSYMExtension))
 					{
 						string BinaryPath = BuildProductPair.Key.FullName;
-						while (BinaryPath.Contains(".app", StringComparison.Ordinal))
+						if (BinaryPath.Contains(".app", StringComparison.Ordinal))
 						{
-							BinaryPath = Path.GetDirectoryName(BinaryPath)!;
+							while (BinaryPath.Contains(".app", StringComparison.Ordinal))
+							{
+								BinaryPath = Path.GetDirectoryName(BinaryPath)!;
+							}
+							BinaryPath = Path.Combine(BinaryPath, BuildProductPair.Key.GetFileName());
+							BinaryPath = Path.ChangeExtension(BinaryPath, DSYMExtension);
+							FileReference Ref = new FileReference(BinaryPath);
+							BuildProducts[Ref] = BuildProductType.SymbolFile;
 						}
-						BinaryPath = Path.Combine(BinaryPath, BuildProductPair.Key.GetFileName());
-						BinaryPath = Path.ChangeExtension(BinaryPath, DSYMExtension);
-						FileReference Ref = new FileReference(BinaryPath);
-						BuildProducts[Ref] = BuildProductType.SymbolFile;
 					}
 					else if (BuildProductPair.Value == BuildProductType.SymbolFile && BuildProductPair.Key.FullName.Contains(".app", StringComparison.Ordinal))
 					{
