UBT: QueryTargetsMode: Exclude program targets by default in the generated targetinfo.json file.
The targetinfo.json is mostly used for the editor launch on menu does not need the program targets.
Add a new command line option -DontIncludeProgramTargets to QueryTargets and specify it when launching UBT from the editor.
The inclusion of program targets that was picked up from a projects AdditionalPluginDirectories was completely invalidating the caching in FDesktopPlatformBase::GetTargetsForProject (because it knows nothing about AdditionalPluginDirectories).

Result: Saves 5-6 seconds editor boot time for projects with AdditionalPluginDirectories that contains program targets.

#rnx

#Synced-CL 49555800

diff --git a/Engine/Source/Developer/DesktopPlatform/Private/DesktopPlatformBase.cpp b/Engine/Source/Developer/DesktopPlatform/Private/DesktopPlatformBase.cpp
--- a/Engine/Source/Developer/DesktopPlatform/Private/DesktopPlatformBase.cpp
+++ b/Engine/Source/Developer/DesktopPlatform/Private/DesktopPlatformBase.cpp
@@ -1466,7 +1466,7 @@ bool IsTargetInfoValid(const TArray<FTargetInfo>& Targets, TArray<FString>& Dire
 
 const TArray<FTargetInfo>& FDesktopPlatformBase::GetTargetsForProject(const FString& ProjectFile) const
 {
-	TRACE_CPUPROFILER_EVENT_SCOPE_STR("FDesktopPlatfomrBase::GetTargetsForProject");
+	TRACE_CPUPROFILER_EVENT_SCOPE_STR("FDesktopPlatformBase::GetTargetsForProject");
 	// Normalize the project filename
 	FString NormalizedProjectFile = ProjectFile;
 	FPaths::NormalizeFilename(NormalizedProjectFile);
@@ -1532,6 +1532,7 @@ const TArray<FTargetInfo>& FDesktopPlatformBase::GetTargetsForProject(const FStr
 	// we only want this game's targets, not every program under the sun (note if running UnrealEditor without a project, this will
 	// still list all the programs). This matches what GenerateProjectFiles outputs when it performs this step
 	Arguments += TEXT(" -DontIncludeParentAssembly");
+	Arguments += TEXT(" -DontIncludeProgramTargets");
 
 	// Run UBT to update the list of targets. Try to run it without building first.
 	FString Output;
diff --git a/Engine/Source/Programs/UnrealBuildTool/Modes/QueryTargetsMode.cs b/Engine/Source/Programs/UnrealBuildTool/Modes/QueryTargetsMode.cs
--- a/Engine/Source/Programs/UnrealBuildTool/Modes/QueryTargetsMode.cs
+++ b/Engine/Source/Programs/UnrealBuildTool/Modes/QueryTargetsMode.cs
@@ -29,6 +29,12 @@ namespace UnrealBuildTool
 		[CommandLine("-Output=")]
 		FileReference? OutputFile = null;
 
+		/// <summary>
+		/// Write out program targets, has presedence over IncludeAllTargets.
+		/// </summary>
+		[CommandLine("-DontIncludeProgramTargets", Value = "false")]
+		bool bIncludeProgramTargets = true;
+
 		/// <summary>
 		/// Write out all targets, even if a default is specified in the BuildSettings section of the Default*.ini files. 
 		/// </summary>
@@ -76,7 +82,7 @@ namespace UnrealBuildTool
 			}
 
 			// Write information about these targets
-			WriteTargetInfo(ProjectFile, Assembly, OutputFile, Arguments, Logger, bIncludeAllTargets, bIncludeParentAssembly);
+			WriteTargetInfo(ProjectFile, Assembly, OutputFile, Arguments, Logger, bIncludeAllTargets, bIncludeParentAssembly, bIncludeProgramTargets);
 			Logger.LogInformation("Written {OutputFile}", OutputFile);
 			return Task.FromResult(0);
 		}
@@ -108,7 +114,9 @@ namespace UnrealBuildTool
 		/// <param name="Logger">Logger for output</param>
 		/// <param name="bIncludeAllTargets">Include all targets even if a default target is specified for a given target type.</param>
 		/// <param name="bIncludeParentAssembly">Include all targets from parent assemblies (useful for EngineRules which can be chained)</param>
-		public static void WriteTargetInfo(FileReference? ProjectFile, RulesAssembly Assembly, FileReference OutputFile, CommandLineArguments Arguments, ILogger Logger, bool bIncludeAllTargets = true, bool bIncludeParentAssembly = false)
+		/// <param name="bIncludeProgramTargets">Include or exclude all program targets, has presedence over bIncludeAllTargets.</param>
+		public static void WriteTargetInfo(FileReference? ProjectFile, RulesAssembly Assembly, FileReference OutputFile, CommandLineArguments Arguments, ILogger Logger,
+			bool bIncludeAllTargets = true, bool bIncludeParentAssembly = false, bool bIncludeProgramTargets = true)
 		{
 			// Construct all the targets in this assembly
 			List<string> TargetNames = new List<string>();
@@ -143,23 +151,35 @@ namespace UnrealBuildTool
 						continue;
 					}
 
-					// Is this a default target?
 					bool? bIsDefaultTarget = null;
-					if (ProjectFile != null)
-					{
-						string? DefaultTargetName = ProjectFileGenerator.GetProjectDefaultTargetNameForType(ProjectFile.Directory, TargetRules.Type);
 
-						// GetProjectDefaultTargetNameForType returns
-						if (DefaultTargetName != null)
+					// Program targets are controlled by bIncludeProgramTargets which has presedence over bIncludeAllTargets
+					if (TargetRules.Type == TargetType.Program)
+					{
+						if (!bIncludeProgramTargets)
 						{
-							bIsDefaultTarget = DefaultTargetName == TargetName;
+							continue;
 						}
 					}
-
-					// If we don't want all targets, skip over non-defaults.
-					if (!bIncludeAllTargets && bIsDefaultTarget.HasValue && !bIsDefaultTarget.Value)
+					// Check if default target for Client, Editor, Game and Server targets
+					else
 					{
-						continue;
+						if (ProjectFile != null)
+						{
+							string? DefaultTargetName = ProjectFileGenerator.GetProjectDefaultTargetNameForType(ProjectFile.Directory, TargetRules.Type);
+
+							// GetProjectDefaultTargetNameForType returns
+							if (DefaultTargetName != null)
+							{
+								bIsDefaultTarget = DefaultTargetName == TargetName;
+							}
+						}
+
+						// If we don't want all targets, skip over non-defaults.
+						if (!bIncludeAllTargets && bIsDefaultTarget.HasValue && !bIsDefaultTarget.Value)
+						{
+							continue;
+						}
 					}
 
 					// Get the path to the target
diff --git a/Engine/Source/Programs/UnrealBuildTool/ProjectFiles/ProjectFileGenerator.cs b/Engine/Source/Programs/UnrealBuildTool/ProjectFiles/ProjectFileGenerator.cs
--- a/Engine/Source/Programs/UnrealBuildTool/ProjectFiles/ProjectFileGenerator.cs
+++ b/Engine/Source/Programs/UnrealBuildTool/ProjectFiles/ProjectFileGenerator.cs
@@ -1546,7 +1546,8 @@ namespace UnrealBuildTool
 				Parallel2.ForEach(ProjectFileToRules, Item =>
 				{
 					using ITimelineEvent t = Timeline.ScopeEvent(Item.Key?.GetFileName() ?? "Unknown");
-					QueryTargetsMode.WriteTargetInfo(Item.Key, Item.Value, QueryTargetsMode.GetDefaultOutputFile(Item.Key), new CommandLineArguments(GetTargetArguments(Arguments)), Logger);
+					QueryTargetsMode.WriteTargetInfo(Item.Key, Item.Value, QueryTargetsMode.GetDefaultOutputFile(Item.Key), new CommandLineArguments(GetTargetArguments(Arguments)), Logger,
+							/*bIncludeAllTargets*/ true, /*bIncludeParentAssembly*/ false, /*bIncludeProgramTargets*/ false);
 					Progress.Write(Interlocked.Increment(ref TasksFinished), ProjectFileToRules.Count);
 				});
 			}
