Adding a WITH_MULTI_VIEW_SLATE_WINDOW_SUPPORT flag for multi view rendering on macos

#Synced-CL 49563214

diff --git a/Engine/Config/Mac/BaseMacEngine.ini b/Engine/Config/Mac/BaseMacEngine.ini
--- a/Engine/Config/Mac/BaseMacEngine.ini
+++ b/Engine/Config/Mac/BaseMacEngine.ini
@@ -74,3 +74,4 @@ GameHUDBind=(Key=Tilde,bShift=True,bCtrl=False,bAlt=False,bCmd=False)
 
 [/Script/BuildSettings.BuildSettings]
 bCompileCEF3=False
+bMultiViewPerSlateWindowSupport=False 
diff --git a/Engine/Source/Programs/UnrealBuildTool/Configuration/Rules/ReadOnlyTargetRules.cs b/Engine/Source/Programs/UnrealBuildTool/Configuration/Rules/ReadOnlyTargetRules.cs
--- a/Engine/Source/Programs/UnrealBuildTool/Configuration/Rules/ReadOnlyTargetRules.cs
+++ b/Engine/Source/Programs/UnrealBuildTool/Configuration/Rules/ReadOnlyTargetRules.cs
@@ -155,6 +155,8 @@ namespace UnrealBuildTool
 
 		public bool bCompileISPC => Inner.bCompileISPC;
 
+		public bool bMultiViewPerSlateWindowSupport => Inner.bMultiViewPerSlateWindowSupport;
+
 		public bool bUseGameplayDebugger => Inner.bUseGameplayDebugger;
 
 		public bool bUseGameplayDebuggerCore => Inner.bUseGameplayDebuggerCore;
diff --git a/Engine/Source/Programs/UnrealBuildTool/Configuration/Rules/TargetRules.cs b/Engine/Source/Programs/UnrealBuildTool/Configuration/Rules/TargetRules.cs
--- a/Engine/Source/Programs/UnrealBuildTool/Configuration/Rules/TargetRules.cs
+++ b/Engine/Source/Programs/UnrealBuildTool/Configuration/Rules/TargetRules.cs
@@ -1080,6 +1080,13 @@ namespace UnrealBuildTool
 		[RequiresUniqueBuildEnvironment]
 		public bool bCompilePython { get; set; } = true;
 
+		/// <summary>
+		/// Whether to compile with support for multiple platform-native views per slate window (right now only implemented on macOS)
+		/// </summary>
+		[RequiresUniqueBuildEnvironment]
+		[ConfigFile(ConfigHierarchyType.Engine, "/Script/BuildSettings.BuildSettings", "bMultiViewPerSlateWindowSupport")]
+		public bool bMultiViewPerSlateWindowSupport { get; set; } = false;
+
 		/// <summary>
 		/// Whether to compile with WITH_GAMEPLAY_DEBUGGER enabled with all Engine's default gameplay debugger categories.
 		/// </summary>
diff --git a/Engine/Source/Programs/UnrealBuildTool/Configuration/UEBuildTarget.cs b/Engine/Source/Programs/UnrealBuildTool/Configuration/UEBuildTarget.cs
--- a/Engine/Source/Programs/UnrealBuildTool/Configuration/UEBuildTarget.cs
+++ b/Engine/Source/Programs/UnrealBuildTool/Configuration/UEBuildTarget.cs
@@ -6534,6 +6534,15 @@ namespace UnrealBuildTool
 				GlobalCompileEnvironment.Definitions.Add("WITH_CEF3=0");
 			}
 
+			if(Rules.bMultiViewPerSlateWindowSupport && Rules.bCompileAgainstEditor && UEBuildPlatform.IsPlatformInGroup(Platform, UnrealPlatformGroup.Apple))
+			{
+				GlobalCompileEnvironment.Definitions.Add("WITH_MULTI_VIEW_SLATE_WINDOW_SUPPORT=1");
+			}
+			else
+			{
+				GlobalCompileEnvironment.Definitions.Add("WITH_MULTI_VIEW_SLATE_WINDOW_SUPPORT=0");
+			}
+
 			// Set the define for enabling live coding
 			if (Rules.bWithLiveCoding)
 			{
