[UBA]
* Fixed bug where wrong UbaHost.toml was used on non-windows

diff --git a/Engine/Source/Programs/Shared/EpicGames.UBA/IScheduler.cs b/Engine/Source/Programs/Shared/EpicGames.UBA/IScheduler.cs
--- a/Engine/Source/Programs/Shared/EpicGames.UBA/IScheduler.cs
+++ b/Engine/Source/Programs/Shared/EpicGames.UBA/IScheduler.cs
@@ -151,11 +151,12 @@ namespace EpicGames.UBA
 		/// </summary>
 		/// <param name="session">Session</param>
 		/// <param name="cacheClients">List of cache clients</param>
+		/// <param name="config">Config for scheduler</param>
 		/// <param name="maxLocalProcessors">Max number of local processes scheduler can run in parallel</param>
 		/// <param name="forceRemote">Force all processes that can to run remote</param>
-		public static IScheduler CreateScheduler(ISessionServer session, IEnumerable<ICacheClient> cacheClients, int maxLocalProcessors, bool forceRemote)
+		public static IScheduler CreateScheduler(ISessionServer session, IEnumerable<ICacheClient> cacheClients, IConfig config, int maxLocalProcessors, bool forceRemote)
 		{
-			return new SchedulerImpl(session, cacheClients, maxLocalProcessors, forceRemote);
+			return new SchedulerImpl(session, cacheClients, config, maxLocalProcessors, forceRemote);
 		}
 	}
 }
\ No newline at end of file
diff --git a/Engine/Source/Programs/Shared/EpicGames.UBA/Impl/SchedulerImpl.cs b/Engine/Source/Programs/Shared/EpicGames.UBA/Impl/SchedulerImpl.cs
--- a/Engine/Source/Programs/Shared/EpicGames.UBA/Impl/SchedulerImpl.cs
+++ b/Engine/Source/Programs/Shared/EpicGames.UBA/Impl/SchedulerImpl.cs
@@ -134,7 +134,7 @@ namespace EpicGames.UBA.Impl
 		internal Func<IProcessFinishedInfo, ProcessFinishedResponse>? _processFinishedCallback;
 		readonly List<SchedulerProcess> _processes = [];
 
-		internal SchedulerImpl(ISessionServer server, IEnumerable<ICacheClient> cacheClients, int maxLocalProcessors, bool forceRemote)
+		internal SchedulerImpl(ISessionServer server, IEnumerable<ICacheClient> cacheClients, IConfig config, int maxLocalProcessors, bool forceRemote)
 		{
 			nint[] cc = new nint[cacheClients.Count()];
 			uint cacheClientCount = 0;
@@ -143,7 +143,6 @@ namespace EpicGames.UBA.Impl
 				cc[cacheClientCount++] = client.GetHandle();
 			}
 
-			using IConfig config = IConfig.LoadConfig("");
 			config.AddValue("Scheduler", "MaxLocalProcessors", maxLocalProcessors.ToString());
 			if (forceRemote)
 			{
diff --git a/Engine/Source/Programs/UnrealBuildTool/Executors/UnrealBuildAccelerator/UBAExecutor.cs b/Engine/Source/Programs/UnrealBuildTool/Executors/UnrealBuildAccelerator/UBAExecutor.cs
--- a/Engine/Source/Programs/UnrealBuildTool/Executors/UnrealBuildAccelerator/UBAExecutor.cs
+++ b/Engine/Source/Programs/UnrealBuildTool/Executors/UnrealBuildAccelerator/UBAExecutor.cs
@@ -31,6 +31,7 @@ namespace UnrealBuildTool
 		readonly CancellationTokenSource _cancellationSource;
 		readonly bool _bCompactOutput;
 
+		IConfig? _config;
 		IServer? _server;
 		IStorageServer? _storageServer;
 		ISessionServer? _session;
@@ -277,40 +278,40 @@ namespace UnrealBuildTool
 			Log.BackupLogFile(ubaTraceFile);
 
 			FileReference configFile = FileReference.Combine(UbaBinariesDir, "UbaHost.toml");
-			IConfig config = EpicGames.UBA.IConfig.LoadConfig(configFile.FullName);
+			_config = EpicGames.UBA.IConfig.LoadConfig(configFile.FullName);
 			if (!String.IsNullOrEmpty(UBAConfig.CompressionLevel))
 			{
-				config.AddValue("Storage", "CompressionLevel", UBAConfig.CompressionLevel);
+				_config.AddValue("Storage", "CompressionLevel", UBAConfig.CompressionLevel);
 			}
 
 			if (!String.IsNullOrEmpty(UBAConfig.SharedMemoryTempFile))
 			{
-				config.AddValue("Storage", "SharedMemoryTempFile", UBAConfig.SharedMemoryTempFile);
+				_config.AddValue("Storage", "SharedMemoryTempFile", UBAConfig.SharedMemoryTempFile);
 			}
 
 			if (!UBAConfig.bCleanupOnDispose)
 			{
-				config.AddValue("Storage", "SkipCleanupOnShutdown", "true"); // For faster shutdown
+				_config.AddValue("Storage", "SkipCleanupOnShutdown", "true"); // For faster shutdown
 			}
 
 			if (UBAConfig.bWritePlaceholders)
 			{
-				config.AddValue("Session", "WritePlaceholders", "true");
+				_config.AddValue("Session", "WritePlaceholders", "true");
 			}
 
 			if (!String.IsNullOrEmpty(UBAConfig.SharedMemoryTempFile))
 			{
-				config.AddValue("Storage", "SharedMemoryTempFile", UBAConfig.SharedMemoryTempFile);
+				_config.AddValue("Storage", "SharedMemoryTempFile", UBAConfig.SharedMemoryTempFile);
 			}
 
 			if (!UBAConfig.bAllowDetour)
 			{
-				config.AddValue("Session", "AllowMemoryMaps", "false");
+				_config.AddValue("Session", "AllowMemoryMaps", "false");
 			}
 
 			if (UBAConfig.MaxRacingPercent > 0)
 			{
-				config.AddValue("Scheduler", "MaxRacingPercent", UBAConfig.MaxRacingPercent.ToString());
+				_config.AddValue("Scheduler", "MaxRacingPercent", UBAConfig.MaxRacingPercent.ToString());
 			}
 
 			EpicGames.UBA.ILogger ubaLogger = EpicGames.UBA.ILogger.GlobalLogger!;
@@ -787,7 +788,7 @@ namespace UnrealBuildTool
 				CompactOutput = _bCompactOutput,
 			};
 
-			using IScheduler scheduler = IScheduler.CreateScheduler(session, cache, NumParallelProcesses, UBAConfig.bForceBuildAllRemote);
+			using IScheduler scheduler = IScheduler.CreateScheduler(session, cache, _config!, NumParallelProcesses, UBAConfig.bForceBuildAllRemote);
 
 			// This must not outlive the cancellation token source being cancelled in order for it to correctly save results.
 			await using IActionResultStore resultStore = new ContentAddressableActionResultStore(new FileSystemStorageProvider(), logger);
