[UBA]
* Fixed so -ubavisualizer works on macos when running ubt

#Synced-CL 49567671

diff --git a/Engine/Source/Programs/UnrealBuildTool/Executors/UnrealBuildAccelerator/UBAExecutor.cs b/Engine/Source/Programs/UnrealBuildTool/Executors/UnrealBuildAccelerator/UBAExecutor.cs
--- a/Engine/Source/Programs/UnrealBuildTool/Executors/UnrealBuildAccelerator/UBAExecutor.cs
+++ b/Engine/Source/Programs/UnrealBuildTool/Executors/UnrealBuildAccelerator/UBAExecutor.cs
@@ -588,9 +588,16 @@ namespace UnrealBuildTool
 
 			try
 			{
-				if (UBAConfig.bLaunchVisualizer && OperatingSystem.IsWindows())
+				if (UBAConfig.bLaunchVisualizer)
 				{
-					_ = Task.Run(LaunchVisualizer);
+					if (OperatingSystem.IsWindows())
+					{
+						_ = Task.Run(LaunchVisualizerWin);	
+					}
+					else if (OperatingSystem.IsMacOS())
+					{
+						_ = Task.Run(() => LaunchVisualizerMac(logger));
+					}
 				}
 
 				OnSessionCreated?.Invoke(_session!);
@@ -629,7 +636,7 @@ namespace UnrealBuildTool
 		}
 
 		[SupportedOSPlatform("windows")]
-		static void LaunchVisualizer()
+		static void LaunchVisualizerWin()
 		{
 			FileReference visualizerPath = FileReference.Combine(UbaBinariesDir, "UbaVisualizer.exe");
 			FileReference tempPath = FileReference.Combine(new DirectoryReference(System.IO.Path.GetTempPath()), visualizerPath.GetFileName());
@@ -671,6 +678,22 @@ namespace UnrealBuildTool
 			}
 		}
 
+		[SupportedOSPlatform("macos")]
+		static void LaunchVisualizerMac(ILogger logger)
+		{
+			DirectoryReference visualizerPath = DirectoryReference.Combine(UbaBinariesDir, "UbaVisualizer.app");
+			if (!DirectoryReference.Exists(visualizerPath))
+				return;
+			System.Diagnostics.ProcessStartInfo psi = new()
+			{
+				FileName = "/usr/bin/open",
+				UseShellExecute = false,
+			};
+			psi.ArgumentList.Add("-a");
+			psi.ArgumentList.Add(visualizerPath.FullName);
+			System.Diagnostics.Process.Start(psi);
+		}
+
 		public override bool VerifyOutputs => UBAConfig.bWriteToDisk;
 
 		(byte[]?, uint) BuildKnownInputs(LinkedAction action)
