Fix static analysis issues

#rnx
#Synced-CL 49574570

diff --git a/Engine/Source/Programs/UnrealBuildTool/Executors/UnrealBuildAccelerator/UBAExecutor.cs b/Engine/Source/Programs/UnrealBuildTool/Executors/UnrealBuildAccelerator/UBAExecutor.cs
--- a/Engine/Source/Programs/UnrealBuildTool/Executors/UnrealBuildAccelerator/UBAExecutor.cs
+++ b/Engine/Source/Programs/UnrealBuildTool/Executors/UnrealBuildAccelerator/UBAExecutor.cs
@@ -82,6 +82,7 @@ namespace UnrealBuildTool
 					_storageServer?.Dispose();
 				}
 
+				_config?.Dispose();
 				_server?.Dispose();
 
 				_cancellationSource.Dispose();
@@ -590,14 +591,7 @@ namespace UnrealBuildTool
 			{
 				if (UBAConfig.bLaunchVisualizer)
 				{
-					if (OperatingSystem.IsWindows())
-					{
-						_ = Task.Run(LaunchVisualizerWin);	
-					}
-					else if (OperatingSystem.IsMacOS())
-					{
-						_ = Task.Run(() => LaunchVisualizerMac(logger));
-					}
+					_ = Task.Run(LaunchVisualizer);
 				}
 
 				OnSessionCreated?.Invoke(_session!);
@@ -679,11 +673,14 @@ namespace UnrealBuildTool
 		}
 
 		[SupportedOSPlatform("macos")]
-		static void LaunchVisualizerMac(ILogger logger)
+		static void LaunchVisualizerMac()
 		{
 			DirectoryReference visualizerPath = DirectoryReference.Combine(UbaBinariesDir, "UbaVisualizer.app");
 			if (!DirectoryReference.Exists(visualizerPath))
+			{
 				return;
+			}
+
 			System.Diagnostics.ProcessStartInfo psi = new()
 			{
 				FileName = "/usr/bin/open",
@@ -694,6 +691,18 @@ namespace UnrealBuildTool
 			System.Diagnostics.Process.Start(psi);
 		}
 
+		static void LaunchVisualizer()
+		{
+			if (OperatingSystem.IsWindows())
+			{
+				LaunchVisualizerWin();
+			}
+			else if (OperatingSystem.IsMacOS())
+			{
+				LaunchVisualizerMac();
+			}
+		}
+
 		public override bool VerifyOutputs => UBAConfig.bWriteToDisk;
 
 		(byte[]?, uint) BuildKnownInputs(LinkedAction action)
