UnrealBuildTool: Small code simplification to MacToolchain.

#Synced-CL 49587745

diff --git a/Engine/Source/Programs/UnrealBuildTool/Platform/Mac/MacToolChain.cs b/Engine/Source/Programs/UnrealBuildTool/Platform/Mac/MacToolChain.cs
--- a/Engine/Source/Programs/UnrealBuildTool/Platform/Mac/MacToolChain.cs
+++ b/Engine/Source/Programs/UnrealBuildTool/Platform/Mac/MacToolChain.cs
@@ -444,18 +444,10 @@ namespace UnrealBuildTool
 
 				if (bBuildImportLibraryOnly)
 				{
-					return PerArchOutputFiles.ToArray();
+					return [.. PerArchOutputFiles];
 				}
 
-				FileItem OutputFileItem;
-				if (bBuildImportLibraryOnly)
-				{
-					OutputFileItem = MakeStubItem(LinkEnvironment, LinkEnvironment.OutputFilePath.FullName);
-				}
-				else
-				{
-					OutputFileItem = FileItem.GetItemByFileReference(LinkEnvironment.OutputFilePath);
-				}
+				FileItem OutputFileItem = FileItem.GetItemByFileReference(LinkEnvironment.OutputFilePath);
 				OutputFiles.Add(OutputFileItem);
 
 				Action LipoAction = Graph.CreateAction(ActionType.Link);
