Fully parallelize the C# project generation in UBT.

#Synced-CL 49587745

diff --git a/Engine/Source/Programs/UnrealBuildTool/ProjectFiles/ProjectFileGenerator.cs b/Engine/Source/Programs/UnrealBuildTool/ProjectFiles/ProjectFileGenerator.cs
--- a/Engine/Source/Programs/UnrealBuildTool/ProjectFiles/ProjectFileGenerator.cs
+++ b/Engine/Source/Programs/UnrealBuildTool/ProjectFiles/ProjectFileGenerator.cs
@@ -546,7 +546,11 @@ namespace UnrealBuildTool
 				}
 			}
 
-			PrimaryProjectFolder Folder = ProgramsFolder.AddSubFolder(ProgramSubDirectory);
+			PrimaryProjectFolder Folder;
+			lock (ProgramsFolder)
+			{
+				Folder = ProgramsFolder.AddSubFolder(ProgramSubDirectory);
+			}
 			DirectoryReference SamplesDirectory = DirectoryReference.Combine(Unreal.RootDirectory, "Samples");
 
 			// Find all the modules .csproj files to add
@@ -585,13 +589,19 @@ namespace UnrealBuildTool
 
 						if (ProjectFile.IsUnderDirectory(SamplesDirectory))
 						{
-							RootFolder.AddSubFolder("Samples").AddSubFolder(ProjectFile.GetFileNameWithoutAnyExtensions()).ChildProjects.Add(Project);
+							lock (RootFolder)
+							{
+								RootFolder.AddSubFolder("Samples").AddSubFolder(ProjectFile.GetFileNameWithoutAnyExtensions()).ChildProjects.Add(Project);
+							}
 						}
 						else
 						{
 							// @todo: adding to a group that is the name of the game's automation - if the name is not the name of the game,
 							// we'd have to search up looking for the .uproject. not sure if this happens
-							RootFolder.AddSubFolder("Games").AddSubFolder(ProjectFile.GetFileNameWithoutAnyExtensions()).ChildProjects.Add(Project);
+							lock (RootFolder)
+							{
+								RootFolder.AddSubFolder("Games").AddSubFolder(ProjectFile.GetFileNameWithoutAnyExtensions()).ChildProjects.Add(Project);
+							}
 						}
 					}
 					else
@@ -611,8 +621,9 @@ namespace UnrealBuildTool
 		/// <summary>
 		/// Adds all the DotNet folder to the solution.
 		/// </summary>
-		void AddSharedDotNetModules(List<DirectoryReference> AllEngineDirectories, PrimaryProjectFolder ProgramsFolder, ILogger Logger)
+		private Task AddSharedDotNetModulesAsync(List<DirectoryReference> AllEngineDirectories, PrimaryProjectFolder ProgramsFolder, ILogger Logger)
 		{
+			List<Task> ProjectTasks = [];
 			foreach (DirectoryReference EngineDir in AllEngineDirectories)
 			{
 				DirectoryInfo DotNetDir = new DirectoryInfo(DirectoryReference.Combine(EngineDir, "Shared").FullName);
@@ -633,14 +644,23 @@ namespace UnrealBuildTool
 							{
 								continue;
 							}
-							VCSharpProjectFile Project = new VCSharpProjectFile(new FileReference(ProjectFile), Logger);
-							Project.ShouldBuildForAllSolutionTargets = false;
-							AddExistingProjectFile(Project, bForceDevelopmentConfiguration: true);
-							Folder.ChildProjects.Add(Project);
+
+							ProjectTasks.Add(Task.Run(() =>
+							{
+								VCSharpProjectFile Project = new VCSharpProjectFile(new FileReference(ProjectFile), Logger);
+								Project.ShouldBuildForAllSolutionTargets = false;
+								AddExistingProjectFile(Project, bForceDevelopmentConfiguration: true);
+								lock (Folder)
+								{
+									Folder.ChildProjects.Add(Project);
+								}
+							}));
 						}
 					}
 				}
 			}
+
+			return Task.WhenAll(ProjectTasks);
 		}
 
 		/// <summary>
@@ -662,13 +682,17 @@ namespace UnrealBuildTool
 		/// <summary>
 		/// Finds all csproj within Engine/Source/Programs, and add them if their .uprogram file exists.
 		/// </summary>
-		void DiscoverCSharpProgramProjects(List<DirectoryReference> AllEngineDirectories, List<FileReference> AllGameProjects, PrimaryProjectFolder ProgramsFolder, ILogger Logger)
+		private Task DiscoverCSharpProgramProjectsAsync(
+			List<DirectoryReference> AllEngineDirectories,
+			List<FileReference> AllGameProjects,
+			PrimaryProjectFolder ProgramsFolder,
+			ILogger Logger)
 		{
-			List<FileReference> FoundProjects = new List<FileReference>();
-			List<DirectoryReference> ProjectDirs = new List<DirectoryReference>();
-
-			ParallelQueue Queue = new();
+			using ITimelineEvent _ = Timeline.ScopeEvent("DiscoverCSharpPrograms", Logger);
 
+			List<FileReference> FoundProjects = [];
+			List<DirectoryReference> ProjectDirs = [];
+			using ParallelQueue Queue = new();
 			if (bIncludeEnginePrograms)
 			{
 				DirectoryReference EngineExtras = DirectoryReference.Combine(Unreal.EngineDirectory, "Extras");
@@ -694,6 +718,8 @@ namespace UnrealBuildTool
 			}
 			Queue.Drain();
 
+			List<Task> ProjectParseTasks = [];
+
 			string[] UnsupportedPlatformNames = Utils.MakeListOfUnsupportedPlatforms(SupportedPlatforms, bIncludeUnbuildablePlatforms: true, Logger).ToArray();
 			foreach (FileReference FoundProject in FoundProjects)
 			{
@@ -703,28 +729,35 @@ namespace UnrealBuildTool
 					{
 						if (!FoundProject.ContainsAnyNames(UnsupportedPlatformNames, ProjectDir))
 						{
-							VCSharpProjectFile Project = new VCSharpProjectFile(FoundProject, Logger);
+							ProjectParseTasks.Add(Task.Run(() =>
+							{
+								using ITimelineEvent _ = Timeline.ScopeEvent(FoundProject.GetFileName());
+								VCSharpProjectFile Project = new VCSharpProjectFile(FoundProject, Logger);
 
-							Project.ShouldBuildForAllSolutionTargets = true;
-							Project.ShouldBuildByDefaultForSolutionTargets = true;
+								Project.ShouldBuildForAllSolutionTargets = true;
+								Project.ShouldBuildByDefaultForSolutionTargets = true;
 
-							AddExistingProjectFile(Project, bForceDevelopmentConfiguration: false);
+								AddExistingProjectFile(Project, bForceDevelopmentConfiguration: false);
 
-							PrimaryProjectFolder ProjectFolderToAddTo = ProgramsFolder;
-							if (FoundProject.Directory != null && FoundProject.Directory.ParentDirectory != null)
-							{
-								string ProjectRelativePath = FoundProject.Directory.ParentDirectory.MakeRelativeTo(ProjectDir);
-								if (ProjectRelativePath != ".")
+								PrimaryProjectFolder ProjectFolderToAddTo = ProgramsFolder;
+								if (FoundProject.Directory != null && FoundProject.Directory.ParentDirectory != null)
 								{
-									string[] Directories = ProjectRelativePath.Split(Path.DirectorySeparatorChar);
-									foreach (string IntermediateDir in Directories)
+									string ProjectRelativePath = FoundProject.Directory.ParentDirectory.MakeRelativeTo(ProjectDir);
+									if (ProjectRelativePath != ".")
 									{
-										ProjectFolderToAddTo = ProjectFolderToAddTo.AddSubFolder(IntermediateDir);
+										string[] Directories = ProjectRelativePath.Split(Path.DirectorySeparatorChar);
+										foreach (string IntermediateDir in Directories)
+										{
+											ProjectFolderToAddTo = ProjectFolderToAddTo.AddSubFolder(IntermediateDir);
+										}
 									}
 								}
-							}
 
-							ProjectFolderToAddTo.ChildProjects.Add(Project);
+								lock (ProjectFolderToAddTo)
+								{
+									ProjectFolderToAddTo.ChildProjects.Add(Project);
+								}
+							}));
 						}
 						else
 						{
@@ -734,13 +767,17 @@ namespace UnrealBuildTool
 					}
 				}
 			}
+
+			return Task.WhenAll(ProjectParseTasks);
 		}
 
 		private static VCSharpProjectFile CreateRulesAssemblyProject(Dictionary<RulesAssembly, DirectoryReference> RulesAssemblies, RulesAssembly RulesAssembly, DirectoryReference FSPathBase, ILogger Logger)
 		{
+			string ProjectFilename = RulesAssembly.GetSimpleAssemblyName() + ".csproj";
+			using ITimelineEvent _ = Timeline.ScopeEvent(ProjectFilename);
 			DirectoryReference ProjectFilesDirectory = DirectoryReference.Combine(FSPathBase, "Intermediate/Build/BuildRulesProjects", RulesAssembly.GetSimpleAssemblyName()!);
 			DirectoryReference.CreateDirectory(ProjectFilesDirectory);
-			FileReference ModuleFilesProjectLocation = FileReference.Combine(ProjectFilesDirectory, RulesAssembly.GetSimpleAssemblyName() + ".csproj");
+			FileReference ModuleFilesProjectLocation = FileReference.Combine(ProjectFilesDirectory, ProjectFilename);
 
 			SortedSet<FileReference> ReferencedProjects = new() {
 				FileReference.Combine(Unreal.EngineDirectory, "Source", "Programs", "Shared", "EpicGames.Build", "EpicGames.Build.csproj"),
@@ -1370,65 +1407,7 @@ namespace UnrealBuildTool
 			}
 
 			{
-				PrimaryProjectFolder ProgramsFolder = RootFolder.AddSubFolder("Programs");
-
-				if (bIncludeDotNetPrograms)
-				{
-					using ITimelineEvent DotNetTimer = Timeline.ScopeEvent("Adding DotNet Programs");
-					// gather engine directories across extension dirs
-					string[] UnsupportedPlatformNames = Utils.MakeListOfUnsupportedPlatforms(SupportedPlatforms, bIncludeUnbuildablePlatforms: true, Logger).ToArray();
-					List<DirectoryReference> AllEngineDirectories =
-						Unreal.GetExtensionDirs(Unreal.EngineDirectory, "Source/Programs").Where(x =>
-						{
-							if (x.ContainsAnyNames(UnsupportedPlatformNames, Unreal.EngineDirectory))
-							{
-								Logger.LogDebug("Skipping any C# project files in \"{Directory}\" due to unsupported platform", x);
-								return false;
-							}
-							return true;
-						}).ToList();
-
-					// Add UnrealBuildTool to the primary project
-					AddUnrealBuildToolProject(ProgramsFolder, Logger);
-
-					// Add AutomationTool to the primary project
-					VCSharpProjectFile AutomationToolProject = AddSimpleCSharpProject("AutomationTool", Logger, bShouldBuildForAllSolutionTargets: true, bForceDevelopmentConfiguration: true);
-					if (AutomationToolProject != null)
-					{
-						ProgramsFolder.ChildProjects.Add(AutomationToolProject);
-					}
-
-					// Add automation.csproj files to the primary project
-					AddRulesModules(Rules.RulesFileType.AutomationModule, "Automation", AutomationProjectFiles, AllGameProjects, RootFolder, ProgramsFolder, Logger);
-
-					// Add ubtplugin.csproj files to the primary project
-					AddRulesModules(Rules.RulesFileType.UbtPlugin, "UnrealBuildTool.Plugins", UbtPluginProjectFiles, AllGameProjects, RootFolder, ProgramsFolder, Logger);
-
-					// Add shared projects
-					AddSharedDotNetModules(AllEngineDirectories, ProgramsFolder, Logger);
-
-					if (bIncludeEnginePrograms)
-					{
-						// Discover C# programs which should additionally be included in the solution.
-						using (Timeline.ScopeEvent("DiscoverCsharpPrograms", Logger))
-						{
-							DiscoverCSharpProgramProjects(AllEngineDirectories, AllGameProjects, ProgramsFolder, Logger);
-						}
-					}
-
-					Parallel2.ForEach(RulesAssemblies, (RulesAssemblyEntry) =>
-					{
-						if (RulesAssemblyEntry.Key.GetSimpleAssemblyName() != null)
-						{
-							VCSharpProjectFile RulesProject = CreateRulesAssemblyProject(RulesAssemblies, RulesAssemblyEntry.Key, RulesAssemblyEntry.Value, Logger);
-							AddExistingProjectFile(RulesProject, bForceDevelopmentConfiguration: false);
-							lock (RootFolder)
-							{
-								RootFolder.AddSubFolder("Rules").ChildProjects.Add(RulesProject);
-							}
-						}
-					});
-				}
+				AddDotnetProjects(AllGameProjects, RulesAssemblies, Logger);
 
 				// Eliminate all redundant project folders.  E.g., folders which contain only one project and that project
 				// has the same name as the folder itself.  To the user, projects "feel like" folders already in the IDE, so we
@@ -1509,6 +1488,72 @@ namespace UnrealBuildTool
 			return bSuccess;
 		}
 
+		private void AddDotnetProjects(List<FileReference> AllGameProjects, Dictionary<RulesAssembly, DirectoryReference> RulesAssemblies, ILogger Logger)
+		{
+			PrimaryProjectFolder ProgramsFolder = RootFolder.AddSubFolder("Programs");
+			ConcurrentBag<Task> DotnetProjectTasks = [];
+
+			if (bIncludeDotNetPrograms)
+			{
+				using ITimelineEvent DotNetTimer = Timeline.ScopeEvent("Adding DotNet Programs");
+
+				// gather engine directories across extension dirs
+				string[] UnsupportedPlatformNames = Utils.MakeListOfUnsupportedPlatforms(SupportedPlatforms, bIncludeUnbuildablePlatforms: true, Logger).ToArray();
+				List<DirectoryReference> AllEngineDirectories =
+					Unreal.GetExtensionDirs(Unreal.EngineDirectory, "Source/Programs").Where(x =>
+					{
+						if (x.ContainsAnyNames(UnsupportedPlatformNames, Unreal.EngineDirectory))
+						{
+							Logger.LogDebug("Skipping any C# project files in \"{Directory}\" due to unsupported platform", x);
+							return false;
+						}
+						return true;
+					}).ToList();
+
+				// Add UnrealBuildTool to the primary project
+				DotnetProjectTasks.Add(Task.Run(() => AddUnrealBuildToolProject(ProgramsFolder, Logger)));
+
+				// Add AutomationTool to the primary project
+				DotnetProjectTasks.Add(Task.Run(() => AddAutomationToolProject(ProgramsFolder, Logger)));
+
+				// Add automation.csproj files to the primary project
+				DotnetProjectTasks.Add(Task.Run(() =>
+					AddRulesModules(Rules.RulesFileType.AutomationModule, "Automation", AutomationProjectFiles, AllGameProjects, RootFolder, ProgramsFolder, Logger)));
+
+				// Add ubtplugin.csproj files to the primary project
+				DotnetProjectTasks.Add(Task.Run(() =>
+					AddRulesModules(Rules.RulesFileType.UbtPlugin, "UnrealBuildTool.Plugins", UbtPluginProjectFiles, AllGameProjects, RootFolder, ProgramsFolder, Logger)));
+
+				// Add shared projects
+				DotnetProjectTasks.Add(Task.Run(() => AddSharedDotNetModulesAsync(AllEngineDirectories, ProgramsFolder, Logger)));
+
+				if (bIncludeEnginePrograms)
+				{
+					// Discover C# programs which should additionally be included in the solution.
+					DotnetProjectTasks.Add(Task.Run(() => DiscoverCSharpProgramProjectsAsync(AllEngineDirectories, AllGameProjects, ProgramsFolder, Logger)));
+				}
+
+				PrimaryProjectFolder RulesFolder = RootFolder.AddSubFolder("Rules");
+				foreach (KeyValuePair<RulesAssembly, DirectoryReference> RulesAssemblyEntry in RulesAssemblies)
+				{
+					if (RulesAssemblyEntry.Key.GetSimpleAssemblyName() != null)
+					{
+						DotnetProjectTasks.Add(Task.Run(() =>
+						{
+							VCSharpProjectFile RulesProject = CreateRulesAssemblyProject(RulesAssemblies, RulesAssemblyEntry.Key, RulesAssemblyEntry.Value, Logger);
+							AddExistingProjectFile(RulesProject, bForceDevelopmentConfiguration: false);
+							lock (RulesFolder)
+							{
+								RulesFolder.ChildProjects.Add(RulesProject);
+							}
+						}));
+					}
+				}
+
+				Task.WaitAll([.. DotnetProjectTasks]);
+			}
+		}
+
 		public void GenerateQueryTargetsDataForEditor(string[] Arguments, ILogger Logger)
 		{
 			using ITimelineEvent GenerateTimer = Timeline.ScopeEvent("Generating QueryTargets data for editor", Logger);
@@ -2142,7 +2187,7 @@ namespace UnrealBuildTool
 			Folder.SubFolders.AddRange(SubFoldersToAdd);
 
 			// Sort subfolders to be deterministic
-			Folder.SubFolders.Sort(static (a, b) => string.Compare(a.FolderName, b.FolderName, StringComparison.Ordinal));
+			Folder.SubFolders.Sort(static (a, b) => String.Compare(a.FolderName, b.FolderName, StringComparison.Ordinal));
 
 			// Sort child projects as well while we're at it... to make .sln deterministic
 			Folder.ChildProjects.Sort(static (a, b) => a.ProjectFilePath.CompareTo(b.ProjectFilePath));
@@ -2210,9 +2255,11 @@ namespace UnrealBuildTool
 		/// </summary>
 		private void AddUnrealBuildToolProject(PrimaryProjectFolder ProgramsFolder, ILogger Logger)
 		{
+			const string UnrealBuildToolFilename = "UnrealBuildTool.csproj";
+			using ITimelineEvent _ = Timeline.ScopeEvent(UnrealBuildToolFilename);
 			DirectoryReference ProgramsDirectory = DirectoryReference.Combine(Unreal.EngineSourceDirectory, "Programs");
 
-			FileReference ProjectFileName = FileReference.Combine(ProgramsDirectory, "UnrealBuildTool", "UnrealBuildTool.csproj");
+			FileReference ProjectFileName = FileReference.Combine(ProgramsDirectory, "UnrealBuildTool", UnrealBuildToolFilename);
 			if (FileReference.Exists(ProjectFileName))
 			{
 				VCSharpProjectFile UnrealBuildToolProject = new VCSharpProjectFile(ProjectFileName, Logger);
@@ -2225,7 +2272,10 @@ namespace UnrealBuildTool
 				AddExistingProjectFile(UnrealBuildToolProject, bNeedsAllPlatformAndConfigurations: true, bForceDevelopmentConfiguration: true);
 
 				// Put this in a solution folder
-				ProgramsFolder.ChildProjects.Add(UnrealBuildToolProject);
+				lock (ProgramsFolder)
+				{
+					ProgramsFolder.ChildProjects.Add(UnrealBuildToolProject);
+				}
 			}
 
 			if (bIncludeEnginePrograms)
@@ -2235,7 +2285,10 @@ namespace UnrealBuildTool
 				{
 					VCSharpProjectFile UatTestsProject = new VCSharpProjectFile(UatTestsProjectFileName, Logger);
 					AddExistingProjectFile(UatTestsProject, bNeedsAllPlatformAndConfigurations: true, bForceDevelopmentConfiguration: true);
-					ProgramsFolder.ChildProjects.Add(UatTestsProject);
+					lock (ProgramsFolder)
+					{
+						ProgramsFolder.ChildProjects.Add(UatTestsProject);
+					}
 				}
 
 				FileReference UbtTestsProjectFileName = FileReference.Combine(ProgramsDirectory, "UnrealBuildTool.Tests", "UnrealBuildTool.Tests.csproj");
@@ -2243,7 +2296,22 @@ namespace UnrealBuildTool
 				{
 					VCSharpProjectFile UbtTestsProject = new VCSharpProjectFile(UbtTestsProjectFileName, Logger);
 					AddExistingProjectFile(UbtTestsProject, bNeedsAllPlatformAndConfigurations: true, bForceDevelopmentConfiguration: true);
-					ProgramsFolder.ChildProjects.Add(UbtTestsProject);
+					lock (ProgramsFolder)
+					{
+						ProgramsFolder.ChildProjects.Add(UbtTestsProject);
+					}
+				}
+			}
+		}
+
+		private void AddAutomationToolProject(PrimaryProjectFolder ProgramsFolder, ILogger Logger)
+		{
+			VCSharpProjectFile AutomationToolProject = AddSimpleCSharpProject("AutomationTool", Logger, bShouldBuildForAllSolutionTargets: true, bForceDevelopmentConfiguration: true);
+			if (AutomationToolProject != null)
+			{
+				lock (ProgramsFolder)
+				{
+					ProgramsFolder.ChildProjects.Add(AutomationToolProject);
 				}
 			}
 		}
