UnrealBuildTool: Add a warning when a module that shouldn't be compiled for the current configuration is referenced by another module.

#Synced-CL 49587745

diff --git a/Engine/Source/Programs/UnrealBuildTool/Configuration/UEBuildTarget.cs b/Engine/Source/Programs/UnrealBuildTool/Configuration/UEBuildTarget.cs
--- a/Engine/Source/Programs/UnrealBuildTool/Configuration/UEBuildTarget.cs
+++ b/Engine/Source/Programs/UnrealBuildTool/Configuration/UEBuildTarget.cs
@@ -1646,10 +1646,35 @@ namespace UnrealBuildTool
 				foreach (UEBuildModuleCPP dependencyModule in referencedModules.Value.OfType<UEBuildModuleCPP>())
 				{
 					PluginInfo? dependencyPluginInfo = dependencyModule.Rules.Plugin;
+
+					ModuleDescriptor? dependencyModuleInfo = dependencyPluginInfo?.Descriptor.Modules?.FirstOrDefault(x => x.Name == dependencyModule.Name);
+					if (dependencyModuleInfo is not null)
+					{
+						bool isCompiledInConfiguration = dependencyModuleInfo.IsCompiledInConfiguration(
+							Platform,
+							Configuration,
+							TargetName,
+							TargetType,
+							Rules.bBuildDeveloperTools,
+							Rules.bBuildRequiresCookedData,
+							Architectures,
+							out string? invalidReason);
+						if (!isCompiledInConfiguration)
+						{
+							logger.LogWarning(
+								KnownLogEvents.Validation,
+								"{ProjectName} plugin dependency '{DependencyPluginName}' module '{DependencyModuleName}' is not enabled because {Reason}, but module '{ModuleName}' depends on it.",
+								projectName,
+								dependencyPluginInfo!.Name,
+								dependencyModule.Name,
+								invalidReason,
+								module.Name);
+						}
+					}
+
 					// Is the modules plugin enabled by default?
-					if (dependencyPluginInfo != null && dependencyPluginInfo.IsEnabledByDefault(bAllowEnginePluginsEnabledByDefault) == false)
+					if (dependencyPluginInfo != null && !dependencyPluginInfo.IsEnabledByDefault(bAllowEnginePluginsEnabledByDefault))
 					{
-						ModuleDescriptor? dependencyModuleInfo = dependencyPluginInfo.Descriptor.Modules?.FirstOrDefault(x => x.Name == dependencyModule.Name);
 						// Try and find the project plugin reference in the plugins list.
 						PluginReferenceDescriptor? projectPluginReference = ProjectDescriptor?.Plugins?.FirstOrDefault(x => x.Name == dependencyPluginInfo.Name);
 						if (projectPluginReference != null)
