Cache C# project parsing in UBT project generation. This significantly speeds up project generation for multiple project file formats in a single run.

#Synced-CL 49587745

diff --git a/Engine/Source/Programs/UnrealBuildTool/Modes/GenerateProjectFilesMode.cs b/Engine/Source/Programs/UnrealBuildTool/Modes/GenerateProjectFilesMode.cs
--- a/Engine/Source/Programs/UnrealBuildTool/Modes/GenerateProjectFilesMode.cs
+++ b/Engine/Source/Programs/UnrealBuildTool/Modes/GenerateProjectFilesMode.cs
@@ -189,6 +189,8 @@ namespace UnrealBuildTool
 				SingleTargetName = Arguments.GetString("-SingleTarget=");
 			}
 
+			CSharpProjectCache CSharpProjCache = new(Logger);
+
 			// Create each project generator and run it
 			Dictionary<ProjectFileFormat, ProjectFileGenerator> Generators = new();
 			foreach (ProjectFileFormat ProjectFileFormat in ProjectFileFormats.Distinct())
@@ -197,53 +199,53 @@ namespace UnrealBuildTool
 				switch (ProjectFileFormat)
 				{
 					case ProjectFileFormat.Make:
-						Generator = new MakefileGenerator(ProjectFile);
+						Generator = new MakefileGenerator(ProjectFile, CSharpProjCache);
 						break;
 					case ProjectFileFormat.CMake:
-						Generator = new CMakefileGenerator(ProjectFile);
+						Generator = new CMakefileGenerator(ProjectFile, CSharpProjCache);
 						break;
 					case ProjectFileFormat.QMake:
-						Generator = new QMakefileGenerator(ProjectFile);
+						Generator = new QMakefileGenerator(ProjectFile, CSharpProjCache);
 						break;
 					case ProjectFileFormat.KDevelop:
-						Generator = new KDevelopGenerator(ProjectFile);
+						Generator = new KDevelopGenerator(ProjectFile, CSharpProjCache);
 						break;
 					case ProjectFileFormat.CodeLite:
-						Generator = new CodeLiteGenerator(ProjectFile, Arguments);
+						Generator = new CodeLiteGenerator(ProjectFile, CSharpProjCache, Arguments);
 						break;
 					case ProjectFileFormat.VisualStudio:
-						Generator = new VCProjectFileGenerator(ProjectFile, VCProjectFileFormat.Default, Arguments);
+						Generator = new VCProjectFileGenerator(ProjectFile, CSharpProjCache, VCProjectFileFormat.Default, Arguments);
 						break;
 					case ProjectFileFormat.VisualStudio2022:
-						Generator = new VCProjectFileGenerator(ProjectFile, VCProjectFileFormat.VisualStudio2022, Arguments);
+						Generator = new VCProjectFileGenerator(ProjectFile, CSharpProjCache, VCProjectFileFormat.VisualStudio2022, Arguments);
 						break;
 					case ProjectFileFormat.VisualStudio2026:
-						Generator = new VCProjectFileGenerator(ProjectFile, VCProjectFileFormat.VisualStudio2026, Arguments);
+						Generator = new VCProjectFileGenerator(ProjectFile, CSharpProjCache, VCProjectFileFormat.VisualStudio2026, Arguments);
 						break;
 					case ProjectFileFormat.XCode:
-						Generator = new XcodeProjectFileGenerator(ProjectFile, Arguments);
+						Generator = new XcodeProjectFileGenerator(ProjectFile, CSharpProjCache, Arguments);
 						break;
 					case ProjectFileFormat.Eddie:
-						Generator = new EddieProjectFileGenerator(ProjectFile);
+						Generator = new EddieProjectFileGenerator(ProjectFile, CSharpProjCache);
 						break;
 					case ProjectFileFormat.VisualStudioCode:
-						Generator = new VSCodeProjectFileGenerator(ProjectFile);
+						Generator = new VSCodeProjectFileGenerator(ProjectFile, CSharpProjCache);
 						break;
 					case ProjectFileFormat.VisualStudioWorkspace:
-						Generator = new VSWorkspaceProjectFileGenerator(ProjectFile, Arguments);
+						Generator = new VSWorkspaceProjectFileGenerator(ProjectFile, CSharpProjCache, Arguments);
 						break;
 					case ProjectFileFormat.CLion:
-						Generator = new CLionGenerator(ProjectFile);
+						Generator = new CLionGenerator(ProjectFile, CSharpProjCache);
 						break;
 					case ProjectFileFormat.Rider:
-						Generator = new RiderProjectFileGenerator(ProjectFile, Arguments);
+						Generator = new RiderProjectFileGenerator(ProjectFile, CSharpProjCache, Arguments);
 						break;
 					case ProjectFileFormat.AndroidStudio:
-						Generator = new AndroidStudioFileGenerator(ProjectFile);
+						Generator = new AndroidStudioFileGenerator(ProjectFile, CSharpProjCache);
 						break;
 #if __VPROJECT_AVAILABLE__
 					case ProjectFileFormat.VProject:
-						Generator = new VProjectFileGenerator(ProjectFile);
+						Generator = new VProjectFileGenerator(ProjectFile, CSharpProjCache);
 						break;
 #endif
 					default:
diff --git a/Engine/Source/Programs/UnrealBuildTool/Platform/Apple/AppleToolChain.cs b/Engine/Source/Programs/UnrealBuildTool/Platform/Apple/AppleToolChain.cs
--- a/Engine/Source/Programs/UnrealBuildTool/Platform/Apple/AppleToolChain.cs
+++ b/Engine/Source/Programs/UnrealBuildTool/Platform/Apple/AppleToolChain.cs
@@ -767,7 +767,7 @@ namespace UnrealBuildTool
 				PlatformProjectGenerators.RegisterPlatformProjectGenerator(UnrealTargetPlatform.IOS, new IOSProjectGenerator(CmdLine, Logger), Logger);
 				PlatformProjectGenerators.RegisterPlatformProjectGenerator(UnrealTargetPlatform.TVOS, new TVOSProjectGenerator(CmdLine, Logger), Logger);
 
-				XcodeProjectFileGenerator Generator = new XcodeProjectFileGenerator(ProjectFile, CmdLine);
+				XcodeProjectFileGenerator Generator = new XcodeProjectFileGenerator(ProjectFile, new(Logger), CmdLine);
 				// this could be improved if ProjectFileGenerator constructor took a Arguments param, and it could parse it there instead of the GenerateProjectFilesMode
 				Generator.SingleTargetName = SingleTargetName;
 				// don't need the editor data since these are stub projects
diff --git a/Engine/Source/Programs/UnrealBuildTool/ProjectFiles/AndroidStudio/AndroidStudioProjectFileGenerator.cs b/Engine/Source/Programs/UnrealBuildTool/ProjectFiles/AndroidStudio/AndroidStudioProjectFileGenerator.cs
--- a/Engine/Source/Programs/UnrealBuildTool/ProjectFiles/AndroidStudio/AndroidStudioProjectFileGenerator.cs
+++ b/Engine/Source/Programs/UnrealBuildTool/ProjectFiles/AndroidStudio/AndroidStudioProjectFileGenerator.cs
@@ -21,8 +21,8 @@ namespace UnrealBuildTool
 	{
 		public override string ProjectFileExtension => ".gradle";
 
-		public AndroidStudioFileGenerator(FileReference? InOnlyGameProject)
-			: base(InOnlyGameProject)
+		public AndroidStudioFileGenerator(FileReference? InOnlyGameProject, CSharpProjectCache CSharpProjCache)
+			: base(InOnlyGameProject, CSharpProjCache)
 		{
 		}
 
diff --git a/Engine/Source/Programs/UnrealBuildTool/ProjectFiles/CLion/CLionGenerator.cs b/Engine/Source/Programs/UnrealBuildTool/ProjectFiles/CLion/CLionGenerator.cs
--- a/Engine/Source/Programs/UnrealBuildTool/ProjectFiles/CLion/CLionGenerator.cs
+++ b/Engine/Source/Programs/UnrealBuildTool/ProjectFiles/CLion/CLionGenerator.cs
@@ -12,8 +12,8 @@ namespace UnrealBuildTool
 		/// <summary>
 		/// Creates a new instance of the <see cref="CMakefileGenerator"/> class.
 		/// </summary>
-		public CLionGenerator(FileReference? InOnlyGameProject)
-			: base(InOnlyGameProject)
+		public CLionGenerator(FileReference? InOnlyGameProject, CSharpProjectCache CSharpProjCache)
+			: base(InOnlyGameProject, CSharpProjCache)
 		{
 		}
 	}
diff --git a/Engine/Source/Programs/UnrealBuildTool/ProjectFiles/CMake/CMakefileGenerator.cs b/Engine/Source/Programs/UnrealBuildTool/ProjectFiles/CMake/CMakefileGenerator.cs
--- a/Engine/Source/Programs/UnrealBuildTool/ProjectFiles/CMake/CMakefileGenerator.cs
+++ b/Engine/Source/Programs/UnrealBuildTool/ProjectFiles/CMake/CMakefileGenerator.cs
@@ -26,8 +26,8 @@ namespace UnrealBuildTool
 		/// <summary>
 		/// Creates a new instance of the <see cref="CMakefileGenerator"/> class.
 		/// </summary>
-		public CMakefileGenerator(FileReference? InOnlyGameProject)
-			: base(InOnlyGameProject)
+		public CMakefileGenerator(FileReference? InOnlyGameProject, CSharpProjectCache CSharpProjCache)
+			: base(InOnlyGameProject, CSharpProjCache)
 		{
 		}
 
diff --git a/Engine/Source/Programs/UnrealBuildTool/ProjectFiles/CSharpProjectCache.cs b/Engine/Source/Programs/UnrealBuildTool/ProjectFiles/CSharpProjectCache.cs
new file mode 100644
--- /dev/null
+++ b/Engine/Source/Programs/UnrealBuildTool/ProjectFiles/CSharpProjectCache.cs
@@ -0,0 +1,20 @@
+// Copyright Epic Games, Inc. All Rights Reserved.
+
+using System;
+using System.Collections.Concurrent;
+using EpicGames.Core;
+using Microsoft.Extensions.Logging;
+
+namespace UnrealBuildTool;
+
+internal class CSharpProjectCache(ILogger logger)
+{
+	/// <summary>
+	/// Returns a partially cloned C# project so that properties like <see cref="ProjectFile.ShouldBuildForAllSolutionTargets"/> can be set.
+	/// </summary>
+	public VCSharpProjectFile this[FileReference projectPath]
+		=> new(_cache.GetOrAdd(projectPath, pp => new Lazy<VCSharpProjectFile>(() => new VCSharpProjectFile(pp, _logger))).Value);
+
+	private readonly ConcurrentDictionary<FileReference, Lazy<VCSharpProjectFile>> _cache = [];
+	private readonly ILogger _logger = logger;
+}
diff --git a/Engine/Source/Programs/UnrealBuildTool/ProjectFiles/CodeLite/CodeLiteGenerator.cs b/Engine/Source/Programs/UnrealBuildTool/ProjectFiles/CodeLite/CodeLiteGenerator.cs
--- a/Engine/Source/Programs/UnrealBuildTool/ProjectFiles/CodeLite/CodeLiteGenerator.cs
+++ b/Engine/Source/Programs/UnrealBuildTool/ProjectFiles/CodeLite/CodeLiteGenerator.cs
@@ -24,8 +24,8 @@ namespace UnrealBuildTool
 		public string CodeCompletionPreProcessorFileName = "CodeLitePreProcessor.txt";
 		readonly CodeliteProjectFileFormat ProjectFileFormat = CodeliteProjectFileFormat.CodeLite10;
 
-		public CodeLiteGenerator(FileReference? InOnlyGameProject, CommandLineArguments CommandLine)
-			: base(InOnlyGameProject)
+		public CodeLiteGenerator(FileReference? InOnlyGameProject, CSharpProjectCache CSharpProjCache, CommandLineArguments CommandLine)
+			: base(InOnlyGameProject, CSharpProjCache)
 		{
 			if (CommandLine.HasOption("-cl10"))
 			{
diff --git a/Engine/Source/Programs/UnrealBuildTool/ProjectFiles/Eddie/EddieGenerator.cs b/Engine/Source/Programs/UnrealBuildTool/ProjectFiles/Eddie/EddieGenerator.cs
--- a/Engine/Source/Programs/UnrealBuildTool/ProjectFiles/Eddie/EddieGenerator.cs
+++ b/Engine/Source/Programs/UnrealBuildTool/ProjectFiles/Eddie/EddieGenerator.cs
@@ -11,8 +11,8 @@ namespace UnrealBuildTool
 {
 	class EddieProjectFileGenerator : ProjectFileGenerator
 	{
-		public EddieProjectFileGenerator(FileReference? InOnlyGameProject)
-			: base(InOnlyGameProject)
+		public EddieProjectFileGenerator(FileReference? InOnlyGameProject, CSharpProjectCache CSharpProjCache)
+			: base(InOnlyGameProject, CSharpProjCache)
 		{
 		}
 
diff --git a/Engine/Source/Programs/UnrealBuildTool/ProjectFiles/KDevelop/KDevelopGenerator.cs b/Engine/Source/Programs/UnrealBuildTool/ProjectFiles/KDevelop/KDevelopGenerator.cs
--- a/Engine/Source/Programs/UnrealBuildTool/ProjectFiles/KDevelop/KDevelopGenerator.cs
+++ b/Engine/Source/Programs/UnrealBuildTool/ProjectFiles/KDevelop/KDevelopGenerator.cs
@@ -24,8 +24,8 @@ namespace UnrealBuildTool
 	/// </summary>
 	class KDevelopGenerator : ProjectFileGenerator
 	{
-		public KDevelopGenerator(FileReference? InOnlyGameProject)
-			: base(InOnlyGameProject)
+		public KDevelopGenerator(FileReference? InOnlyGameProject, CSharpProjectCache CSharpProjCache)
+			: base(InOnlyGameProject, CSharpProjCache)
 		{
 		}
 
diff --git a/Engine/Source/Programs/UnrealBuildTool/ProjectFiles/Make/MakefileGenerator.cs b/Engine/Source/Programs/UnrealBuildTool/ProjectFiles/Make/MakefileGenerator.cs
--- a/Engine/Source/Programs/UnrealBuildTool/ProjectFiles/Make/MakefileGenerator.cs
+++ b/Engine/Source/Programs/UnrealBuildTool/ProjectFiles/Make/MakefileGenerator.cs
@@ -28,8 +28,8 @@ namespace UnrealBuildTool
 		const bool bGenerateIntelliSenseData = true;
 
 		/// Default constructor
-		public MakefileGenerator(FileReference? InOnlyGameProject)
-			: base(InOnlyGameProject)
+		public MakefileGenerator(FileReference? InOnlyGameProject, CSharpProjectCache CSharpProjCache)
+			: base(InOnlyGameProject, CSharpProjCache)
 		{
 		}
 
diff --git a/Engine/Source/Programs/UnrealBuildTool/ProjectFiles/ProjectFileGenerator.cs b/Engine/Source/Programs/UnrealBuildTool/ProjectFiles/ProjectFileGenerator.cs
--- a/Engine/Source/Programs/UnrealBuildTool/ProjectFiles/ProjectFileGenerator.cs
+++ b/Engine/Source/Programs/UnrealBuildTool/ProjectFiles/ProjectFileGenerator.cs
@@ -454,6 +454,8 @@ namespace UnrealBuildTool
 		/// </summary>
 		protected Dictionary<FileReference, ProjectFile> ModuleToNonEditorProjectFileMap = new Dictionary<FileReference, ProjectFile>();
 
+		protected CSharpProjectCache CSharpProjCache;
+
 		/// <summary>
 		/// If generating project files for a single project, the path to its .uproject file.
 		/// </summary>
@@ -519,20 +521,22 @@ namespace UnrealBuildTool
 		/// Default constructor.
 		/// </summary>
 		/// <param name="InOnlyGameProject">The project file passed in on the command line</param>
-		public ProjectFileGenerator(FileReference? InOnlyGameProject)
+		/// <param name="InCSharpProjCache">A cache for C# projects to prevent needing to parse projects multiple times</param>
+		public ProjectFileGenerator(FileReference? InOnlyGameProject, CSharpProjectCache InCSharpProjCache)
 		{
 			OnlyGameProject = InOnlyGameProject;
 			XmlConfig.ApplyTo(this);
 
 			RootFolder = new PrimaryProjectFolder(this, "<Root>");
 			OnlyGameProjectDescriptor = OnlyGameProject != null && FileReference.Exists(OnlyGameProject) ? ProjectDescriptor.FromFile(OnlyGameProject) : null;
+			CSharpProjCache = InCSharpProjCache;
 		}
 
 		/// <summary>
 		/// Adds all rules project files to the solution.
 		/// </summary>
 		protected void AddRulesModules(Rules.RulesFileType RulesFileType, string ProgramSubDirectory, List<ProjectFile> AddedProjectFiles,
-			List<FileReference> UnrealProjectFiles, PrimaryProjectFolder RootFolder, PrimaryProjectFolder ProgramsFolder, ILogger Logger)
+			List<FileReference> UnrealProjectFiles, PrimaryProjectFolder RootFolder, PrimaryProjectFolder ProgramsFolder)
 		{
 			List<DirectoryReference> GameFolders = new List<DirectoryReference>();
 			List<DirectoryReference> BuildFolders = new List<DirectoryReference>();
@@ -562,7 +566,7 @@ namespace UnrealBuildTool
 				if (FileReference.Exists(ProjectFile))
 				{
 					using ITimelineEvent _ = Timeline.ScopeEvent(ProjectFile.GetFileName());
-					VCSharpProjectFile Project = new VCSharpProjectFile(ProjectFile, Logger);
+					VCSharpProjectFile Project = CSharpProjCache[ProjectFile];
 					lock (projectFiles)
 					{
 						projectFiles.Add(ProjectFile, Project);
@@ -621,7 +625,7 @@ namespace UnrealBuildTool
 		/// <summary>
 		/// Adds all the DotNet folder to the solution.
 		/// </summary>
-		private Task AddSharedDotNetModulesAsync(List<DirectoryReference> AllEngineDirectories, PrimaryProjectFolder ProgramsFolder, ILogger Logger)
+		private Task AddSharedDotNetModulesAsync(List<DirectoryReference> AllEngineDirectories, PrimaryProjectFolder ProgramsFolder)
 		{
 			List<Task> ProjectTasks = [];
 			foreach (DirectoryReference EngineDir in AllEngineDirectories)
@@ -647,7 +651,7 @@ namespace UnrealBuildTool
 
 							ProjectTasks.Add(Task.Run(() =>
 							{
-								VCSharpProjectFile Project = new VCSharpProjectFile(new FileReference(ProjectFile), Logger);
+								VCSharpProjectFile Project = CSharpProjCache[new FileReference(ProjectFile)];
 								Project.ShouldBuildForAllSolutionTargets = false;
 								AddExistingProjectFile(Project, bForceDevelopmentConfiguration: true);
 								lock (Folder)
@@ -732,7 +736,7 @@ namespace UnrealBuildTool
 							ProjectParseTasks.Add(Task.Run(() =>
 							{
 								using ITimelineEvent _ = Timeline.ScopeEvent(FoundProject.GetFileName());
-								VCSharpProjectFile Project = new VCSharpProjectFile(FoundProject, Logger);
+								VCSharpProjectFile Project = CSharpProjCache[FoundProject];
 
 								Project.ShouldBuildForAllSolutionTargets = true;
 								Project.ShouldBuildByDefaultForSolutionTargets = true;
@@ -771,7 +775,7 @@ namespace UnrealBuildTool
 			return Task.WhenAll(ProjectParseTasks);
 		}
 
-		private static VCSharpProjectFile CreateRulesAssemblyProject(Dictionary<RulesAssembly, DirectoryReference> RulesAssemblies, RulesAssembly RulesAssembly, DirectoryReference FSPathBase, ILogger Logger)
+		private VCSharpProjectFile CreateRulesAssemblyProject(Dictionary<RulesAssembly, DirectoryReference> RulesAssemblies, RulesAssembly RulesAssembly, DirectoryReference FSPathBase, ILogger Logger)
 		{
 			string ProjectFilename = RulesAssembly.GetSimpleAssemblyName() + ".csproj";
 			using ITimelineEvent _ = Timeline.ScopeEvent(ProjectFilename);
@@ -843,7 +847,7 @@ namespace UnrealBuildTool
 				Writer.WriteLine("</Project>");
 			}
 
-			VCSharpProjectFile ModuleFilesProject = new VCSharpProjectFile(ModuleFilesProjectLocation, Logger);
+			VCSharpProjectFile ModuleFilesProject = CSharpProjCache[ModuleFilesProjectLocation];
 			return ModuleFilesProject;
 		}
 
@@ -1511,21 +1515,21 @@ namespace UnrealBuildTool
 					}).ToList();
 
 				// Add UnrealBuildTool to the primary project
-				DotnetProjectTasks.Add(Task.Run(() => AddUnrealBuildToolProject(ProgramsFolder, Logger)));
+				DotnetProjectTasks.Add(Task.Run(() => AddUnrealBuildToolProject(ProgramsFolder)));
 
 				// Add AutomationTool to the primary project
-				DotnetProjectTasks.Add(Task.Run(() => AddAutomationToolProject(ProgramsFolder, Logger)));
+				DotnetProjectTasks.Add(Task.Run(() => AddAutomationToolProject(ProgramsFolder)));
 
 				// Add automation.csproj files to the primary project
 				DotnetProjectTasks.Add(Task.Run(() =>
-					AddRulesModules(Rules.RulesFileType.AutomationModule, "Automation", AutomationProjectFiles, AllGameProjects, RootFolder, ProgramsFolder, Logger)));
+					AddRulesModules(Rules.RulesFileType.AutomationModule, "Automation", AutomationProjectFiles, AllGameProjects, RootFolder, ProgramsFolder)));
 
 				// Add ubtplugin.csproj files to the primary project
 				DotnetProjectTasks.Add(Task.Run(() =>
-					AddRulesModules(Rules.RulesFileType.UbtPlugin, "UnrealBuildTool.Plugins", UbtPluginProjectFiles, AllGameProjects, RootFolder, ProgramsFolder, Logger)));
+					AddRulesModules(Rules.RulesFileType.UbtPlugin, "UnrealBuildTool.Plugins", UbtPluginProjectFiles, AllGameProjects, RootFolder, ProgramsFolder)));
 
 				// Add shared projects
-				DotnetProjectTasks.Add(Task.Run(() => AddSharedDotNetModulesAsync(AllEngineDirectories, ProgramsFolder, Logger)));
+				DotnetProjectTasks.Add(Task.Run(() => AddSharedDotNetModulesAsync(AllEngineDirectories, ProgramsFolder)));
 
 				if (bIncludeEnginePrograms)
 				{
@@ -2253,7 +2257,7 @@ namespace UnrealBuildTool
 		/// <summary>
 		/// Adds UnrealBuildTool to the primary project
 		/// </summary>
-		private void AddUnrealBuildToolProject(PrimaryProjectFolder ProgramsFolder, ILogger Logger)
+		private void AddUnrealBuildToolProject(PrimaryProjectFolder ProgramsFolder)
 		{
 			const string UnrealBuildToolFilename = "UnrealBuildTool.csproj";
 			using ITimelineEvent _ = Timeline.ScopeEvent(UnrealBuildToolFilename);
@@ -2262,7 +2266,7 @@ namespace UnrealBuildTool
 			FileReference ProjectFileName = FileReference.Combine(ProgramsDirectory, "UnrealBuildTool", UnrealBuildToolFilename);
 			if (FileReference.Exists(ProjectFileName))
 			{
-				VCSharpProjectFile UnrealBuildToolProject = new VCSharpProjectFile(ProjectFileName, Logger);
+				VCSharpProjectFile UnrealBuildToolProject = CSharpProjCache[ProjectFileName];
 				UnrealBuildToolProject.ShouldBuildForAllSolutionTargets = true;
 
 				// Store it off as we need it when generating target projects.
@@ -2283,7 +2287,7 @@ namespace UnrealBuildTool
 				FileReference UatTestsProjectFileName = FileReference.Combine(ProgramsDirectory, "AutomationTool.Tests", "AutomationTool.Tests.csproj");
 				if (FileReference.Exists(UatTestsProjectFileName))
 				{
-					VCSharpProjectFile UatTestsProject = new VCSharpProjectFile(UatTestsProjectFileName, Logger);
+					VCSharpProjectFile UatTestsProject = CSharpProjCache[UatTestsProjectFileName];
 					AddExistingProjectFile(UatTestsProject, bNeedsAllPlatformAndConfigurations: true, bForceDevelopmentConfiguration: true);
 					lock (ProgramsFolder)
 					{
@@ -2294,7 +2298,7 @@ namespace UnrealBuildTool
 				FileReference UbtTestsProjectFileName = FileReference.Combine(ProgramsDirectory, "UnrealBuildTool.Tests", "UnrealBuildTool.Tests.csproj");
 				if (FileReference.Exists(UbtTestsProjectFileName))
 				{
-					VCSharpProjectFile UbtTestsProject = new VCSharpProjectFile(UbtTestsProjectFileName, Logger);
+					VCSharpProjectFile UbtTestsProject = CSharpProjCache[UbtTestsProjectFileName];
 					AddExistingProjectFile(UbtTestsProject, bNeedsAllPlatformAndConfigurations: true, bForceDevelopmentConfiguration: true);
 					lock (ProgramsFolder)
 					{
@@ -2304,9 +2308,9 @@ namespace UnrealBuildTool
 			}
 		}
 
-		private void AddAutomationToolProject(PrimaryProjectFolder ProgramsFolder, ILogger Logger)
+		private void AddAutomationToolProject(PrimaryProjectFolder ProgramsFolder)
 		{
-			VCSharpProjectFile AutomationToolProject = AddSimpleCSharpProject("AutomationTool", Logger, bShouldBuildForAllSolutionTargets: true, bForceDevelopmentConfiguration: true);
+			VCSharpProjectFile AutomationToolProject = AddSimpleCSharpProject("AutomationTool", bShouldBuildForAllSolutionTargets: true, bForceDevelopmentConfiguration: true);
 			if (AutomationToolProject != null)
 			{
 				lock (ProgramsFolder)
@@ -2320,12 +2324,11 @@ namespace UnrealBuildTool
 		/// Adds a C# project to the primary project
 		/// </summary>
 		/// <param name="ProjectName">Name of project file to add</param>
-		/// <param name="Logger">Logger for output</param>
 		/// <param name="bShouldBuildForAllSolutionTargets"></param>
 		/// <param name="bForceDevelopmentConfiguration"></param>
 		/// <param name="bShouldBuildByDefaultForSolutionTargets"></param>
 		/// <returns>ProjectFile if the operation was successful, otherwise null.</returns>
-		private VCSharpProjectFile AddSimpleCSharpProject(string ProjectName, ILogger Logger, bool bShouldBuildForAllSolutionTargets = false, bool bForceDevelopmentConfiguration = false, bool bShouldBuildByDefaultForSolutionTargets = true)
+		private VCSharpProjectFile AddSimpleCSharpProject(string ProjectName, bool bShouldBuildForAllSolutionTargets = false, bool bForceDevelopmentConfiguration = false, bool bShouldBuildByDefaultForSolutionTargets = true)
 		{
 			VCSharpProjectFile Project;
 
@@ -2334,7 +2337,7 @@ namespace UnrealBuildTool
 			FileInfo Info = new FileInfo(ProjectFileName.FullName);
 			if (Info.Exists)
 			{
-				Project = new VCSharpProjectFile(ProjectFileName, Logger);
+				Project = CSharpProjCache[ProjectFileName];
 
 				Project.ShouldBuildForAllSolutionTargets = bShouldBuildForAllSolutionTargets;
 				Project.ShouldBuildByDefaultForSolutionTargets = bShouldBuildByDefaultForSolutionTargets;
diff --git a/Engine/Source/Programs/UnrealBuildTool/ProjectFiles/QMake/QMakefileGenerator.cs b/Engine/Source/Programs/UnrealBuildTool/ProjectFiles/QMake/QMakefileGenerator.cs
--- a/Engine/Source/Programs/UnrealBuildTool/ProjectFiles/QMake/QMakefileGenerator.cs
+++ b/Engine/Source/Programs/UnrealBuildTool/ProjectFiles/QMake/QMakefileGenerator.cs
@@ -25,8 +25,8 @@ namespace UnrealBuildTool
 	class QMakefileGenerator : ProjectFileGenerator
 	{
 		/// Default constructor
-		public QMakefileGenerator(FileReference? InOnlyGameProject)
-			: base(InOnlyGameProject)
+		public QMakefileGenerator(FileReference? InOnlyGameProject, CSharpProjectCache CSharpProjCache)
+			: base(InOnlyGameProject, CSharpProjCache)
 		{
 		}
 
diff --git a/Engine/Source/Programs/UnrealBuildTool/ProjectFiles/Rider/RiderProjectFileGenerator.cs b/Engine/Source/Programs/UnrealBuildTool/ProjectFiles/Rider/RiderProjectFileGenerator.cs
--- a/Engine/Source/Programs/UnrealBuildTool/ProjectFiles/Rider/RiderProjectFileGenerator.cs
+++ b/Engine/Source/Programs/UnrealBuildTool/ProjectFiles/Rider/RiderProjectFileGenerator.cs
@@ -65,8 +65,9 @@ namespace UnrealBuildTool
 #pragma warning restore IDE0044
 
 		public RiderProjectFileGenerator(FileReference? InOnlyGameProject,
+			CSharpProjectCache CSharpProjCache,
 			CommandLineArguments InArguments)
-			: base(InOnlyGameProject)
+			: base(InOnlyGameProject, CSharpProjCache)
 		{
 			Arguments = InArguments;
 			Arguments.ApplyTo(this);
@@ -478,8 +479,8 @@ namespace UnrealBuildTool
 
 			// Generate all necessary .props files for C# projects
 			PrimaryProjectFolder ProgramsFolder = RootFolder.AddSubFolder("Programs");
-			AddRulesModules(Rules.RulesFileType.AutomationModule, "Automation", AutomationProjectFiles, AllGameProjects, RootFolder, ProgramsFolder, Logger);
-			AddRulesModules(Rules.RulesFileType.UbtPlugin, "UnrealBuildTool.Plugins", UbtPluginProjectFiles, AllGameProjects, RootFolder, ProgramsFolder, Logger);
+			AddRulesModules(Rules.RulesFileType.AutomationModule, "Automation", AutomationProjectFiles, AllGameProjects, RootFolder, ProgramsFolder);
+			AddRulesModules(Rules.RulesFileType.UbtPlugin, "UnrealBuildTool.Plugins", UbtPluginProjectFiles, AllGameProjects, RootFolder, ProgramsFolder);
 
 			Logger.LogDebug("Project generation complete ({NumGenerated} generated, {NumImported} imported)", GeneratedProjectFiles.Count,
 				OtherProjectFiles.Count);
diff --git a/Engine/Source/Programs/UnrealBuildTool/ProjectFiles/VisualStudio/VCProjectFileGenerator.cs b/Engine/Source/Programs/UnrealBuildTool/ProjectFiles/VisualStudio/VCProjectFileGenerator.cs
--- a/Engine/Source/Programs/UnrealBuildTool/ProjectFiles/VisualStudio/VCProjectFileGenerator.cs
+++ b/Engine/Source/Programs/UnrealBuildTool/ProjectFiles/VisualStudio/VCProjectFileGenerator.cs
@@ -131,10 +131,15 @@ namespace UnrealBuildTool
 		/// Default constructor
 		/// </summary>
 		/// <param name="InOnlyGameProject">The single project to generate project files for, or null</param>
+		/// <param name="CSharpProjCache">A cache for C# projects to prevent needing to parse projects multiple times</param>
 		/// <param name="InProjectFileFormat">Override the project file format to use</param>
 		/// <param name="InArguments">Additional command line arguments</param>
-		public VCProjectFileGenerator(FileReference? InOnlyGameProject, VCProjectFileFormat InProjectFileFormat, CommandLineArguments InArguments)
-			: base(InOnlyGameProject)
+		public VCProjectFileGenerator(
+			FileReference? InOnlyGameProject,
+			CSharpProjectCache CSharpProjCache,
+			VCProjectFileFormat InProjectFileFormat,
+			CommandLineArguments InArguments)
+			: base(InOnlyGameProject, CSharpProjCache)
 		{
 			XmlConfig.ApplyTo(Settings);
 
diff --git a/Engine/Source/Programs/UnrealBuildTool/ProjectFiles/VisualStudio/VCSharpProjectFile.cs b/Engine/Source/Programs/UnrealBuildTool/ProjectFiles/VisualStudio/VCSharpProjectFile.cs
--- a/Engine/Source/Programs/UnrealBuildTool/ProjectFiles/VisualStudio/VCSharpProjectFile.cs
+++ b/Engine/Source/Programs/UnrealBuildTool/ProjectFiles/VisualStudio/VCSharpProjectFile.cs
@@ -14,31 +14,25 @@ namespace UnrealBuildTool;
 /// <summary>
 /// Represents an existing C# project file, created by using MSBuild to parse the file from disk.
 /// </summary>
-class VCSharpProjectFile : MSBuildProjectFile
+internal class VCSharpProjectFile : MSBuildProjectFile
 {
 	/// <summary>
 	/// This is the GUID that Visual Studio uses to identify a C# project file in the solution
 	/// </summary>
 	public override string ProjectTypeGUID => "{FAE04EC0-301F-11D3-BF4B-00C04F79EFBC}";
 
-	/// <summary>
-	/// Platforms that this project supports
-	/// </summary>
-	public HashSet<string> Platforms = new HashSet<string>();
-
 	/// <summary>
 	/// Configurations that this project supports
 	/// </summary>
-	public HashSet<string> Configurations = new HashSet<string>();
+	public IReadOnlySet<string> Configurations => _configurations;
 
 	/// <summary>
 	/// Constructs a new project file object
 	/// </summary>
 	/// <param name="InitFilePath">The path to the project file on disk</param>
 	/// <param name="Logger">Logger for output</param>
-	/// <param name="BaseDir">The base directory for files within this project - if not specified, InitFilePath.Directory will be used</param>
-	public VCSharpProjectFile(FileReference InitFilePath, ILogger Logger, DirectoryReference? BaseDir = null)
-		: base(InitFilePath, BaseDir ?? InitFilePath.Directory, true)
+	public VCSharpProjectFile(FileReference InitFilePath, ILogger Logger)
+		: base(InitFilePath, InitFilePath.Directory, true)
 	{
 		try
 		{
@@ -65,8 +59,8 @@ class VCSharpProjectFile : MSBuildProjectFile
 							Match Match = Regex.Match(Condition, "^\\s*'\\$\\(Configuration\\)\\|\\$\\(Platform\\)'\\s*==\\s*'(.+)\\|(.+)'\\s*$");
 							if (Match.Success && Match.Groups.Count == 3)
 							{
-								Configurations.Add(Match.Groups[1].Value);
-								Platforms.Add(Match.Groups[2].Value);
+								_configurations.Add(Match.Groups[1].Value);
+								_platforms.Add(Match.Groups[2].Value);
 							}
 							else
 							{
@@ -80,9 +74,9 @@ class VCSharpProjectFile : MSBuildProjectFile
 			{
 				foreach (string c in GetProjectProperty("Configurations").Split(';'))
 				{
-					Configurations.Add(c);
+					_configurations.Add(c);
 				}
-				bool ConfigurationsFound = Configurations.Any();
+				bool ConfigurationsFound = _configurations.Any();
 
 				if (!ConfigurationsFound)
 				{
@@ -107,13 +101,13 @@ class VCSharpProjectFile : MSBuildProjectFile
 						{
 							foreach (string c in ParsedConfigurations)
 							{
-								Configurations.Add(c);
+								_configurations.Add(c);
 							}
 						}
 
 						// platforms change meaning quite a bit in .net core but typically you do not specify this and its derived from the build instead
 						// for most intents it is just Any CPU from .net framework
-						Platforms.Add("AnyCPU");
+						_platforms.Add("AnyCPU");
 
 						ConfigurationsFound = true;
 						break;
@@ -123,9 +117,9 @@ class VCSharpProjectFile : MSBuildProjectFile
 				// dotnet does not require you to specify configurations or platforms, if you do not debug and release are the defaults
 				if (!ConfigurationsFound)
 				{
-					Configurations.Add("Debug");
-					Configurations.Add("Release");
-					Platforms.Add("AnyCPU");
+					_configurations.Add("Debug");
+					_configurations.Add("Release");
+					_platforms.Add("AnyCPU");
 				}
 			}
 		}
@@ -135,12 +129,24 @@ class VCSharpProjectFile : MSBuildProjectFile
 		}
 	}
 
+	/// <summary>
+	/// Partially clone a VCSharpProjectFile, sharing platforms, configurations and cached project info, but with separate base class properties.
+	/// </summary>
+	/// <param name="other">The instance to partially clone</param>
+	public VCSharpProjectFile(VCSharpProjectFile other)
+		: base(other.ProjectFilePath, other.BaseDir, true)
+	{
+		_platforms = other._platforms;
+		_configurations = other._configurations;
+		_cachedProjectInfo = other._cachedProjectInfo;
+	}
+
 	/// <summary>
 	/// Extract information from the csproj file based on the supplied configuration
 	/// </summary>
 	public IDotnetProjectInfo GetProjectInfo(UnrealTargetConfiguration InConfiguration)
 	{
-		if (CachedProjectInfo.TryGetValue(InConfiguration, out IDotnetProjectInfo? value))
+		if (_cachedProjectInfo.TryGetValue(InConfiguration, out IDotnetProjectInfo? value))
 		{
 			return value;
 		}
@@ -151,7 +157,7 @@ class VCSharpProjectFile : MSBuildProjectFile
 		Properties.Add("EngineDirectory", Unreal.EngineDirectory.FullName);
 
 		DotnetProjectInfo Info = DotnetProjectInfo.Read(ProjectFilePath, Properties);
-		CachedProjectInfo.Add(InConfiguration, Info);
+		_cachedProjectInfo.Add(InConfiguration, Info);
 		return Info;
 	}
 
@@ -179,7 +185,7 @@ class VCSharpProjectFile : MSBuildProjectFile
 	{
 		// Find the matching platform name
 		string ProjectPlatformName;
-		if (Platforms.Contains("x64"))
+		if (_platforms.Contains("x64"))
 		{
 			ProjectPlatformName = "x64";
 		}
@@ -190,11 +196,11 @@ class VCSharpProjectFile : MSBuildProjectFile
 
 		// Find the matching configuration
 		string ProjectConfigurationName;
-		if (Configurations.Contains(SolutionConfiguration.ToString()))
+		if (_configurations.Contains(SolutionConfiguration.ToString()))
 		{
 			ProjectConfigurationName = SolutionConfiguration.ToString();
 		}
-		else if (Configurations.Contains("Development"))
+		else if (_configurations.Contains("Development"))
 		{
 			ProjectConfigurationName = "Development";
 		}
@@ -227,6 +233,16 @@ class VCSharpProjectFile : MSBuildProjectFile
 		throw new BuildException("Support for writing C# projects from UnrealBuildTool has been removed.");
 	}
 
+	/// <summary>
+	/// Platforms that this project supports
+	/// </summary>
+	private readonly HashSet<string> _platforms = [];
+
+	/// <summary>
+	/// Configurations that this project supports
+	/// </summary>
+	private readonly HashSet<string> _configurations = [];
+
 	/// Cache of parsed info about this project
-	protected readonly Dictionary<UnrealTargetConfiguration, IDotnetProjectInfo> CachedProjectInfo = [];
+	protected readonly Dictionary<UnrealTargetConfiguration, IDotnetProjectInfo> _cachedProjectInfo = [];
 }
diff --git a/Engine/Source/Programs/UnrealBuildTool/ProjectFiles/VisualStudioCode/VSCodeProjectFileGenerator.cs b/Engine/Source/Programs/UnrealBuildTool/ProjectFiles/VisualStudioCode/VSCodeProjectFileGenerator.cs
--- a/Engine/Source/Programs/UnrealBuildTool/ProjectFiles/VisualStudioCode/VSCodeProjectFileGenerator.cs
+++ b/Engine/Source/Programs/UnrealBuildTool/ProjectFiles/VisualStudioCode/VSCodeProjectFileGenerator.cs
@@ -164,8 +164,8 @@ namespace UnrealBuildTool
 			}
 		}
 
-		public VSCodeProjectFileGenerator(FileReference? InOnlyGameProject)
-			: base(InOnlyGameProject)
+		public VSCodeProjectFileGenerator(FileReference? InOnlyGameProject, CSharpProjectCache CSharpProjCache)
+			: base(InOnlyGameProject, CSharpProjCache)
 		{
 			ProjectRoot = Unreal.RootDirectory;
 		}
diff --git a/Engine/Source/Programs/UnrealBuildTool/ProjectFiles/VisualStudioWorkspace/VSWorkspaceProjectFileGenerator.cs b/Engine/Source/Programs/UnrealBuildTool/ProjectFiles/VisualStudioWorkspace/VSWorkspaceProjectFileGenerator.cs
--- a/Engine/Source/Programs/UnrealBuildTool/ProjectFiles/VisualStudioWorkspace/VSWorkspaceProjectFileGenerator.cs
+++ b/Engine/Source/Programs/UnrealBuildTool/ProjectFiles/VisualStudioWorkspace/VSWorkspaceProjectFileGenerator.cs
@@ -77,9 +77,11 @@ namespace UnrealBuildTool
 		private bool GenerateVsConfigOnly = false;
 #pragma warning restore IDE0044
 
-		public VSWorkspaceProjectFileGenerator(FileReference? InOnlyGameProject,
+		public VSWorkspaceProjectFileGenerator(
+			FileReference? InOnlyGameProject,
+			CSharpProjectCache CSharpProjCache,
 			CommandLineArguments InArguments)
-			: base(InOnlyGameProject)
+			: base(InOnlyGameProject, CSharpProjCache)
 		{
 			Arguments = InArguments;
 			Arguments.ApplyTo(this);
diff --git a/Engine/Source/Programs/UnrealBuildTool/ProjectFiles/Xcode/XcodeProjectFileGenerator.cs b/Engine/Source/Programs/UnrealBuildTool/ProjectFiles/Xcode/XcodeProjectFileGenerator.cs
--- a/Engine/Source/Programs/UnrealBuildTool/ProjectFiles/Xcode/XcodeProjectFileGenerator.cs
+++ b/Engine/Source/Programs/UnrealBuildTool/ProjectFiles/Xcode/XcodeProjectFileGenerator.cs
@@ -103,8 +103,8 @@ namespace UnrealBuildTool
 			}
 		}
 
-		public XcodeProjectFileGenerator(FileReference? InOnlyGameProject, CommandLineArguments CommandLine)
-			: base(InOnlyGameProject)
+		public XcodeProjectFileGenerator(FileReference? InOnlyGameProject, CSharpProjectCache CSharpProjCache, CommandLineArguments CommandLine)
+			: base(InOnlyGameProject, CSharpProjCache)
 		{
 			SingleGameProject = InOnlyGameProject;
 			if (CommandLine.HasOption("-distribution"))
