UnrealBuildTool: Query valid clang analysis checkers rather than hard coding

#rnx

#Synced-CL 49597545

diff --git a/Engine/Source/Programs/UnrealBuildTool/Platform/Clang/ClangToolChain.cs b/Engine/Source/Programs/UnrealBuildTool/Platform/Clang/ClangToolChain.cs
--- a/Engine/Source/Programs/UnrealBuildTool/Platform/Clang/ClangToolChain.cs
+++ b/Engine/Source/Programs/UnrealBuildTool/Platform/Clang/ClangToolChain.cs
@@ -154,10 +154,12 @@ namespace UnrealBuildTool
 			public Version ClangVersion => LazyClangVersion.Value;
 			public string ClangVersionString => LazyClangVersionString.Value;
 			public string ArchiverVersionString => LazyArchiverVersionString.Value;
+			public HashSet<string> ValidAnalyzerCheckers => LazyValidAnalyzerCheckers.Value;
 
 			readonly Lazy<Version> LazyClangVersion;
 			readonly Lazy<string> LazyClangVersionString;
 			readonly Lazy<string> LazyArchiverVersionString;
+			readonly Lazy<HashSet<string>> LazyValidAnalyzerCheckers;
 
 			/// <summary>
 			/// Constructor for ClangToolChainInfo
@@ -176,6 +178,7 @@ namespace UnrealBuildTool
 				LazyClangVersion = new Lazy<Version>(() => QueryClangVersion());
 				LazyClangVersionString = new Lazy<string>(() => QueryClangVersionString());
 				LazyArchiverVersionString = new Lazy<string>(() => QueryArchiverVersionString());
+				LazyValidAnalyzerCheckers = new Lazy<HashSet<string>>(() => ClangWarnings.ValidAnalyzerCheckers(Clang));
 			}
 
 			/// <summary>
@@ -1002,13 +1005,6 @@ namespace UnrealBuildTool
 			// Enable the static analyzer with default checks.
 			Arguments.Add("--analyze");
 
-			// Deprecated in LLVM 15
-			if (CompilerVersionLessThan(15, 0, 0))
-			{
-				// Make sure we check inside nested blocks (e.g. 'if ((foo = getchar()) == 0) {}')
-				Arguments.Add("-Xclang -analyzer-opt-analyze-nested-blocks");
-			}
-
 			// Write out a pretty web page with navigation to understand how the analysis was derived if HTML is enabled.
 #pragma warning disable CA1308 // Normalize strings to uppercase - Clang arg may be case-sensitive
 			Arguments.Add($"-Xclang -analyzer-output={StaticAnalyzerOutputType.ToString().ToLowerInvariant()}");
@@ -1036,7 +1032,7 @@ namespace UnrealBuildTool
 				Arguments.Add("--analyzer-no-default-checks");
 
 				// Only enable specific checkers.
-				foreach (string Checker in CompileEnvironment.StaticAnalyzerCheckers.Where(x => ClangWarnings.IsAvailableAnalyzerChecker(x, ClangVersion)))
+				foreach (string Checker in CompileEnvironment.StaticAnalyzerCheckers.Where(Info.ValidAnalyzerCheckers.Contains))
 				{
 					Arguments.Add($"-Xclang -analyzer-checker -Xclang {Checker}");
 				}
@@ -1044,12 +1040,12 @@ namespace UnrealBuildTool
 			else
 			{
 				// Disable default checks.
-				foreach (string Checker in CompileEnvironment.StaticAnalyzerDisabledCheckers.Where(x => ClangWarnings.IsAvailableAnalyzerChecker(x, ClangVersion)))
+				foreach (string Checker in CompileEnvironment.StaticAnalyzerDisabledCheckers.Where(Info.ValidAnalyzerCheckers.Contains))
 				{
 					Arguments.Add($"-Xclang -analyzer-disable-checker -Xclang {Checker}");
 				}
 				// Enable additional non-default checks.
-				foreach (string Checker in CompileEnvironment.StaticAnalyzerAdditionalCheckers.Where(x => ClangWarnings.IsAvailableAnalyzerChecker(x, ClangVersion)))
+				foreach (string Checker in CompileEnvironment.StaticAnalyzerAdditionalCheckers.Where(Info.ValidAnalyzerCheckers.Contains))
 				{
 					Arguments.Add($"-Xclang -analyzer-checker -Xclang {Checker}");
 				}
diff --git a/Engine/Source/Programs/UnrealBuildTool/Platform/Clang/ClangWarnings.cs b/Engine/Source/Programs/UnrealBuildTool/Platform/Clang/ClangWarnings.cs
--- a/Engine/Source/Programs/UnrealBuildTool/Platform/Clang/ClangWarnings.cs
+++ b/Engine/Source/Programs/UnrealBuildTool/Platform/Clang/ClangWarnings.cs
@@ -1,8 +1,12 @@
 ﻿// Copyright Epic Games, Inc. All Rights Reserved.
 
 using System;
+using System.Collections.Concurrent;
 using System.Collections.Generic;
+using System.Linq;
+using System.Text.RegularExpressions;
 using EpicGames.Core;
+using UnrealBuildBase;
 
 namespace UnrealBuildTool
 {
@@ -21,22 +25,28 @@ namespace UnrealBuildTool
 			Arguments.Add("-Wno-#pragma-messages");
 		}
 
-		// TODO: Get valid checkers by calling clang -cc1 -analyzer-checker-help
-		static readonly Lazy<Dictionary<string, int>> CheckerAddedVersion = new Lazy<Dictionary<string, int>>(() => new()
+		static readonly ConcurrentDictionary<FileReference, HashSet<string>> s_validClangCheckers = [];
+
+		internal static HashSet<string> ValidAnalyzerCheckers(FileReference clang)
+		{
+			lock (FileItem.GetItemByFileReference(clang))
 			{
-				{ "core.BitwiseShift", 18 },
-				{ "optin.core.EnumCastOutOfRange", 18 },
-				{ "security.cert.env.InvalidPtr", 18 },
-				{ "unix.Errno", 18 },
-				{ "unix.StdCLibraryFunctions", 18 },
-				{ "cplusplus.ArrayDelete", 19 },
-				{ "optin.taint.TaintedAlloc", 19 },
-				{ "security.PutenvStackArray", 19 },
-				{ "security.SetgidSetuidOrder", 19 },
-				{ "unix.BlockInCriticalSection", 19 },
-				{ "unix.Stream", 19 },
-			});
+				return s_validClangCheckers.GetOrAdd(clang, QueryValidAnalyzerCheckers);
+			}
+		}
 
-		internal static bool IsAvailableAnalyzerChecker(string Checker, VersionNumber ClangVersion) => !CheckerAddedVersion.Value.ContainsKey(Checker) || ClangVersion.Components[0] >= CheckerAddedVersion.Value[Checker];
+		static HashSet<string> QueryValidAnalyzerCheckers(FileReference clang)
+		{
+			string[] args = ["-analyzer-checker-help", "-analyzer-checker-help-alpha", "-analyzer-checker-help-developer"];
+			HashSet<string> checkers = [];
+			foreach (string arg in args)
+			{
+				string output = Utils.RunLocalProcessAndReturnStdOut(clang.FullName, $"-cc1 {arg}", null);
+				checkers.AddRange(Regex.Matches(output, "^\\s\\s(\\S*?)\\s", RegexOptions.Multiline)
+					.Select(x => x.Groups[1].Value)
+					.Where(x => !String.IsNullOrWhiteSpace(x)));
+			}
+			return checkers;
+		}
 	}
 }
diff --git a/Engine/Source/Programs/UnrealBuildTool/Platform/Windows/VCToolChain.cs b/Engine/Source/Programs/UnrealBuildTool/Platform/Windows/VCToolChain.cs
--- a/Engine/Source/Programs/UnrealBuildTool/Platform/Windows/VCToolChain.cs
+++ b/Engine/Source/Programs/UnrealBuildTool/Platform/Windows/VCToolChain.cs
@@ -522,13 +522,15 @@ namespace UnrealBuildTool
 						Arguments.Add("-Xclang -analyzer-config -Xclang mode=shallow");
 					}
 
+					HashSet<string> ValidAnalyzerCheckers = ClangWarnings.ValidAnalyzerCheckers(EnvVars.CompilerPath);
+
 					if (CompileEnvironment.StaticAnalyzerCheckers.Count > 0)
 					{
 						// Disable all default checks
 						Arguments.Add("--analyzer-no-default-checks");
 
 						// Only enable specific checks.
-						foreach (string Checker in CompileEnvironment.StaticAnalyzerCheckers.Where(x => ClangWarnings.IsAvailableAnalyzerChecker(x, EnvVars.CompilerVersion)))
+						foreach (string Checker in CompileEnvironment.StaticAnalyzerCheckers.Where(ValidAnalyzerCheckers.Contains))
 						{
 							Arguments.Add($"-Xclang -analyzer-checker -Xclang {Checker}");
 						}
@@ -536,12 +538,12 @@ namespace UnrealBuildTool
 					else
 					{
 						// Disable default checks.
-						foreach (string Checker in CompileEnvironment.StaticAnalyzerDisabledCheckers.Where(x => ClangWarnings.IsAvailableAnalyzerChecker(x, EnvVars.CompilerVersion)))
+						foreach (string Checker in CompileEnvironment.StaticAnalyzerDisabledCheckers.Where(ValidAnalyzerCheckers.Contains))
 						{
 							Arguments.Add($"-Xclang -analyzer-disable-checker -Xclang {Checker}");
 						}
 						// Enable additional non-default checks.
-						foreach (string Checker in CompileEnvironment.StaticAnalyzerAdditionalCheckers.Where(x => ClangWarnings.IsAvailableAnalyzerChecker(x, EnvVars.CompilerVersion)))
+						foreach (string Checker in CompileEnvironment.StaticAnalyzerAdditionalCheckers.Where(ValidAnalyzerCheckers.Contains))
 						{
 							Arguments.Add($"-Xclang -analyzer-checker -Xclang {Checker}");
 						}
