Check disallowed modules when not linking
Update several plugin modules from Runtime -> DeveloperTool to prevent pulling in Shipping disallowed Developer modules

#rnx
#Synced-CL 49598926

diff --git a/Engine/Plugins/Experimental/MovieSceneAnimMixer/MovieSceneAnimMixer.uplugin b/Engine/Plugins/Experimental/MovieSceneAnimMixer/MovieSceneAnimMixer.uplugin
--- a/Engine/Plugins/Experimental/MovieSceneAnimMixer/MovieSceneAnimMixer.uplugin
+++ b/Engine/Plugins/Experimental/MovieSceneAnimMixer/MovieSceneAnimMixer.uplugin
@@ -30,7 +30,7 @@
 	[
 		{
 			"Name" : "MovieSceneAnimMixer",
-			"Type" : "Runtime",
+			"Type" : "DeveloperTool",
 			"LoadingPhase" : "Default"
 		},
 		{
diff --git a/Engine/Plugins/Experimental/MovieSceneMixedControlRig/MovieSceneMixedControlRig.uplugin b/Engine/Plugins/Experimental/MovieSceneMixedControlRig/MovieSceneMixedControlRig.uplugin
--- a/Engine/Plugins/Experimental/MovieSceneMixedControlRig/MovieSceneMixedControlRig.uplugin
+++ b/Engine/Plugins/Experimental/MovieSceneMixedControlRig/MovieSceneMixedControlRig.uplugin
@@ -45,7 +45,7 @@
 	[
 		{
 			"Name" : "MovieSceneMixedControlRig",
-			"Type" : "Runtime",
+			"Type" : "DeveloperTool",
 			"LoadingPhase" : "Default"
 		}
 	]
diff --git a/Engine/Plugins/Experimental/MovieScenePoseSearchTracks/MovieScenePoseSearchTracks.uplugin b/Engine/Plugins/Experimental/MovieScenePoseSearchTracks/MovieScenePoseSearchTracks.uplugin
--- a/Engine/Plugins/Experimental/MovieScenePoseSearchTracks/MovieScenePoseSearchTracks.uplugin
+++ b/Engine/Plugins/Experimental/MovieScenePoseSearchTracks/MovieScenePoseSearchTracks.uplugin
@@ -30,7 +30,7 @@
 	[
 		{
 			"Name" : "MovieScenePoseSearchTracks",
-			"Type" : "Runtime",
+			"Type" : "DeveloperTool",
 			"LoadingPhase" : "Default"
 		},
 		{
diff --git a/Engine/Plugins/Runtime/ScriptableToolsFramework/ScriptableToolsFramework.uplugin b/Engine/Plugins/Runtime/ScriptableToolsFramework/ScriptableToolsFramework.uplugin
--- a/Engine/Plugins/Runtime/ScriptableToolsFramework/ScriptableToolsFramework.uplugin
+++ b/Engine/Plugins/Runtime/ScriptableToolsFramework/ScriptableToolsFramework.uplugin
@@ -17,7 +17,7 @@
 	"Modules": [
 		{
 			"Name": "ScriptableToolsFramework",
-			"Type": "Runtime",
+			"Type": "DeveloperTool",
 			"LoadingPhase": "PreDefault"
 		},
 		{
diff --git a/Engine/Plugins/Tests/AudioCQTests/AudioCQTests.uplugin b/Engine/Plugins/Tests/AudioCQTests/AudioCQTests.uplugin
--- a/Engine/Plugins/Tests/AudioCQTests/AudioCQTests.uplugin
+++ b/Engine/Plugins/Tests/AudioCQTests/AudioCQTests.uplugin
@@ -14,10 +14,10 @@
   "IsBetaVersion": false,
   "Installed": false,
   "Modules": [
-    {
-      "Name": "AudioCQTests",
-      "Type": "Runtime",
-      "LoadingPhase": "Default"
-    }
+	{
+		"Name": "AudioCQTests",
+		"Type": "DeveloperTool",
+		"LoadingPhase": "Default"
+	}
   ]
 }
diff --git a/Engine/Plugins/Tests/CQTestExperimental/CQTestExperimental.uplugin b/Engine/Plugins/Tests/CQTestExperimental/CQTestExperimental.uplugin
--- a/Engine/Plugins/Tests/CQTestExperimental/CQTestExperimental.uplugin
+++ b/Engine/Plugins/Tests/CQTestExperimental/CQTestExperimental.uplugin
@@ -17,7 +17,7 @@
 	"Modules": [
 		{
 			"Name": "CQTestExperimentalTests",
-			"Type": "Runtime",
+			"Type": "DeveloperTool",
 			"LoadingPhase": "Default"
 		}
 	]
diff --git a/Engine/Plugins/TextureGraph/Source/TextureGraphEngine/TextureGraphEngine.Build.cs b/Engine/Plugins/TextureGraph/Source/TextureGraphEngine/TextureGraphEngine.Build.cs
--- a/Engine/Plugins/TextureGraph/Source/TextureGraphEngine/TextureGraphEngine.Build.cs
+++ b/Engine/Plugins/TextureGraph/Source/TextureGraphEngine/TextureGraphEngine.Build.cs
@@ -85,7 +85,7 @@ public class TextureGraphEngine : ModuleRules
 			"DeveloperSettings",
 		});
 
-		if (Target.Platform == UnrealTargetPlatform.Win64 || Target.Platform == UnrealTargetPlatform.Linux)
+		if (Target.bBuildDeveloperTools && (Target.Platform == UnrealTargetPlatform.Win64 || Target.Platform == UnrealTargetPlatform.Linux))
 		{
 			PublicDependencyModuleNames.Add("RenderDocPlugin");
 			PublicDefinitions.Add("TEXTUREGRAPH_RENDERDOC_ENABLED=1");
diff --git a/Engine/Plugins/VirtualProduction/Avalanche/Avalanche.uplugin b/Engine/Plugins/VirtualProduction/Avalanche/Avalanche.uplugin
--- a/Engine/Plugins/VirtualProduction/Avalanche/Avalanche.uplugin
+++ b/Engine/Plugins/VirtualProduction/Avalanche/Avalanche.uplugin
@@ -250,7 +250,7 @@
 		},
 		{
 			"Name": "AvalancheFunctionalTest",
-			"Type": "Runtime",
+			"Type": "DeveloperTool",
 			"LoadingPhase": "Default",
 			"TargetDenyList": [ "Server" ]
 		}
diff --git a/Engine/Source/Programs/UnrealBuildTool/Configuration/UEBuildTarget.cs b/Engine/Source/Programs/UnrealBuildTool/Configuration/UEBuildTarget.cs
--- a/Engine/Source/Programs/UnrealBuildTool/Configuration/UEBuildTarget.cs
+++ b/Engine/Source/Programs/UnrealBuildTool/Configuration/UEBuildTarget.cs
@@ -3399,11 +3399,11 @@ namespace UnrealBuildTool
 						throw new BuildException("Unable to create binaries in less restricted locations than their input files.");
 					}
 				}
-
-				// Check for linking against modules prohibited by the EULA
-				CheckForEULAViolation(Logger);
 			}
 
+			// Check for linking against modules prohibited by the EULA
+			CheckForEULAViolation(Logger);
+
 			// Add all the plugins to be tracked
 			foreach (FileReference PluginFile in PluginsBase.EnumeratePlugins(ProjectFile))
 			{
@@ -3892,10 +3892,10 @@ namespace UnrealBuildTool
 						IEnumerable<UEBuildModule> NonRedistDeps = AllDependencies.Where((DependantModule) =>
 							DependantModule.GetDirectDependencyModules().Intersect(NonRedistModules).Any()
 						);
-						string Message = String.Format("Non-editor build cannot depend on non-redistributable modules. {0} depends on '{1}'.", Binary.ToString(), String.Join("', '", NonRedistModules));
+						string Message = $"Non-editor build cannot depend on non-redistributable modules. {Binary.ToString()} depends on '{String.Join("', '", NonRedistModules)}'.";
 						if (NonRedistDeps.Any())
 						{
-							Message = String.Format("{0}\nDependant modules '{1}'", Message, String.Join("', '", NonRedistDeps));
+							Message = $"{Message}\nDependant modules '{String.Join("', '", NonRedistDeps)}'";
 						}
 						if (Rules.bBreakBuildOnLicenseViolation)
 						{
@@ -3929,13 +3929,13 @@ namespace UnrealBuildTool
 			{
 				foreach (DirectoryReference ExtensionDir in Unreal.GetExtensionDirs(Unreal.EngineDirectory))
 				{
-					DirectoryReference SourceDeveloperDir = DirectoryReference.Combine(ExtensionDir, "Source/Developer");
+					DirectoryReference SourceDeveloperDir = DirectoryReference.Combine(ExtensionDir, "Source", "Developer");
 					if (Module.RulesFile.IsUnderDirectory(SourceDeveloperDir))
 					{
 						return false;
 					}
 
-					DirectoryReference SourceEditorDir = DirectoryReference.Combine(ExtensionDir, "Source/Editor");
+					DirectoryReference SourceEditorDir = DirectoryReference.Combine(ExtensionDir, "Source", "Editor");
 					if (Module.RulesFile.IsUnderDirectory(SourceEditorDir))
 					{
 						return false;
diff --git a/Engine/Source/Runtime/Experimental/TraceBasedDebuggersAnalysis/TraceBasedDebuggersAnalysis.Build.cs b/Engine/Source/Runtime/Experimental/TraceBasedDebuggersAnalysis/TraceBasedDebuggersAnalysis.Build.cs
--- a/Engine/Source/Runtime/Experimental/TraceBasedDebuggersAnalysis/TraceBasedDebuggersAnalysis.Build.cs
+++ b/Engine/Source/Runtime/Experimental/TraceBasedDebuggersAnalysis/TraceBasedDebuggersAnalysis.Build.cs
@@ -41,7 +41,7 @@ namespace UnrealBuildTool.Rules
 			}
 
 			// Allow debugger trace analysis on desktop platforms
-			if (Target.Platform.IsInGroup(UnrealPlatformGroup.Desktop))
+			if (Target.bBuildDeveloperTools && Target.Platform.IsInGroup(UnrealPlatformGroup.Desktop))
 			{
 				PublicDefinitions.Add("UE_WITH_TRACE_BASED_DEBUGGERS_ANALYSIS=1");
 				PublicDependencyModuleNames.AddRange(
