UnrealBuildTool: Print a recommendation if a uba cache requires VFS when VFS isn't enabled

#rnx
#Synced-CL 49598926

diff --git a/Engine/Source/Programs/UnrealBuildTool/Executors/UnrealBuildAccelerator/UBAExecutor.cs b/Engine/Source/Programs/UnrealBuildTool/Executors/UnrealBuildAccelerator/UBAExecutor.cs
--- a/Engine/Source/Programs/UnrealBuildTool/Executors/UnrealBuildAccelerator/UBAExecutor.cs
+++ b/Engine/Source/Programs/UnrealBuildTool/Executors/UnrealBuildAccelerator/UBAExecutor.cs
@@ -317,7 +317,7 @@ namespace UnrealBuildTool
 
 			EpicGames.UBA.ILogger ubaLogger = EpicGames.UBA.ILogger.GlobalLogger!;
 			_server = IServer.CreateServer(UBAConfig.MaxWorkers, UBAConfig.SendSize, ubaLogger, UBAConfig.bUseQuic);
-			_storageServer = IStorageServer.CreateStorageServer(_server, ubaLogger, new StorageServerCreateInfo(_rootDirRef.FullName, ((ulong)UBAConfig.StoreCapacityGb) * 1000 * 1000 * 1000, !UBAConfig.bStoreRaw, UBAConfig.Zone));
+			_storageServer = IStorageServer.CreateStorageServer(_server, ubaLogger, new StorageServerCreateInfo(_rootDirRef.FullName, ((ulong)UBAConfig.StoreCapacityGb) * 1024 * 1024 * 1024, !UBAConfig.bStoreRaw, UBAConfig.Zone));
 			using ISessionServerCreateInfo serverCreateInfo = ISessionServerCreateInfo.CreateSessionServerCreateInfo(_storageServer, _server, ubaLogger, new SessionServerCreateInfo(_rootDirRef.FullName, ubaTraceFile.FullName.Replace('\\', '/'), UBAConfig.bDisableCustomAlloc, false, UBAConfig.bResetCas, UBAConfig.bWriteToDisk, UBAConfig.bDetailedTrace, !UBAConfig.bDisableWaitOnMem, UBAConfig.bAllowKillOnMem, UBAConfig.bStoreObjFilesCompressed));
 			_session = ISessionServer.CreateSessionServer(serverCreateInfo);
 			_trace = _session.GetTrace();
@@ -432,17 +432,32 @@ namespace UnrealBuildTool
 			_trace?.UpdateStatus(statusRow, statusColumn, statusText, statusType, statusLink);
 		}
 
-		private void PrintConfiguration(ILogger logger)
+		private void PrintConfiguration(IEnumerable<LinkedAction> inputActions, ILogger logger)
 		{
-			logger.LogInformation("  Storage capacity {StoreCapacityGb}Gb", UBAConfig.StoreCapacityGb);
+			logger.LogInformation("  Storage capacity {StoreCapacity}", StringUtils.FormatBytesString((ulong)UBAConfig.StoreCapacityGb * 1024 * 1024 * 1024));
 			if (UBAConfig.bStoreObjFilesCompressed)
 			{
 				logger.LogInformation("  Object Compression Allowed");
 			}
 			logger.LogDebug("  RootDir {RootDir}", _rootDirRef);
+
+			CaptureLogger recommendationLogger = new();
+			if (_cacheableActionsLeft > 0 && !inputActions.Any(x => x.RootPaths.bUseVfs))
+			{
+				foreach (KeyValuePair<UnrealBuildAcceleratorCacheConfig, ICacheClient> pair in _cacheClients.Where(x => x.Key.bRequireVfs))
+				{
+					recommendationLogger.LogInformation("    Cache '{Server}' requires VFS but no actions found with VFS enabled. Please ensure TargetRules.bUseVFS is enabled.", pair.Key.CacheServer);
+				}
+			}
+
+			if (recommendationLogger.Events.Count > 0)
+			{
+				logger.LogInformation("  Recommendations:");
+				recommendationLogger.RenderTo(logger);
+			}
 		}
 
-		private async Task WriteActionOutputFileAsync(ILogger logger, IEnumerable<LinkedAction> inputActions)
+		private async Task WriteActionOutputFileAsync(IEnumerable<LinkedAction> inputActions, ILogger logger)
 		{
 			if (String.IsNullOrEmpty(UBAConfig.ActionsOutputFile))
 			{
@@ -561,8 +576,6 @@ namespace UnrealBuildTool
 
 			_initTask.Wait();
 
-			PrintConfiguration(logger);
-
 			_trace?.UpdateStatus(1, 6, "Hits 0 Misses 0", LogEntryType.Info);
 
 			if (!inputActions.Any())
@@ -585,7 +598,9 @@ namespace UnrealBuildTool
 				_cacheableActionsLeft = inputActions.Count(x => x.ArtifactMode.HasFlag(ArtifactMode.Enabled));
 			}
 
-			await WriteActionOutputFileAsync(logger, inputActions);
+			PrintConfiguration(inputActions, logger);
+
+			await WriteActionOutputFileAsync(inputActions, logger);
 
 			try
 			{
