UnrealBuildTool: Don't throw on task cancellation in the ContentAddressableActionResultStore.

#Synced-CL 49616476

diff --git a/Engine/Source/Programs/UnrealBuildTool/Actions/ResultStore/ContentAddressableActionResultStore.cs b/Engine/Source/Programs/UnrealBuildTool/Actions/ResultStore/ContentAddressableActionResultStore.cs
--- a/Engine/Source/Programs/UnrealBuildTool/Actions/ResultStore/ContentAddressableActionResultStore.cs
+++ b/Engine/Source/Programs/UnrealBuildTool/Actions/ResultStore/ContentAddressableActionResultStore.cs
@@ -134,9 +134,16 @@ internal class ContentAddressableActionResultStore(IStorageProvider storageProvi
 
 		public async ValueTask DisposeAsync()
 		{
-			await using Stream stream = storageProvider.CreateWriter(targetHash, cancellationToken);
-			PartialActionResultManifest manifest = new(targetDescriptor, [.. Actions]);
-			await JsonSerializer.SerializeAsync(stream, manifest, _jsonSerializerOptions, cancellationToken);
+			try
+			{
+				await using Stream stream = storageProvider.CreateWriter(targetHash, cancellationToken);
+				PartialActionResultManifest manifest = new(targetDescriptor, [.. Actions]);
+				await JsonSerializer.SerializeAsync(stream, manifest, _jsonSerializerOptions, cancellationToken);
+			}
+			catch (TaskCanceledException)
+			{
+				// Don't throw on cancellation.
+			}
 		}
 	}
 }
