Ensure Resources/Localizations/*.lproj are deployed to Apple app bundles under Modernized Xcode.

#Synced-CL 49564828

diff --git a/Engine/Source/Programs/UnrealBuildTool/ProjectFiles/Xcode/XcodeProject.cs b/Engine/Source/Programs/UnrealBuildTool/ProjectFiles/Xcode/XcodeProject.cs
--- a/Engine/Source/Programs/UnrealBuildTool/ProjectFiles/Xcode/XcodeProject.cs
+++ b/Engine/Source/Programs/UnrealBuildTool/ProjectFiles/Xcode/XcodeProject.cs
@@ -1791,8 +1791,11 @@ namespace UnrealBuildTool.XcodeProjectXcconfig
 				}
 			}
 
+			string? LocalizationsPlatform = null;
+			
 			if (Platform == UnrealTargetPlatform.Mac)
 			{
+				LocalizationsPlatform = "Mac";
 				ResourcesBuildPhase.AddFolderResource(DirectoryReference.Combine(Unreal.EngineDirectory, "Build/Mac/Resources/UEMetadata"), "Resources");
 				if (UnrealData.Metadata?.ProjectPrivacyInfoFiles.ContainsKey(MetadataPlatform.Mac) == true)
 				{
@@ -1805,6 +1808,7 @@ namespace UnrealBuildTool.XcodeProjectXcconfig
 			}
 			else if (Platform == UnrealTargetPlatform.IOS || Platform == UnrealTargetPlatform.TVOS || Platform == UnrealTargetPlatform.VisionOS)
 			{
+				LocalizationsPlatform = "IOS";
 				ResourcesBuildPhase.AddFolderResource(DirectoryReference.Combine(Unreal.EngineDirectory, "Build/IOS/Resources/UEMetadata"), "Resources");
 				if (UnrealData.Metadata?.ProjectPrivacyInfoFiles.ContainsKey(MetadataPlatform.IOS) == true)
 				{
@@ -1816,6 +1820,19 @@ namespace UnrealBuildTool.XcodeProjectXcconfig
 				}
 			}
 
+			if (LocalizationsPlatform != null)
+			{
+				DirectoryReference LocalizationsDir = DirectoryReference.Combine(UnrealData.ProductDirectory, "Build", LocalizationsPlatform, "Resources", "Localizations");
+				if (DirectoryReference.Exists(LocalizationsDir))
+				{
+					IEnumerable<DirectoryReference> Localizations = DirectoryReference.EnumerateDirectories(LocalizationsDir, "*.lproj");
+					foreach (DirectoryReference Localization in Localizations)
+					{
+						ResourcesBuildPhase.AddFolderResource(Localization, "Resources");
+					}					
+				}
+			}
+
 			string? ExtraFolderToCopyToApp;
 			ConfigHierarchy PlatformIni = ConfigCache.ReadHierarchy(ConfigHierarchyType.Engine, UnrealData.ConfigDirectory, Platform);
 			PlatformIni.TryGetValue("/Script/MacTargetPlatform.XcodeProjectSettings", "ExtraFolderToCopyToApp", out ExtraFolderToCopyToApp);
