[Mac] Fix race condition when building multiple Editor targets

Fix race condition when building multiple Editor-type targets together on Mac (e.g., MHATestEditor + LiveLinkHub).

Previously:
- Editor targets with bForceTargetNameAsExeName=true still attempted to write their binaries into the UnrealEditor.app bundle causing the binary to be placed in a nested directory, which also resulted in signing failures.
- Each editor target would write to the same intermediate PostBuildSync.json file causing a race condition resulting in binary files being copied into the wrong app bundle.

Changes:
- Uses MacPlatform.bForceTargetNameAsExeName to determine the binary output path in UEBuildTarget.cs
- Exposes bForceTargetNameAsExeName on ReadOnlyMacTargetRules in UEBuildMac.cs
- Makes the PostBuildSync JSON filename unique per target in AppleToolChain.cs to prevent overwrites (e.g. PostBuildSync_LiveLinkHub.json).

#rnx

#Synced-CL 49653240

diff --git a/Engine/Source/Programs/UnrealBuildTool/Configuration/UEBuildTarget.cs b/Engine/Source/Programs/UnrealBuildTool/Configuration/UEBuildTarget.cs
--- a/Engine/Source/Programs/UnrealBuildTool/Configuration/UEBuildTarget.cs
+++ b/Engine/Source/Programs/UnrealBuildTool/Configuration/UEBuildTarget.cs
@@ -6088,7 +6088,12 @@ namespace UnrealBuildTool
 			}
 
 			bool bCompileAsDLL = Rules.bShouldCompileAsDLL && bCompileMonolithic;
-			List<FileReference> OutputPaths = MakeBinaryPaths(OutputDirectory, bCompileMonolithic ? TargetName : AppName, Platform, Configuration, bCompileAsDLL ? UEBuildBinaryType.DynamicLinkLibrary : UEBuildBinaryType.Executable, Rules.Architectures, Rules.UndecoratedConfiguration, Rules.ExeBinariesSubFolder, ProjectFile, Rules);
+			string BinaryName = bCompileMonolithic ? TargetName : AppName;
+			if (Platform == UnrealTargetPlatform.Mac && Rules.MacPlatform.bForceTargetNameAsExeName)
+			{
+				BinaryName = TargetName;
+			}
+			List<FileReference> OutputPaths = MakeBinaryPaths(OutputDirectory, BinaryName, Platform, Configuration, bCompileAsDLL ? UEBuildBinaryType.DynamicLinkLibrary : UEBuildBinaryType.Executable, Rules.Architectures, Rules.UndecoratedConfiguration, Rules.ExeBinariesSubFolder, ProjectFile, Rules);
 
 			// Create the binary
 			UEBuildBinary Binary = new UEBuildBinary(
diff --git a/Engine/Source/Programs/UnrealBuildTool/Platform/Apple/AppleToolChain.cs b/Engine/Source/Programs/UnrealBuildTool/Platform/Apple/AppleToolChain.cs
--- a/Engine/Source/Programs/UnrealBuildTool/Platform/Apple/AppleToolChain.cs
+++ b/Engine/Source/Programs/UnrealBuildTool/Platform/Apple/AppleToolChain.cs
@@ -1189,7 +1189,7 @@ namespace UnrealBuildTool
 		public static Action CreatePostBuildSyncAction(ReadOnlyTargetRules Target, FileItem Executable, DirectoryReference IntermediateDir, IActionGraphBuilder Graph, bool bNoEntitlements = false)
 		{
 			ApplePostBuildSyncTarget PostBuildSync = new(Target, Executable, IntermediateDir);
-			FileReference PostBuildSyncFile = FileReference.Combine(IntermediateDir!, "PostBuildSync.json");
+			FileReference PostBuildSyncFile = FileReference.Combine(IntermediateDir!, $"PostBuildSync_{Target.Name}.json");
 			DirectoryReference.CreateDirectory(PostBuildSyncFile.Directory);
 			using (FileStream stream = new(PostBuildSyncFile.FullName, FileMode.Create, FileAccess.Write))
 			{
diff --git a/Engine/Source/Programs/UnrealBuildTool/Platform/Mac/UEBuildMac.cs b/Engine/Source/Programs/UnrealBuildTool/Platform/Mac/UEBuildMac.cs
--- a/Engine/Source/Programs/UnrealBuildTool/Platform/Mac/UEBuildMac.cs
+++ b/Engine/Source/Programs/UnrealBuildTool/Platform/Mac/UEBuildMac.cs
@@ -93,6 +93,8 @@ namespace UnrealBuildTool
 
 		public bool bEnableRayTracing => Inner.bEnableRayTracing;
 
+		public bool bForceTargetNameAsExeName => Inner.bForceTargetNameAsExeName;
+
 #pragma warning restore CS1591
 		#endregion
 	}
