Fix app extensions being included in other Targets by being able to specify which Products the appex should be included in.

#rnx

#Synced-CL 49688090

diff --git a/Engine/Source/Programs/UnrealBuildTool/ProjectFiles/Xcode/XcodeProject.cs b/Engine/Source/Programs/UnrealBuildTool/ProjectFiles/Xcode/XcodeProject.cs
--- a/Engine/Source/Programs/UnrealBuildTool/ProjectFiles/Xcode/XcodeProject.cs
+++ b/Engine/Source/Programs/UnrealBuildTool/ProjectFiles/Xcode/XcodeProject.cs
@@ -1383,6 +1383,7 @@ namespace UnrealBuildTool.XcodeProjectXcconfig
 		public int FileVersion { get; set; }
 		public string? ExtensionType { get; set; }
 		public string? Description { get; set; }
+		public string[] Products { get; set; } = Array.Empty<string>();
 		public string[] Frameworks { get; set; } = Array.Empty<string>();
 	}
 #pragma warning restore CS1591
@@ -1436,8 +1437,6 @@ namespace UnrealBuildTool.XcodeProjectXcconfig
 					{
 						foreach (DirectoryInfo DI in new System.IO.DirectoryInfo(ProjectExtensionsDir).EnumerateDirectories())
 						{
-							Logger.LogInformation("Creating App Extension Target for '{Name}' in '{FullName}'", DI.Name, DI.FullName);
-							
 							// Get the .uappex details file for this extension
 							ExtensionConfig? ExtConfig;
 							string UExtPath = Path.Combine(DI.FullName, DI.Name) + ".uappex";
@@ -1457,11 +1456,19 @@ namespace UnrealBuildTool.XcodeProjectXcconfig
 								Logger.LogWarning("There is no extension '{Name}.uappex' file. Skipping Extension...", DI.Name);
 								continue;
 							}
+							
+							// Check that this appex should be added to this Produst instance. ie: there maybe be multiple programs/products
+							// defined in the Project, but only a one or more might need this extension
+							if (!ExtConfig.Products.Contains(Project.UnrealData.ProductName))
+							{
+								continue;
+							}
 
+							Logger.LogInformation("Adding App Extension Target '{Name}' to '{ProductName}'", DI.Name, Project.UnrealData.ProductName);
+							
 							// assume each Extension in here will create a resulting Extension.appex
 							string Extension = DI.Name + ".appex";
 							string ExtensionGuid = XcodeProjectFileGenerator.MakeXcodeGuid();
-							
 							XcodeCopyAppExtFilesBuildPhase EmbedAppExt = new XcodeCopyAppExtFilesBuildPhase(Project.FileCollection, ExtConfig);
 							
 							// Get the source files
@@ -1472,10 +1479,11 @@ namespace UnrealBuildTool.XcodeProjectXcconfig
 									continue;
 								}
 
-								Logger.LogInformation("Ext File - init: '{Name}'", FSI.Name);
 								FileReference FileRef = new FileReference(FSI.FullName);
-								XcodeSourceFile SourceFile = new XcodeSourceFile(FileRef, FileRef.Directory, XcodeProjectFileGenerator.MakeXcodeGuid());
+								string FileGuid = XcodeProjectFileGenerator.MakeXcodeGuid();
+								XcodeSourceFile SourceFile = new XcodeSourceFile(FileRef, FileRef.Directory, FileGuid);
 								Project.FileCollection.ProcessFile(SourceFile, true, false, "Source");
+								// Logger.LogInformation("  AppEx file:'{FileRef}'  FileRef: '{FileGuid}' GUID:'{ExtensionGuid}'", FileRef, FileGuid, ExtensionGuid);
 							}
 							
 							Project.FileCollection.AddFileReference(ExtensionGuid, Extension, "explicitFileType", "wrapper.app-extension", "BUILT_PRODUCTS_DIR", "Products");
