UnrealBuildTool: Fix for duplicate C# projects in generated solutions that caused solutions to be marked as dirty by VS.

#Synced-CL 49759094

diff --git a/Engine/Source/Programs/UnrealBuildTool/ProjectFiles/VisualStudio/VCProjectFileGenerator.cs b/Engine/Source/Programs/UnrealBuildTool/ProjectFiles/VisualStudio/VCProjectFileGenerator.cs
--- a/Engine/Source/Programs/UnrealBuildTool/ProjectFiles/VisualStudio/VCProjectFileGenerator.cs
+++ b/Engine/Source/Programs/UnrealBuildTool/ProjectFiles/VisualStudio/VCProjectFileGenerator.cs
@@ -457,11 +457,13 @@ namespace UnrealBuildTool
 			string OriginalPrimaryProjectName = PrimaryProjectName;
 			PrimaryProjectFolder OriginalRootFolder = RootFolder;
 			List<ProjectFile> OriginalAutomationProjectFiles = [.. AutomationProjectFiles];
+			List<ProjectFile> OriginalOtherProjectFiles = [.. OtherProjectFiles];
 			try
 			{
 				PrimaryProjectName = $"Automation_{PrimaryProjectName}";
 				RootFolder = new(this, "<Root>");
 				AutomationProjectFiles.Clear();
+				OtherProjectFiles.Clear();
 
 				AddAutomationProjects(AllGameProjects);
 
@@ -476,6 +478,8 @@ namespace UnrealBuildTool
 				RootFolder = OriginalRootFolder;
 				AutomationProjectFiles.Clear();
 				AutomationProjectFiles.AddRange(OriginalAutomationProjectFiles);
+				OtherProjectFiles.Clear();
+				OtherProjectFiles.AddRange(OriginalOtherProjectFiles);
 			}
 		}
 
@@ -779,7 +783,7 @@ namespace UnrealBuildTool
 			{
 				VCSolutionFileContent.AppendLine("Microsoft Visual Studio Solution File, Format Version 12.00");
 				VCSolutionFileContent.AppendLine("# Visual Studio Version 18");
-				VCSolutionFileContent.AppendLine("VisualStudioVersion = 18.0.11205.157 d18.0");
+				VCSolutionFileContent.AppendLine("VisualStudioVersion = 18.0.11205.157");
 				VCSolutionFileContent.AppendLine("MinimumVisualStudioVersion = 10.0.40219.1");
 			}
 			else
