UnrealBuildTool: Downgrade failure to load makefile message to information

#rnx
#Synced-CL 49781069

diff --git a/Engine/Source/Programs/UnrealBuildTool/System/TargetMakefile.cs b/Engine/Source/Programs/UnrealBuildTool/System/TargetMakefile.cs
--- a/Engine/Source/Programs/UnrealBuildTool/System/TargetMakefile.cs
+++ b/Engine/Source/Programs/UnrealBuildTool/System/TargetMakefile.cs
@@ -788,7 +788,7 @@ namespace UnrealBuildTool
 				}
 				catch (Exception Ex)
 				{
-					Logger.LogWarning("Failed to read makefile {File}: {ExMessage}", MakefilePath.FullName, Ex.Message);
+					Logger.LogInformation("Failed to read makefile {File}: {ExMessage}", MakefilePath.FullName, Ex.Message);
 					Logger.LogDebug(Ex, "Exception: {Ex}", Ex.ToString());
 					ReasonNotLoaded = "couldn't read existing makefile";
 					return null;
