UnrealBuildTool: Do not use insecure randomness (CA5394)

#rnx

#Synced-CL 49789638

diff --git a/Engine/Source/Programs/UnrealBuildTool/ProjectFiles/Xcode/XcodeProjectFileGenerator.cs b/Engine/Source/Programs/UnrealBuildTool/ProjectFiles/Xcode/XcodeProjectFileGenerator.cs
--- a/Engine/Source/Programs/UnrealBuildTool/ProjectFiles/Xcode/XcodeProjectFileGenerator.cs
+++ b/Engine/Source/Programs/UnrealBuildTool/ProjectFiles/Xcode/XcodeProjectFileGenerator.cs
@@ -166,7 +166,9 @@ namespace UnrealBuildTool
 			byte[] Randoms = new byte[12];
 			lock(Rand)
 			{
-				Rand.NextBytes(Randoms);
+#pragma warning disable CA5394 // Random Guids need to be stable so Random is fine
+			Rand.NextBytes(Randoms);
+#pragma warning restore CA5394
 			}
 			for (int Index = 0; Index < 12; Index++)
 			{
