UnrealBuildTool: Adjust the global mutex for xmlconfig to always be active when reading the files

#rnx

#Synced-CL 49956463

diff --git a/Engine/Source/Programs/UnrealBuildTool/Configuration/Xml/XmlConfig.cs b/Engine/Source/Programs/UnrealBuildTool/Configuration/Xml/XmlConfig.cs
--- a/Engine/Source/Programs/UnrealBuildTool/Configuration/Xml/XmlConfig.cs
+++ b/Engine/Source/Programs/UnrealBuildTool/Configuration/Xml/XmlConfig.cs
@@ -73,6 +73,9 @@ namespace UnrealBuildTool
 		/// <param name="logger">Logger for output</param>
 		public static void ReadConfigFiles(FileReference? overrideCacheFile, DirectoryReference? projectRootDirectory, ILogger logger)
 		{
+			string XmlConfigMutexName = GlobalSingleInstanceMutex.GetUniqueMutexForPath("UnrealBuildTool_Mutex_XmlConfig", Unreal.EngineDirectory);
+			using GlobalSingleInstanceMutex XmlConfigMutex = new GlobalSingleInstanceMutex(XmlConfigMutexName, true);
+
 			// Find all the configurable types
 			List<Type> configTypes = FindConfigurableTypes();
 
diff --git a/Engine/Source/Programs/UnrealBuildTool/UnrealBuildTool.cs b/Engine/Source/Programs/UnrealBuildTool/UnrealBuildTool.cs
--- a/Engine/Source/Programs/UnrealBuildTool/UnrealBuildTool.cs
+++ b/Engine/Source/Programs/UnrealBuildTool/UnrealBuildTool.cs
@@ -323,13 +323,9 @@ namespace UnrealBuildTool
 					using (Timeline.ScopeEvent("Apply XmlConfig"))
 					using (GlobalTracer.Instance.BuildSpan("XmlConfig.ReadConfigFiles()").StartActive())
 					{
-						string XmlConfigMutexName = GlobalSingleInstanceMutex.GetUniqueMutexForPath("UnrealBuildTool_Mutex_XmlConfig", FileReference.FromString(Assembly.GetExecutingAssembly().Location));
-						using (GlobalSingleInstanceMutex XmlConfigMutex = new GlobalSingleInstanceMutex(XmlConfigMutexName, true))
-						{
-							FileReference? XmlConfigCache = Arguments.GetFileReferenceOrDefault("-XmlConfigCache=", null);
-							Utils.TryParseProjectFileArgument(Arguments, Logger, out FileReference? ProjectFile, false);
-							XmlConfig.ReadConfigFiles(XmlConfigCache, ProjectFile?.Directory, Logger);
-						}
+						FileReference? XmlConfigCache = Arguments.GetFileReferenceOrDefault("-XmlConfigCache=", null);
+						Utils.TryParseProjectFileArgument(Arguments, Logger, out FileReference? ProjectFile, false);
+						XmlConfig.ReadConfigFiles(XmlConfigCache, ProjectFile?.Directory, Logger);
 					}
 
 					XmlConfig.ApplyTo(Options);
