[UBT]
* Changed so UbaExecutor connects to horde earlier

#Synced-CL 49953810

#Synced-CL 49953810

diff --git a/Engine/Source/Programs/UnrealBuildTool/Executors/UnrealBuildAccelerator/UBAAgentCoordinatorHorde.cs b/Engine/Source/Programs/UnrealBuildTool/Executors/UnrealBuildAccelerator/UBAAgentCoordinatorHorde.cs
--- a/Engine/Source/Programs/UnrealBuildTool/Executors/UnrealBuildAccelerator/UBAAgentCoordinatorHorde.cs
+++ b/Engine/Source/Programs/UnrealBuildTool/Executors/UnrealBuildAccelerator/UBAAgentCoordinatorHorde.cs
@@ -136,8 +136,14 @@ namespace UnrealBuildTool
 				throw new InvalidOperationException("Session has already been initialized");
 			}
 
-			_client = _serviceProvider.GetRequiredService<IHordeClient>().Compute;
-			_hordeHttpClient = _serviceProvider.GetRequiredService<IHordeClient>().CreateHttpClient();
+			IHordeClient hordeClient;
+			using (ITimelineEvent _ = Timeline.ScopeEvent($"HordeGetService"))
+			{
+				hordeClient = _serviceProvider.GetRequiredService<IHordeClient>();
+			}
+
+			_client = hordeClient.Compute;
+			_hordeHttpClient = hordeClient.CreateHttpClient();
 
 			{
 				using ITimelineEvent _ = Timeline.ScopeEvent($"HordeConnect");
@@ -149,7 +155,10 @@ namespace UnrealBuildTool
 			if (OperatingSystem.IsWindows())
 			{
 				agentFiles.Add("UbaAgent.exe");
-				agentFiles.Add("UbaWine.dll.so"); // Not needed but used by linux/wine helpers for more detailed logging or optimizations
+				if (FileReference.Exists(FileReference.Combine(ubaDir, "UbaWine.dll.so")))
+				{
+					agentFiles.Add("UbaWine.dll.so"); // Not needed but used by linux/wine helpers for more detailed logging or optimizations
+				}
 				if (useSentry)
 				{
 					agentFiles.Add("crashpad_handler.exe");
@@ -201,7 +210,8 @@ namespace UnrealBuildTool
 			{
 				return null;
 			}
-			return new UnrealBuildAcceleratorCacheConfig() { CacheServer = ubaConfig.CacheEndpoint, bRequireVfs = true, WriteCache = ubaConfig.WriteAccess.ToString(), Crypto = ubaConfig.CacheSessionKey, DesiredConnectionCount = desiredConnectionCount };
+			bool writeAccess = ubaConfig.WriteAccess;
+			return new UnrealBuildAcceleratorCacheConfig() { CacheServer = ubaConfig.CacheEndpoint, bRequireVfs = true, WriteCache = writeAccess.ToString(), Crypto = ubaConfig.CacheSessionKey, DesiredConnectionCount = desiredConnectionCount };
 		}
 
 		async Task<BlobLocator> CreateToolAsync(DirectoryReference baseDir, IEnumerable<FileReference> files, CancellationToken cancellationToken)
@@ -456,6 +466,8 @@ namespace UnrealBuildTool
 
 		public static async Task<(UBAHordeSession?, string)> TryCreateHordeSessionAsync(UnrealBuildAcceleratorHordeConfig hordeConfig, UnrealBuildAcceleratorConfig ubaConfig, string crypto, bool bStrictErrors, ILogger logger, CancellationToken cancellationToken = default)
 		{
+			using ITimelineEvent hordeCreateTimlineEvent = Timeline.ScopeEvent($"HordeCreate");
+
 			if (hordeConfig.bDisableHorde)
 			{
 				logger.LogInformation("Horde disabled via command line option.");
@@ -488,6 +500,7 @@ namespace UnrealBuildTool
 				bool allowWine = hordeConfig.bHordeAllowWine && OperatingSystem.IsWindows();
 
 				UBAHordeSession session = new(ubaConfig, crypto, server, token, hordeConfig.HordePool, allowWine, hordeConfig.HordeMaxCores, bStrictErrors, connectionMode, encryption, logger);
+				hordeCreateTimlineEvent.Dispose();
 				await session.InitAsync(useSentry: !String.IsNullOrEmpty(hordeConfig.UBASentryUrl), cancellationToken);
 				return (session, "Ready");
 			}
