UnrealBuildTool: Allow using toolchains from older versions of VisualStudio as they may still be preferred

#rnx
#Synced-CL 49956463

diff --git a/Engine/Source/Programs/UnrealBuildTool/Platform/Windows/MicrosoftPlatformSDK.cs b/Engine/Source/Programs/UnrealBuildTool/Platform/Windows/MicrosoftPlatformSDK.cs
--- a/Engine/Source/Programs/UnrealBuildTool/Platform/Windows/MicrosoftPlatformSDK.cs
+++ b/Engine/Source/Programs/UnrealBuildTool/Platform/Windows/MicrosoftPlatformSDK.cs
@@ -889,6 +889,12 @@ namespace UnrealBuildTool
 								FindVisualStudioToolChains(ExperimentalBaseDir, null, WindowsCompilerChannel.Experimental, ToolChains, IsAutoSdk: true, Logger);
 							}
 						}
+
+						// Additionally allow msvc toolchains from older versions of Visual Studio
+						if (Compiler == WindowsCompiler.VisualStudio2026)
+						{
+							ToolChains.AddRange(FindToolChainInstallations(WindowsCompiler.VisualStudio2022, Logger));
+						}
 					}
 					else
 					{
diff --git a/Engine/Source/Programs/UnrealBuildTool/Platform/Windows/UEBuildWindows.cs b/Engine/Source/Programs/UnrealBuildTool/Platform/Windows/UEBuildWindows.cs
--- a/Engine/Source/Programs/UnrealBuildTool/Platform/Windows/UEBuildWindows.cs
+++ b/Engine/Source/Programs/UnrealBuildTool/Platform/Windows/UEBuildWindows.cs
@@ -1140,6 +1140,8 @@ namespace UnrealBuildTool
 
 		internal VCEnvironment? Environment => Inner.Environment;
 
+		public string? CompilerDir => Inner.Environment?.CompilerDir.FullName ?? null;
+
 		public string? ToolChainDir => Inner.Environment?.ToolChainDir.FullName ?? null;
 
 		public string? ToolChainVersion => Inner.Environment?.ToolChainVersion.ToString() ?? null;
@@ -1538,7 +1540,14 @@ namespace UnrealBuildTool
 					{
 						throw new BuildLogEventException("{Compiler} compiler version {Version} is not a preferred version. Please use the latest preferred version {PreferredVersion}", WindowsPlatform.GetCompilerName(Target.WindowsPlatform.Compiler), Target.WindowsPlatform.Environment.CompilerVersion, preferred);
 					}
-					Logger.LogInformation("{Compiler} compiler version {Version} is not a preferred version. Please use the latest preferred version {PreferredVersion}", WindowsPlatform.GetCompilerName(Target.WindowsPlatform.Compiler), Target.WindowsPlatform.Environment.CompilerVersion, preferred);
+					else if (Target.WindowsPlatform.Environment.CompilerVersion < preferred)
+					{
+						Logger.LogInformation("{Compiler} compiler version {Version} is not a preferred version. Please use the latest preferred version {PreferredVersion}", WindowsPlatform.GetCompilerName(Target.WindowsPlatform.Compiler), Target.WindowsPlatform.Environment.CompilerVersion, preferred);
+					}
+					else
+					{
+						Logger.LogInformation("{Compiler} compiler version {Version} is newer than latest preferred version {PreferredVersion}. Please use caution as this compiler has not been heavily tested", WindowsPlatform.GetCompilerName(Target.WindowsPlatform.Compiler), Target.WindowsPlatform.Environment.CompilerVersion, preferred);
+					}
 				}
 
 				if (Target.WindowsPlatform.Compiler != Target.WindowsPlatform.ToolChain && !MicrosoftPlatformSDK.IsPreferredVersion(Target.WindowsPlatform.ToolChain, Target.WindowsPlatform.Environment.ToolChainVersion))
@@ -1550,7 +1559,14 @@ namespace UnrealBuildTool
 					{
 						throw new BuildLogEventException("{Toolchain} toolchain version {Version} is not a preferred version. Please use the latest preferred version {PreferredVersion}", WindowsPlatform.GetCompilerName(Target.WindowsPlatform.ToolChain), Target.WindowsPlatform.Environment.ToolChainVersion, preferred);
 					}
-					Logger.LogInformation("{Toolchain} toolchain version {Version} is not a preferred version. Please use a preferred toolchain such as {PreferredVersion}", WindowsPlatform.GetCompilerName(Target.WindowsPlatform.ToolChain), Target.WindowsPlatform.Environment.ToolChainVersion, preferred);
+					else if (Target.WindowsPlatform.Environment.ToolChainVersion < preferred)
+					{
+						Logger.LogInformation("{Toolchain} toolchain version {Version} is not a preferred version. Please use a preferred toolchain such as {PreferredVersion}", WindowsPlatform.GetCompilerName(Target.WindowsPlatform.ToolChain), Target.WindowsPlatform.Environment.ToolChainVersion, preferred);
+					}
+					else
+					{
+						Logger.LogInformation("{Toolchain} toolchain version {Version} is newer than latest preferred version {PreferredVersion}. Please use caution as this toolchain has not been heavily tested", WindowsPlatform.GetCompilerName(Target.WindowsPlatform.ToolChain), Target.WindowsPlatform.Environment.ToolChainVersion, preferred);
+					}
 				}
 			}
 		}
@@ -1659,7 +1675,18 @@ namespace UnrealBuildTool
 				}
 			}
 
-			// Second, default based on what's installed, test for 2022 first
+			// Next, default based on what Visual Studio is installed, test for 2022 first
+			IReadOnlyList<VisualStudioInstallation> vsInstalls = MicrosoftPlatformSDK.FindVisualStudioInstallations(Logger);
+			if (vsInstalls.Any(x => x.Compiler == WindowsCompiler.VisualStudio2022))
+			{
+				return WindowsCompiler.VisualStudio2022;
+			}
+			else if (vsInstalls.Any(x => x.Compiler == WindowsCompiler.VisualStudio2026))
+			{
+				return WindowsCompiler.VisualStudio2026;
+			}
+
+			// Finally, default based on what compiler is available, test for 2022 first
 			if (MicrosoftPlatformSDK.HasValidCompiler(WindowsCompiler.VisualStudio2022, Architecture, Logger))
 			{
 				return WindowsCompiler.VisualStudio2022;
