UnrealBuildTool: Pre-create a new ActionExecutor for every pass in build mode - this fixes an issue where targets with PreBuildTargets are executed but not processed or logged properly post-execution.

#Synced-CL 49949829

diff --git a/Engine/Source/Programs/UnrealBuildTool/Modes/BuildMode.cs b/Engine/Source/Programs/UnrealBuildTool/Modes/BuildMode.cs
--- a/Engine/Source/Programs/UnrealBuildTool/Modes/BuildMode.cs
+++ b/Engine/Source/Programs/UnrealBuildTool/Modes/BuildMode.cs
@@ -173,6 +173,7 @@ namespace UnrealBuildTool
 			CppDependencyCache cppDependencyCache = new();
 			FileHashCache fileHashCache = new(warnOnNoPartition, true);
 			List<Task> tasks = [];
+			ActionExecutor? executor = null;
 
 			// Parse and build the targets
 			try
@@ -204,7 +205,10 @@ namespace UnrealBuildTool
 					CleanMode.Clean(TargetDescriptors.Where(D => D.bRebuild).ToList(), BuildConfiguration, Logger);
 				}
 
-				using ActionExecutor? executor = ExecutorFactory.Instance.PrecreateExecutor(BuildConfiguration, TargetDescriptors, Logger);
+				// If we're using the UBA executor (very likely), we want to get a headstart on creating it, so it'll be ready when it needs it.
+				// Based on not-particularly-rigorous profiling, this saves maybe ~0.4s per execution of UBT by doing this here -
+				// enough to justify this slightly ugly pre-creation logic.
+				executor = ExecutorFactory.Instance.PrecreateExecutor(BuildConfiguration, TargetDescriptors, Logger);
 
 				// Handle remote builds
 				for (int Idx = 0; Idx < TargetDescriptors.Count; ++Idx)
@@ -269,15 +273,26 @@ namespace UnrealBuildTool
 							throw new BuildException("FileHashCache had errors");
 						}
 
+						bool usedExecutor = false;
 						foreach (List<(TargetDescriptor, TargetMakefile)> pass in targetsPasses)
 						{
+							// We need a fresh executor for each pass, as each one is only expected to have its ExecuteAsync method called once.
+							if (usedExecutor)
+							{
+								executor?.Dispose();
+								executor = ExecutorFactory.Instance.PrecreateExecutor(BuildConfiguration, TargetDescriptors, Logger);
+							}
+
 							await BuildAsync(pass, BuildConfiguration, Options, WriteOutdatedActionsFile, cppDependencyCache, Logger, ActionTypeFilter, executor);
+							usedExecutor = true;
 						}
 					}
 				}
 			}
 			finally
 			{
+				executor?.Dispose();
+
 				// Check if anything failed during our run, and act accordingly.
 				ProcessCoreDumps(SaveCrashDumpDirectory, Logger);
 
