UnrealBuildTool: Generate the Automation solution for VS2026 if it's found on the system, even if project file generation is requested for VS2022.

#Synced-CL 50040876

diff --git a/Engine/Source/Programs/UnrealBuildTool/ProjectFiles/VisualStudio/VCProjectFileGenerator.cs b/Engine/Source/Programs/UnrealBuildTool/ProjectFiles/VisualStudio/VCProjectFileGenerator.cs
--- a/Engine/Source/Programs/UnrealBuildTool/ProjectFiles/VisualStudio/VCProjectFileGenerator.cs
+++ b/Engine/Source/Programs/UnrealBuildTool/ProjectFiles/VisualStudio/VCProjectFileGenerator.cs
@@ -446,21 +446,33 @@ namespace UnrealBuildTool
 			if (Settings.ProjectFileFormat == VCProjectFileFormat.VisualStudio2022)
 			{
 				Logger.LogInformation(
-					"Visual Studio 2022 does not support .NET 10.0 C# projects, these projects will not be added to the generated automation solution as the current dotnet runtime version is '{Version}'. Compiling with an older dotnet SDK from Visual Studio can cause errors.",
+					"Visual Studio 2022 does not support .NET 10.0 C# projects, and the current dotnet runtime version is '{Version}'. Compiling with an older dotnet SDK from Visual Studio can cause errors.",
 					System.Runtime.InteropServices.RuntimeInformation.FrameworkDescription);
-				Logger.LogInformation("Please generate the Visual Studio 2026 solution if .NET 10.0 C# project support is required.");
-				return;
+
+				// Some users can't use VS2026 yet (certain platforms are VS2022 only) so we should accomodate for them by generating the Automation solution for 2026 if available.
+				if (WindowsPlatform.HasCompiler(WindowsCompiler.VisualStudio2026, UnrealArch.X64, Logger))
+				{
+					Logger.LogInformation("The {SolutionName} solution will be generated for Visual Studio 2026.", $"Automation_{PrimaryProjectName}");
+				}
+				else
+				{
+					Logger.LogInformation("The {SolutionName} solution will not be generated.", $"Automation_{PrimaryProjectName}");
+					return;
+				}
 			}
 
 			// We expect this to be called as part of a call that also writes a C++ solution too, so we have to make sure to put
 			// this property back to its original value when done.
 			string OriginalPrimaryProjectName = PrimaryProjectName;
+			VCProjectFileFormat OriginalProjectFileFormat = Settings.ProjectFileFormat;
 			PrimaryProjectFolder OriginalRootFolder = RootFolder;
 			List<ProjectFile> OriginalAutomationProjectFiles = [.. AutomationProjectFiles];
 			List<ProjectFile> OriginalOtherProjectFiles = [.. OtherProjectFiles];
+
 			try
 			{
 				PrimaryProjectName = $"Automation_{PrimaryProjectName}";
+				Settings.ProjectFileFormat = VCProjectFileFormat.VisualStudio2026;
 				RootFolder = new(this, "<Root>");
 				AutomationProjectFiles.Clear();
 				OtherProjectFiles.Clear();
@@ -475,6 +487,7 @@ namespace UnrealBuildTool
 			finally
 			{
 				PrimaryProjectName = OriginalPrimaryProjectName;
+				Settings.ProjectFileFormat = OriginalProjectFileFormat;
 				RootFolder = OriginalRootFolder;
 				AutomationProjectFiles.Clear();
 				AutomationProjectFiles.AddRange(OriginalAutomationProjectFiles);
