UnrealBuildTool: clean restricted and platform directories as part of the Clean mode. Add a dry-run flag to the Clean mode.

#Synced-CL 50040876

diff --git a/Engine/Source/Programs/UnrealBuildTool/Modes/CleanMode.cs b/Engine/Source/Programs/UnrealBuildTool/Modes/CleanMode.cs
--- a/Engine/Source/Programs/UnrealBuildTool/Modes/CleanMode.cs
+++ b/Engine/Source/Programs/UnrealBuildTool/Modes/CleanMode.cs
@@ -10,6 +10,7 @@ using Microsoft.Extensions.Logging;
 using OpenTracing;
 using OpenTracing.Util;
 using UnrealBuildBase;
+using static UnrealBuildTool.DataDrivenPlatformInfo;
 
 namespace UnrealBuildTool
 {
@@ -34,6 +35,12 @@ namespace UnrealBuildTool
 		[CommandLine("-SkipPreBuildTargets")]
 		public bool bSkipPreBuildTargets = false;
 
+		/// <summary>
+		/// If true, lists entries to be cleaned but does not clean them.
+		/// </summary>
+		[CommandLine("-DryRun")]
+		public bool bDryRun { get; set; } = false;
+
 		/// <summary>
 		/// Main entry point
 		/// </summary>
@@ -112,20 +119,41 @@ namespace UnrealBuildTool
 						}
 					}
 
+					RestrictedFolder[] RestrictedFolders =
+					[
+						RestrictedFolder.EpicInternal,
+						RestrictedFolder.CarefullyRedist,
+						RestrictedFolder.LimitedAccess,
+						RestrictedFolder.NotForLicensees,
+						RestrictedFolder.NoRedist,
+					];
+
 					// Find the base folders that can contain binaries
-					List<DirectoryReference> BaseDirs = new List<DirectoryReference>();
-					BaseDirs.Add(Unreal.EngineDirectory);
+					HashSet<DirectoryReference> BaseDirs = [Unreal.EngineDirectory];
+
+					foreach (RestrictedFolder Folder in RestrictedFolders)
+					{
+						BaseDirs.Add(DirectoryReference.Combine(Unreal.EngineDirectory, "Restricted", Folder.ToString()));
+					}
+
+					foreach (string Platform in GetAllPlatformInfos().Keys)
+					{
+						BaseDirs.Add(DirectoryReference.Combine(Unreal.EngineDirectory, "Platforms", Platform));
+					}
+
 					foreach (FileReference Plugin in PluginsBase.EnumeratePlugins(Target.ProjectFile))
 					{
 						BaseDirs.Add(Plugin.Directory);
 					}
+
 					if (Target.ProjectFile != null)
 					{
 						BaseDirs.Add(Target.ProjectFile.Directory);
 					}
 
 					// If we're running a precompiled build, remove anything under the engine folder
-					BaseDirs.RemoveAll(x => RulesAssembly.IsReadOnly(x));
+					List<DirectoryReference> ReadOnlyDirectories = BaseDirs.Where(x => RulesAssembly.IsReadOnly(x)).ToList();
+					BaseDirs.ExceptWith(ReadOnlyDirectories);
 
 					// Get all the names which can prefix build products
 					List<string> NamePrefixes = new List<string>();
@@ -236,6 +264,8 @@ namespace UnrealBuildTool
 				});
 			}
 
+			LogLevel DeleteLogLevel = bDryRun ? LogLevel.Information : LogLevel.Debug;
+
 			Task DeleteDirectories = Task.Run(() =>
 			{
 				using ScopedTimer Timer = new ScopedTimer($"Delete {DirectoriesToDelete.Count} directories", Logger);
@@ -245,7 +275,12 @@ namespace UnrealBuildTool
 
 					if (DirectoryReference.Exists(DirectoryToDelete))
 					{
-						Logger.LogDebug("    Deleting {DirectoryToDelete}...", $"{DirectoryToDelete}{Path.DirectorySeparatorChar}");
+						Logger.Log(DeleteLogLevel, "    Deleting {DirectoryToDelete}...", $"{DirectoryToDelete}{Path.DirectorySeparatorChar}");
+
+						if (bDryRun)
+						{
+							return;
+						}
 
 						try
 						{
@@ -266,7 +301,12 @@ namespace UnrealBuildTool
 				{
 					if (FileReference.Exists(FileToDelete))
 					{
-						Logger.LogDebug("    Deleting {FileToDelete}...", FileToDelete);
+						Logger.Log(DeleteLogLevel, "    Deleting {FileToDelete}...", FileToDelete);
+
+						if (bDryRun)
+						{
+							return;
+						}
 
 						try
 						{
@@ -283,6 +323,11 @@ namespace UnrealBuildTool
 			DeleteDirectories.Wait();
 			DeleteFiles.Wait();
 
+			if (bDryRun)
+			{
+				return;
+			}
+
 			// Also clean all the remote targets
 			for (int Idx = 0; Idx < TargetDescriptors.Count; Idx++)
 			{
