UnrealBuildTool: Allow specifying the toolchain when getting VC include paths to allow for Clang support.

#Synced-CL 50040876

diff --git a/Engine/Source/Programs/UnrealBuildTool/Platform/Windows/VCToolChain.cs b/Engine/Source/Programs/UnrealBuildTool/Platform/Windows/VCToolChain.cs
--- a/Engine/Source/Programs/UnrealBuildTool/Platform/Windows/VCToolChain.cs
+++ b/Engine/Source/Programs/UnrealBuildTool/Platform/Windows/VCToolChain.cs
@@ -3932,10 +3932,10 @@ namespace UnrealBuildTool
 		/// Gets the default include paths for the given platform.
 		/// </summary>
 		[SupportedOSPlatform("windows")]
-		public static string GetVCIncludePaths(UnrealTargetPlatform Platform, WindowsCompiler Compiler, string? CompilerVersion, string? ToolchainVersion, ILogger Logger)
+		public static string GetVCIncludePaths(UnrealTargetPlatform Platform, WindowsCompiler Compiler, WindowsCompiler ToolChain, string? CompilerVersion, string? ToolchainVersion, ILogger Logger)
 		{
 			// Make sure we've got the environment variables set up for this target
-			VCEnvironment EnvVars = VCEnvironment.Create(Compiler, WindowsCompiler.Default, Platform, UnrealArch.X64, CompilerVersion, ToolchainVersion, null, null, false, false, false, Logger);
+			VCEnvironment EnvVars = VCEnvironment.Create(Compiler, ToolChain, Platform, UnrealArch.X64, CompilerVersion, ToolchainVersion, null, null, false, false, false, Logger);
 
 			// Also add any include paths from the INCLUDE environment variable.  MSVC is not necessarily running with an environment that
 			// matches what UBT extracted from the vcvars*.bat using SetEnvironmentVariablesFromBatchFile().  We'll use the variables we
diff --git a/Engine/Source/Programs/UnrealBuildTool/ProjectFiles/VisualStudio/VCProjectFileGenerator.cs b/Engine/Source/Programs/UnrealBuildTool/ProjectFiles/VisualStudio/VCProjectFileGenerator.cs
--- a/Engine/Source/Programs/UnrealBuildTool/ProjectFiles/VisualStudio/VCProjectFileGenerator.cs
+++ b/Engine/Source/Programs/UnrealBuildTool/ProjectFiles/VisualStudio/VCProjectFileGenerator.cs
@@ -678,7 +678,7 @@ namespace UnrealBuildTool
 				{
 					if (SupportedPlatforms.Contains(UnrealTargetPlatform.Win64))
 					{
-						VCCommonTargetFileContent.AppendLine("    <DefaultSystemIncludePaths>{0}</DefaultSystemIncludePaths>", VCToolChain.GetVCIncludePaths(UnrealTargetPlatform.Win64, GetCompilerForIntellisense(Settings.ProjectFileFormat), null, null, Logger));
+						VCCommonTargetFileContent.AppendLine("    <DefaultSystemIncludePaths>{0}</DefaultSystemIncludePaths>", VCToolChain.GetVCIncludePaths(UnrealTargetPlatform.Win64, GetCompilerForIntellisense(Settings.ProjectFileFormat), WindowsCompiler.Default, null, null, Logger));
 					}
 				}
 				else
diff --git a/Engine/Source/Programs/UnrealBuildTool/ProjectFiles/VisualStudioWorkspace/VSWorkspaceProjectFileGenerator.cs b/Engine/Source/Programs/UnrealBuildTool/ProjectFiles/VisualStudioWorkspace/VSWorkspaceProjectFileGenerator.cs
--- a/Engine/Source/Programs/UnrealBuildTool/ProjectFiles/VisualStudioWorkspace/VSWorkspaceProjectFileGenerator.cs
+++ b/Engine/Source/Programs/UnrealBuildTool/ProjectFiles/VisualStudioWorkspace/VSWorkspaceProjectFileGenerator.cs
@@ -695,7 +695,7 @@ namespace UnrealBuildTool
 					{
 						if (SupportedPlatforms.Contains(UnrealTargetPlatform.Win64))
 						{
-							VCCommonTargetFileContent.AppendLine("    <DefaultSystemIncludePaths>{0}</DefaultSystemIncludePaths>", VCToolChain.GetVCIncludePaths(UnrealTargetPlatform.Win64, VCProjectFileGenerator.GetCompilerForIntellisense(Format), null, null, Logger));
+							VCCommonTargetFileContent.AppendLine("    <DefaultSystemIncludePaths>{0}</DefaultSystemIncludePaths>", VCToolChain.GetVCIncludePaths(UnrealTargetPlatform.Win64, VCProjectFileGenerator.GetCompilerForIntellisense(Format), WindowsCompiler.Default, null, null, Logger));
 						}
 					}
 					else
