Adding storage in ConfigHierarchy to track ordering of lines in ini files. Required for certain tooling operations.

#Synced-CL 50025708

diff --git a/Engine/Source/Programs/UnrealBuildTool/Configuration/Ini/ConfigHierarchy.cs b/Engine/Source/Programs/UnrealBuildTool/Configuration/Ini/ConfigHierarchy.cs
--- a/Engine/Source/Programs/UnrealBuildTool/Configuration/Ini/ConfigHierarchy.cs
+++ b/Engine/Source/Programs/UnrealBuildTool/Configuration/Ini/ConfigHierarchy.cs
@@ -73,6 +73,17 @@ namespace UnrealBuildTool
 		/// </summary>
 		readonly Dictionary<string, List<string>> KeyToValue = new Dictionary<string, List<string>>(StringComparer.InvariantCultureIgnoreCase);
 
+		/// <summary>
+		/// List of key names in the order they were first added (preserves INI file order)
+		/// </summary>
+		readonly List<string> KeyOrder = new List<string>();
+
+		/// <summary>
+		/// List of (key, value) pairs in the exact order they were added from INI file
+		/// Preserves interleaved order of array elements with different keys (e.g., +Steps= and +Steps[Variant]=)
+		/// </summary>
+		readonly List<(string Key, string Value)> ArrayElementOrder = new List<(string, string)>();
+
 		/// <summary>
 		/// Construct a merged config section from the given per-file config sections
 		/// </summary>
@@ -86,6 +97,8 @@ namespace UnrealBuildTool
 					if (Line.Action == ConfigLineAction.RemoveKey)
 					{
 						KeyToValue.Remove(Line.Key);
+						KeyOrder.RemoveAll(k => k.Equals(Line.Key, StringComparison.OrdinalIgnoreCase));
+						ArrayElementOrder.RemoveAll(pair => pair.Key.Equals(Line.Key, StringComparison.OrdinalIgnoreCase));
 						continue;
 					}
 
@@ -99,14 +112,20 @@ namespace UnrealBuildTool
 						{
 							Values.Clear();
 							Values.Add(Line.Value);
+							// Track order: when Set clears the array, also remove all previous entries for this key
+							ArrayElementOrder.RemoveAll(pair => pair.Key.Equals(Line.Key, StringComparison.OrdinalIgnoreCase));
+							ArrayElementOrder.Add((Line.Key, Line.Value));
 						}
 						else if (Line.Action == ConfigLineAction.Add)
 						{
 							Values.Add(Line.Value);
+							ArrayElementOrder.Add((Line.Key, Line.Value));
 						}
 						else if (Line.Action == ConfigLineAction.RemoveKeyValue)
 						{
 							Values.RemoveAll(x => x.Equals(Line.Value, StringComparison.OrdinalIgnoreCase));
+							// Also remove from order tracking
+							ArrayElementOrder.RemoveAll(pair => pair.Key.Equals(Line.Key, StringComparison.OrdinalIgnoreCase) && pair.Value.Equals(Line.Value, StringComparison.OrdinalIgnoreCase));
 						}
 					}
 					else
@@ -117,6 +136,8 @@ namespace UnrealBuildTool
 							Values = new List<string>();
 							Values.Add(Line.Value);
 							KeyToValue.Add(Line.Key, Values);
+							KeyOrder.Add(Line.Key);
+							ArrayElementOrder.Add((Line.Key, Line.Value));
 						}
 					}
 				}
@@ -126,7 +147,7 @@ namespace UnrealBuildTool
 		/// <summary>
 		/// Returns a list of key names
 		/// </summary>
-		public IEnumerable<string> KeyNames => KeyToValue.Keys;
+		public IEnumerable<string> KeyNames => KeyOrder;
 
 		/// <summary>
 		/// Tries to find the value for a given key
@@ -169,6 +190,15 @@ namespace UnrealBuildTool
 				return false;
 			}
 		}
+
+		/// <summary>
+		/// Returns array elements in the exact order they were added from INI files.
+		/// Preserves interleaving of elements with different keys (e.g., +Steps= and +Steps[Variant]=)
+		/// </summary>
+		public IReadOnlyList<(string Key, string Value)> GetArrayElementsInOrder()
+		{
+			return ArrayElementOrder;
+		}
 	}
 
 	/// <summary>
