Linux HWPGO support

#Synced-CL 50060205

diff --git a/Engine/Source/Programs/UnrealBuildTool/Platform/Clang/ClangToolChain.Link.cs b/Engine/Source/Programs/UnrealBuildTool/Platform/Clang/ClangToolChain.Link.cs
--- a/Engine/Source/Programs/UnrealBuildTool/Platform/Clang/ClangToolChain.Link.cs
+++ b/Engine/Source/Programs/UnrealBuildTool/Platform/Clang/ClangToolChain.Link.cs
@@ -13,6 +13,7 @@ namespace UnrealBuildTool
 		{
 			bool bLTOEnabled = Options.HasFlag(ClangToolChainOptions.EnableLinkTimeOptimization);
 			bool bThinLTOEnabled = Options.HasFlag(ClangToolChainOptions.EnableThinLTO);
+			bool bSPGOEnabled = Options.HasFlag(ClangToolChainOptions.SampleBasedPGO);
 
 			// Profile Guided Optimization (PGO)
 			// Disable PGOOptimize when no profdata available
@@ -39,7 +40,15 @@ namespace UnrealBuildTool
 			if (linkEnvironment.bPGOOptimize && linkEnvironment.PGODirectory != null && linkEnvironment.PGOFilenamePrefix != null)
 			{
 				DirectoryReference PGODir = new(linkEnvironment.PGODirectory);
-				arguments.Add($"-fprofile-use=\"{NormalizeCommandLinePath(DirectoryReference.Combine(PGODir, linkEnvironment.PGOFilenamePrefix), linkEnvironment.RootPaths)}\"");
+
+				if (bSPGOEnabled)
+				{
+					arguments.Add($"-fprofile-sample-use=\"{NormalizeCommandLinePath(DirectoryReference.Combine(PGODir, linkEnvironment.PGOFilenamePrefix), linkEnvironment.RootPaths)}\"");
+				}
+				else
+				{
+					arguments.Add($"-fprofile-use=\"{NormalizeCommandLinePath(DirectoryReference.Combine(PGODir, linkEnvironment.PGOFilenamePrefix), linkEnvironment.RootPaths)}\"");
+				}
 
 				//
 				// Clang emits a warning for each compiled function that doesn't have a matching entry in the profile data.
diff --git a/Engine/Source/Programs/UnrealBuildTool/Platform/Clang/ClangToolChain.cs b/Engine/Source/Programs/UnrealBuildTool/Platform/Clang/ClangToolChain.cs
--- a/Engine/Source/Programs/UnrealBuildTool/Platform/Clang/ClangToolChain.cs
+++ b/Engine/Source/Programs/UnrealBuildTool/Platform/Clang/ClangToolChain.cs
@@ -140,6 +140,11 @@ namespace UnrealBuildTool
 		/// and will strip them in a later step
 		/// </summary>
 		DisableStripSymbols = 1 << 22,
+
+		/// <summary>
+		/// If specified along with -PGOProfile or -PGOOptimize, use sample-based PGO instead of instrumented.
+		/// </summary>
+		SampleBasedPGO = 1 << 23,
 	}
 
 	abstract partial class ClangToolChain : ISPCToolChain
@@ -744,6 +749,7 @@ namespace UnrealBuildTool
 		protected virtual void GetCompileArguments_Optimizations(CppCompileEnvironment CompileEnvironment, List<string> Arguments)
 		{
 			bool bLTOEnabled = Options.HasFlag(ClangToolChainOptions.EnableLinkTimeOptimization);
+			bool bSPGOEnabled = Options.HasFlag(ClangToolChainOptions.SampleBasedPGO);
 
 			// Link Time Optimization (LTO)
 			if (bLTOEnabled)
@@ -788,14 +794,35 @@ namespace UnrealBuildTool
 
 			if (CompileEnvironment.bPGOOptimize)
 			{
-				Log.TraceInformationOnce("Enabling Profile Guided Optimization (PGO). Linking will take a while.");
-				DirectoryReference? PGODir = DirectoryReference.FromString(CompileEnvironment.PGODirectory!);
-				Arguments.Add($"-fprofile-use=\"{NormalizeCommandLinePath(DirectoryReference.Combine(PGODir!, CompileEnvironment.PGOFilenamePrefix!), CompileEnvironment.RootPaths)}\"");
+				if (bSPGOEnabled)
+				{
+					Log.TraceInformationOnce("Enabling Sampling Based Profile Guided Optimization (SPGO). Linking will take a while.");
+					DirectoryReference? PGODir = DirectoryReference.FromString(CompileEnvironment.PGODirectory!);
+					Arguments.Add($"-fprofile-sample-use=\"{NormalizeCommandLinePath(DirectoryReference.Combine(PGODir!, CompileEnvironment.PGOFilenamePrefix!), CompileEnvironment.RootPaths)}\"");
+				}
+				else
+				{
+					Log.TraceInformationOnce("Enabling Profile Guided Optimization (PGO). Linking will take a while.");
+					DirectoryReference? PGODir = DirectoryReference.FromString(CompileEnvironment.PGODirectory!);
+					Arguments.Add($"-fprofile-use=\"{NormalizeCommandLinePath(DirectoryReference.Combine(PGODir!, CompileEnvironment.PGOFilenamePrefix!), CompileEnvironment.RootPaths)}\"");
+				}
 			}
 			else if (CompileEnvironment.bPGOProfile)
 			{
-				Log.TraceInformationOnce("Enabling Profile Guided Instrumentation (PGI). Linking will take a while.");
-				Arguments.Add("-fprofile-generate");
+				if (bSPGOEnabled)
+				{
+					Log.TraceInformationOnce("Enabling Sample Based Profile Guided Instrumentation (SPGI). Linking will take a while.");
+					Arguments.Add("-gdwarf");
+					Arguments.Add("-gsplit-dwarf");
+					Arguments.Add("-gline-tables-only");
+					Arguments.Add("-fdebug-info-for-profiling");
+					Arguments.Add("-funique-internal-linkage-names");
+				}
+				else
+				{
+					Log.TraceInformationOnce("Enabling Profile Guided Instrumentation (PGI). Linking will take a while.");
+					Arguments.Add("-fprofile-generate");
+				}
 			}
 
 			if (!CompileEnvironment.bUseInlining)
diff --git a/Engine/Source/Programs/UnrealBuildTool/Platform/Linux/UEBuildLinux.cs b/Engine/Source/Programs/UnrealBuildTool/Platform/Linux/UEBuildLinux.cs
--- a/Engine/Source/Programs/UnrealBuildTool/Platform/Linux/UEBuildLinux.cs
+++ b/Engine/Source/Programs/UnrealBuildTool/Platform/Linux/UEBuildLinux.cs
@@ -137,7 +137,14 @@ namespace UnrealBuildTool
 		[CommandLine("-NoStripSymbols")]
 		[XmlConfigFile(Category = "BuildConfiguration", Name = "bDisableStripSymbols")]
 		public bool bDisableStripSymbols = false;
-		
+
+		/// <summary>
+		/// If specified along with -PGOProfile, use sample-based PGO instead of instrumented.
+		/// </summary>
+		[XmlConfigFile(Category = "BuildConfiguration")]
+		[CommandLine("-LinuxSampleBasedPGO")]
+		public bool bLinuxSampleBasedPGO = false;
+
 		/// <summary>
 		/// Enables runtime ray tracing support.
 		/// </summary>
@@ -188,6 +195,8 @@ namespace UnrealBuildTool
 		public LinuxRayTracingRuntimeMode RayTracingMode => Inner.RayTracingMode;
 
 		public bool bDisableStripSymbols => Inner.bDisableStripSymbols;
+
+		public bool bLinuxSampleBasedPGO => Inner.bLinuxSampleBasedPGO;
 #pragma warning restore CS1591
 		#endregion
 	}
@@ -720,6 +729,11 @@ namespace UnrealBuildTool
 				}
 			}
 
+			if (Target.LinuxPlatform.bLinuxSampleBasedPGO)
+			{
+				Options |= ClangToolChainOptions.SampleBasedPGO;
+			}
+
 			if (Target.bUseAutoRTFMCompiler)
 			{
 				Options |= ClangToolChainOptions.UseAutoRTFMCompiler;
