[UBT] Extending cross platform / cross architecture uba support

#Synced-CL 50037559

diff --git a/Engine/Config/Apple/AppleOpenSource_SDK.json b/Engine/Config/Apple/AppleOpenSource_SDK.json
new file mode 100644
--- /dev/null
+++ b/Engine/Config/Apple/AppleOpenSource_SDK.json
@@ -0,0 +1,25 @@
+{
+	"//1": "Configuration for Apple open source toolchain (clang + OSS headers)",
+	"//2": "Used for distributed compilation of non-SDK-dependent Apple platform files",
+
+	"//3": "Version directory within UE_SDKS_ROOT/Host{Platform}/Apple/",
+	"//4": "If MainVersion is empty, auto-selects based on XcodeToToolchainVersion mapping",
+	"MainVersion": "",
+
+	"//5": "Subdirectory paths relative to the version directory",
+	"//6": "Same path on all hosts enables simple UBA Host{Platform} substitution",
+	"ClangSubdir": "Clang/bin/clang++",
+	"OSSHeadersSubdir": "AppleOSS",
+
+	"//7": "Maps Xcode version to toolchain directory prefix (vN)",
+	"//8": "Format: 'FirstXcodeVersion-vN'. Selects highest matching version.",
+	"//9": "Xcode    | Apple Clang | LLVM    | Prefix | Auto SDK Directory",
+	"//10": "16.0.0+  | 16.0.0      | 17.0.6  | v1     | v1_clang-17.0.6_oss-1350",
+	"//11": "16.3.0+  | 17.0.0      | 19.1.4  | v2     | v2_clang-19.1.4_oss-1350",
+	"//12": "26.0.0+  | 17.0.6      | 19.1.5  | v3     | v3_clang-19.1.6_oss-1350",
+	"XcodeToToolchainVersion": [
+		"16.0.0-v1",
+		"16.3.0-v2",
+		"26.0.0-v3"
+	]
+}
diff --git a/Engine/Platforms/VisionOS/Source/Programs/UnrealBuildTool/VisionOSToolChain.cs b/Engine/Platforms/VisionOS/Source/Programs/UnrealBuildTool/VisionOSToolChain.cs
--- a/Engine/Platforms/VisionOS/Source/Programs/UnrealBuildTool/VisionOSToolChain.cs
+++ b/Engine/Platforms/VisionOS/Source/Programs/UnrealBuildTool/VisionOSToolChain.cs
@@ -6,6 +6,8 @@ namespace UnrealBuildTool
 {
 	class VisionOSToolChain : IOSToolChain
 	{
+		protected override string XnuPlatformDefine => "XNU_PLATFORM_XROS";
+
 		public VisionOSToolChain(ReadOnlyTargetRules InTarget, VisionOSProjectSettings InProjectSettings, ILogger InLogger)
 			: base(InTarget, InProjectSettings, () => new VisionOSToolChainSettings(InProjectSettings, InLogger), ClangToolChainOptions.None, InLogger)
 		{
diff --git a/Engine/Source/Programs/UnrealBuildTool/Configuration/Rules/ModuleRules.cs b/Engine/Source/Programs/UnrealBuildTool/Configuration/Rules/ModuleRules.cs
--- a/Engine/Source/Programs/UnrealBuildTool/Configuration/Rules/ModuleRules.cs
+++ b/Engine/Source/Programs/UnrealBuildTool/Configuration/Rules/ModuleRules.cs
@@ -1154,8 +1154,13 @@ namespace UnrealBuildTool
 		/// </summary>
 		public bool bAllowConfidentialPlatformDefines { get; set; }
 
-        /// <summary>
-        /// Disables AutoRTFM instrumentation in this module. Only matters when the AutoRTFM compiler is being used.
+		/// <summary>
+		/// Experimental: Whether this module requires platform specific SDK files. Not supported by all platforms yet
+		/// </summary>
+		public bool bRequiresPlatformSDK { get; set; }
+
+		/// <summary>
+		/// Disables AutoRTFM instrumentation in this module. Only matters when the AutoRTFM compiler is being used.
 		/// The default value of this depends on whether any `TargetType` can be deduced for the module.
 		/// It will default to `true` if:
 		/// - The module has `SupportedTargetTypes` specified and `TargetType.Server` is not one of the target types.
@@ -1166,8 +1171,8 @@ namespace UnrealBuildTool
 		/// - The module is part of a plugin where the module has `"TargetDenyList"` specified and `"Server"` is one
 		///   of the denied target types.
 		/// Otherwise it will default to `false`.
-        /// </summary>
-        public bool bDisableAutoRTFMInstrumentation
+		/// </summary>
+		public bool bDisableAutoRTFMInstrumentation
 		{
 			get
 			{
diff --git a/Engine/Source/Programs/UnrealBuildTool/Configuration/UEBuildModuleCPP.cs b/Engine/Source/Programs/UnrealBuildTool/Configuration/UEBuildModuleCPP.cs
--- a/Engine/Source/Programs/UnrealBuildTool/Configuration/UEBuildModuleCPP.cs
+++ b/Engine/Source/Programs/UnrealBuildTool/Configuration/UEBuildModuleCPP.cs
@@ -1773,6 +1773,11 @@ namespace UnrealBuildTool
 				return false;
 			}
 
+			if (ModuleCompileEnvironment.bRequiresPlatformSDK != CompileEnvironment.bRequiresPlatformSDK)
+			{
+				return false;
+			}
+
 			if (ModuleCompileEnvironment.bOptimizeCode != CompileEnvironment.bOptimizeCode)
 			{
 				return false;
@@ -1841,6 +1846,17 @@ namespace UnrealBuildTool
 					Variant += ".Project";
 				}
 			}
+			if (CompileEnvironment.bRequiresPlatformSDK != BaseCompileEnvironment.bRequiresPlatformSDK)
+			{
+				if (CompileEnvironment.bRequiresPlatformSDK)
+				{
+					Variant += ".Plat";
+				}
+				else
+				{
+					Variant += ".NonPlat";
+				}
+			}
 			if (CompileEnvironment.bOptimizeCode != BaseCompileEnvironment.bOptimizeCode)
 			{
 				if (CompileEnvironment.bOptimizeCode)
@@ -1976,6 +1992,7 @@ namespace UnrealBuildTool
 			CompileEnvironment.bValidateFormatStrings = ModuleCompileEnvironment.bValidateFormatStrings;
 			CompileEnvironment.bValidateInternalApi = ModuleCompileEnvironment.bValidateInternalApi;
 			CompileEnvironment.bValidateExperimentalApi = ModuleCompileEnvironment.bValidateExperimentalApi;
+			CompileEnvironment.bRequiresPlatformSDK = ModuleCompileEnvironment.bRequiresPlatformSDK;
 			CompileEnvironment.bUseUnity = ModuleCompileEnvironment.bUseUnity; // Used by cache bucket. We want shared pch in same bucket
 		}
 
@@ -2618,6 +2635,11 @@ namespace UnrealBuildTool
 			Result.bUseAutoRTFMCompiler = Target.bUseAutoRTFMCompiler;
 			Result.FPassPlugins = Rules.FPassPlugins;
 
+			if (Target.Platform.IsInGroup(UnrealPlatformGroup.Apple))
+			{
+				Result.bRequiresPlatformSDK = Rules.bRequiresPlatformSDK;
+			}
+
 			CompileEnvironmentDebugInfoSettings(Target, Result);
 
 			// Only copy the AutoRTFM flags if we are using the AutoRTFM compiler
diff --git a/Engine/Source/Programs/UnrealBuildTool/Executors/UnrealBuildAccelerator/UBAAgentCoordinatorHorde.cs b/Engine/Source/Programs/UnrealBuildTool/Executors/UnrealBuildAccelerator/UBAAgentCoordinatorHorde.cs
--- a/Engine/Source/Programs/UnrealBuildTool/Executors/UnrealBuildAccelerator/UBAAgentCoordinatorHorde.cs
+++ b/Engine/Source/Programs/UnrealBuildTool/Executors/UnrealBuildAccelerator/UBAAgentCoordinatorHorde.cs
@@ -58,6 +58,8 @@ namespace UnrealBuildTool
 		IComputeClient? _client;
 		HordeHttpClient? _hordeHttpClient;
 
+		OSPlatform _targetWorkerPlatform;
+
 		public class Worker
 		{
 			public Stopwatch? StartTime { get; init; }
@@ -129,7 +131,7 @@ namespace UnrealBuildTool
 			_cancellationTokenSource.Dispose();
 		}
 
-		public async Task InitAsync(bool useSentry, CancellationToken cancellationToken)
+		public async Task InitAsync(bool useSentry, string? hordeCondition, CancellationToken cancellationToken)
 		{
 			if (_client != null)
 			{
@@ -150,9 +152,26 @@ namespace UnrealBuildTool
 				GetServerInfoResponse serverInfo = await _hordeHttpClient.GetServerInfoAsync(cancellationToken);
 				_logger.LogInformation("Horde server: {ServerVersion}, agent: {AgentVersion}", serverInfo.ServerVersion, serverInfo.AgentVersion);
 			}
-			DirectoryReference ubaDir = UBAExecutor.UbaBinariesDir;
+
+			_targetWorkerPlatform = GetTargetWorkerPlatform(hordeCondition);
+			bool bCrossPlatformCompilation = _targetWorkerPlatform != GetHostPlatform();
+
+			DirectoryReference ubaDir;
+			if (bCrossPlatformCompilation && _targetWorkerPlatform == OSPlatform.Linux)
+			{
+				ubaDir = DirectoryReference.Combine(Unreal.EngineDirectory, "Binaries", "Linux", "UnrealBuildAccelerator");
+			}
+			else
+			{
+				ubaDir = UBAExecutor.UbaBinariesDir;
+			}
+
 			List<string> agentFiles = new();
-			if (OperatingSystem.IsWindows())
+			if (bCrossPlatformCompilation && _targetWorkerPlatform == OSPlatform.Linux)
+			{
+				agentFiles.Add("UbaAgent");
+			}
+			else if (OperatingSystem.IsWindows())
 			{
 				agentFiles.Add("UbaAgent.exe");
 				if (FileReference.Exists(FileReference.Combine(ubaDir, "UbaWine.dll.so")))
@@ -501,7 +520,7 @@ namespace UnrealBuildTool
 
 				UBAHordeSession session = new(ubaConfig, crypto, server, token, hordeConfig.HordePool, allowWine, hordeConfig.HordeMaxCores, bStrictErrors, connectionMode, encryption, logger);
 				hordeCreateTimlineEvent.Dispose();
-				await session.InitAsync(useSentry: !String.IsNullOrEmpty(hordeConfig.UBASentryUrl), cancellationToken);
+				await session.InitAsync(useSentry: !String.IsNullOrEmpty(hordeConfig.UBASentryUrl), hordeConfig.HordeCondition, cancellationToken);
 				return (session, "Ready");
 			}
 			catch (TaskCanceledException)
@@ -586,7 +605,7 @@ namespace UnrealBuildTool
 						//arguments.Add("-NoStore");
 						//arguments.Add("-KillRandom"); // For debugging
 
-						if (OperatingSystem.IsMacOS())// && EpicGames.UBA.Utils.DisallowedPaths.Any())
+						if (OperatingSystem.IsMacOS() && _targetWorkerPlatform == OSPlatform.OSX)
 						{
 							// we need to populate the cas with all known xcodes so we can serve all the ones we have installed
 							string xcodeVersion = Utils.RunLocalProcessAndReturnStdOut("/bin/sh", "-c '/usr/bin/defaults read $(xcode-select -p)/../version.plist ProductBuildVersion");
@@ -722,6 +741,40 @@ namespace UnrealBuildTool
 				_updateStatus(0, 6, status, EpicGames.UBA.LogEntryType.Info, null);
 			}
 		}
+
+		static OSPlatform GetTargetWorkerPlatform(string? HordeCondition)
+		{
+			if (!String.IsNullOrEmpty(HordeCondition))
+			{
+				if (HordeCondition.Contains("Linux", StringComparison.OrdinalIgnoreCase))
+				{
+					return OSPlatform.Linux;
+				}
+				if (HordeCondition.Contains("Windows", StringComparison.OrdinalIgnoreCase))
+				{
+					return OSPlatform.Windows;
+				}
+				if (HordeCondition.Contains("MacOS", StringComparison.OrdinalIgnoreCase))
+				{
+					return OSPlatform.OSX;
+				}
+			}
+
+			return GetHostPlatform();
+		}
+
+		static OSPlatform GetHostPlatform()
+		{
+			if (OperatingSystem.IsLinux())
+			{
+				return OSPlatform.Linux;
+			}
+			if (OperatingSystem.IsWindows())
+			{
+				return OSPlatform.Windows;
+			}
+			return OSPlatform.OSX;
+		}
 	}
 
 	class UBAAgentCoordinatorHorde : IUBAAgentCoordinator, IDisposable
@@ -1042,8 +1095,13 @@ namespace UnrealBuildTool
 			if (!String.IsNullOrEmpty(config.HordeCluster))
 			{
 				// Apply filtering options when a cluster, either explicit or auto is specified, ensuring the cluster options take precidence.
-				string condition = "";
-				if (OperatingSystem.IsWindows())
+				string condition;
+				if (config.HordeCondition != null)
+				{
+					// Use explicitly configured condition, allowing cross-platform distribution (e.g., Mac host to Linux agents)
+					condition = config.HordeCondition;
+				}
+				else if (OperatingSystem.IsWindows())
 				{
 					condition = $"({KnownPropertyNames.OsFamily} == 'Windows' || {KnownPropertyNames.WineEnabled} == 'true')";
 				}
@@ -1055,10 +1113,9 @@ namespace UnrealBuildTool
 				{
 					condition = $"{KnownPropertyNames.OsFamily} == 'Linux'";
 				}
-				
-				if (config.HordeCondition != null)
+				else
 				{
-					condition = " " + config.HordeCondition;
+					condition = "";
 				}
 				requirements.Condition = Condition.Parse(condition);
 			}
diff --git a/Engine/Source/Programs/UnrealBuildTool/Executors/UnrealBuildAccelerator/UBAExecutor.cs b/Engine/Source/Programs/UnrealBuildTool/Executors/UnrealBuildAccelerator/UBAExecutor.cs
--- a/Engine/Source/Programs/UnrealBuildTool/Executors/UnrealBuildAccelerator/UBAExecutor.cs
+++ b/Engine/Source/Programs/UnrealBuildTool/Executors/UnrealBuildAccelerator/UBAExecutor.cs
@@ -334,6 +334,14 @@ namespace UnrealBuildTool
 				_config.AddValue("Session", "AllowMemoryMaps", "false");
 			}
 
+			// Register Apple cross-architecture paths BEFORE creating the session.
+			// SessionServerImpl iterates Utils.CrossArchitecturePaths in its constructor,
+			// so these paths must be registered before CreateSessionServer() is called.
+			if (OperatingSystem.IsMacOS())
+			{ 				
+				AppleToolChain.RegisterCrossArchitecturePathsForUBA(logger);
+			}
+	
 			if (UBAConfig.MaxRacingPercent > 0)
 			{
 				_config.AddValue("Scheduler", "MaxRacingPercent", UBAConfig.MaxRacingPercent.ToString());
diff --git a/Engine/Source/Programs/UnrealBuildTool/Platform/Apple/AppleToolChain.cs b/Engine/Source/Programs/UnrealBuildTool/Platform/Apple/AppleToolChain.cs
--- a/Engine/Source/Programs/UnrealBuildTool/Platform/Apple/AppleToolChain.cs
+++ b/Engine/Source/Programs/UnrealBuildTool/Platform/Apple/AppleToolChain.cs
@@ -87,6 +87,12 @@ namespace UnrealBuildTool
 
 		protected bool bCreatingAppOnWindows => AppleExports.CreatingAppOnWindows(ProjectFile);
 
+		protected ApplePortableToolChainInfo? PortableToolchain;
+		private static bool CrossArchitecturePathsRegistered;
+
+		// Platform identifier for Apple OSS headers when compiling without the real SDK
+		protected abstract string XnuPlatformDefine { get; }
+
 		public AppleToolChain(ReadOnlyTargetRules? Target, Func<AppleToolChainSettings> InCreateSettings, ClangToolChainOptions InOptions, ILogger InLogger) : base(InOptions, InLogger)
 		{
 			this.Target = Target;
@@ -94,6 +100,79 @@ namespace UnrealBuildTool
 			ToolChainSettings = new Lazy<AppleToolChainSettings>(InCreateSettings);
 
 			AppleExports.GetSwiftIntegrationSettings(ProjectFile, Target == null ? TargetType.Game : Target.Type, Target == null ? UnrealTargetPlatform.Mac : Target.Platform, out bUseSwiftUIMain, out _);
+			// Initialize portable toolchain for distributed compilation of non-SDK files
+			InitializePortableToolchain(Target);
+		}
+
+		private void InitializePortableToolchain(ReadOnlyTargetRules? Target)
+		{
+			if (PortableToolchain != null)
+			{
+				return;
+			}
+
+			if (Target == null || !GetAppleTargetRules(Target).bUsePortableToolchain)
+			{
+				return;
+			}
+
+			ApplePortableToolChainInfo ToolChainInfo = new ApplePortableToolChainInfo(Logger);
+			if (ToolChainInfo.IsValid)
+			{
+				PortableToolchain = ToolChainInfo;
+				RegisterUBAPathsForDistribution();
+			}
+		}
+
+		private void RegisterUBAPathsForDistribution()
+		{
+			EpicGames.UBA.Utils.RegisterDisallowedPaths(
+				"/usr/bin/",
+				"/Applications/Xcode.app/",
+				"/Library/Developer/CommandLineTools/");
+		}
+
+		public static void RegisterCrossArchitecturePathsForUBA(ILogger Logger)
+		{
+			if (CrossArchitecturePathsRegistered)
+			{
+				return;
+			}
+			
+			CrossArchitecturePathsRegistered = true;
+
+			string? SdksRoot = Environment.GetEnvironmentVariable("UE_SDKS_ROOT");
+			if (String.IsNullOrEmpty(SdksRoot))
+			{
+				Logger.LogDebug("[UBA] UE_SDKS_ROOT not set, skipping Apple cross-architecture path registration");
+				return;
+			}
+
+			string MacAppleDirectory = Path.Combine(SdksRoot, "HostMac", "Apple");
+			string? Version = ApplePortableToolChainInfo.ResolveVersion(MacAppleDirectory, Logger);
+			if (Version == null)
+			{
+				Logger.LogDebug("[UBA] No Apple toolchain version found in {Path}", MacAppleDirectory);
+				return;
+			}
+
+			string MacClangPath = Path.Combine(SdksRoot, "HostMac", "Apple", Version, "Clang");
+			string LinuxClangPath = Path.Combine(SdksRoot, "HostLinux", "Apple", Version, "Clang");
+			string WindowsClangPath = Path.Combine(SdksRoot, "HostWin64", "Apple", Version, "Clang");
+
+			if (Directory.Exists(MacClangPath))
+			{
+				if (Directory.Exists(LinuxClangPath))
+				{
+					EpicGames.UBA.Utils.RegisterCrossArchitecturePath(MacClangPath, LinuxClangPath);
+					Logger.LogDebug("[UBA] Early registration: {MacPath} -> {LinuxPath}", MacClangPath, LinuxClangPath);
+				}
+				if (Directory.Exists(WindowsClangPath))
+				{
+					EpicGames.UBA.Utils.RegisterCrossArchitecturePath(MacClangPath, WindowsClangPath);
+					Logger.LogDebug("[UBA] Early registration: {MacPath} -> {WindowsPath}", MacClangPath, WindowsClangPath);
+				}
+			}
 		}
 
 		/// <inheritdoc/>
@@ -146,12 +225,9 @@ namespace UnrealBuildTool
 				return new CPPOutput();
 			}
 
-			// log some iseful info
-			Log.TraceInformationOnce($"Compiling with {CompileEnvironment.Platform} SDK {Settings.SDKVersion}, minimum Target OS version {Settings.GetTargetVersionForTargetType(Target!.Type)}");
-
-			List<string> GlobalArguments = new();
+			bool bUsePortableToolchain = PortableToolchain != null && !CompileEnvironment.bRequiresPlatformSDK;
 
-			GetCompileArguments_Global(CompileEnvironment, GlobalArguments);
+			Log.TraceInformationOnce($"Compiling with {CompileEnvironment.Platform} SDK {Settings.SDKVersion}, minimum Target OS version {Settings.GetTargetVersionForTargetType(Target!.Type)}");
 
 			List<FileItem> FrameworkTokenFiles = new List<FileItem>();
 			foreach (UEBuildFramework Framework in CompileEnvironment.AdditionalFrameworks)
@@ -164,6 +240,17 @@ namespace UnrealBuildTool
 					FrameworkTokenFiles.Add(Framework.ExtractedTokenFile!);
 				}
 			}
+			
+			List<string> GlobalArguments = new();
+
+			if (bUsePortableToolchain)
+			{
+				GetCompileArguments_PortableToolchain(CompileEnvironment, GlobalArguments);
+			}
+			else
+			{
+				GetCompileArguments_Global(CompileEnvironment, GlobalArguments);
+			}
 
 			CPPOutput Result = new CPPOutput();
 			// Create a compile action for each source file.
@@ -179,12 +266,27 @@ namespace UnrealBuildTool
 				{
 					Action CompileAction = CompileCPPFile(CompileEnvironment, SourceFile, OutputDir, ModuleName, Graph, GlobalArguments, Result);
 					CompileAction.PrerequisiteItems.UnionWith(FrameworkTokenFiles);
+
+					if (bUsePortableToolchain)
+					{
+						CompileAction.CommandPath = PortableToolchain!.ClangPath!;
+						CompileAction.bCanExecuteInUBACrossArchitecture = true;
+					}
+					else
+					{
+						CompileAction.bCanExecuteInUBACrossArchitecture = false;
+					}
 				}
 			}
 			if (SwiftFiles.Any())
 			{
 				Action CompileAction = CompileSwiftFiles(CompileEnvironment, SwiftFiles, OutputDir, ModuleName, Graph, Result);
 				CompileAction.PrerequisiteItems.UnionWith(FrameworkTokenFiles);
+
+				if (bUsePortableToolchain)
+				{
+					CompileAction.bCanExecuteInUBACrossArchitecture = false;
+				}
 			}
 	
 			return Result;
@@ -325,16 +427,25 @@ namespace UnrealBuildTool
 			base.GetCppStandardCompileArgument(CompileEnvironment, Arguments);
 		}
 
-		protected override void GetCompileArguments_CPP(CppCompileEnvironment CompileEnvironment, List<string> Arguments)
+		protected void GetCompileArguments_ClangModules(List<string> Arguments)
 		{
-			Arguments.Add("-x objective-c++");
-			GetCppStandardCompileArgument(CompileEnvironment, Arguments);
-			Arguments.Add("-stdlib=libc++");
 			Arguments.Add("-fmodules");
 			Arguments.Add("-fno-implicit-modules");
 			Arguments.Add("-fimplicit-module-maps");
 			Arguments.Add("-Wno-module-import-in-extern-c");
-			//Arguments.Add("-fcxx-modules"); // For some reason this does not work
+		}
+
+		protected override void GetCompileArguments_CPP(CppCompileEnvironment CompileEnvironment, List<string> Arguments)
+		{
+			Arguments.Add("-x objective-c++");
+			GetCppStandardCompileArgument(CompileEnvironment, Arguments);
+			Arguments.Add("-stdlib=libc++");
+
+			bool bUsesPortableToolchain = CompileEnvironment.bRequiresPlatformSDK && PortableToolchain != null;
+			if (bUsesPortableToolchain)
+			{
+				GetCompileArguments_ClangModules(Arguments);
+			}
 		}
 
 		protected override void GetCompileArguments_MM(CppCompileEnvironment CompileEnvironment, List<string> Arguments)
@@ -476,6 +587,67 @@ namespace UnrealBuildTool
 			Arguments.Add($"-target {ToolChainSettings.Value.GetTargetTuple(CompileEnvironment.Architecture, Target!.Platform, Target!.Type, Override)}");
 		}
 
+		/// <summary>
+		/// Adds include paths for portable toolchain compilation with Apple OSS headers.
+		/// Used for distributed compilation on non-Apple workers.
+		/// </summary>
+		protected void AddPortableToolchainIncludePaths(List<string> Arguments)
+		{
+			if (PortableToolchain == null || PortableToolchain.OSSHeadersPath == null)
+			{
+				throw new InvalidOperationException("AddPortableToolchainIncludePaths requires initialized PortableToolchain");
+			}
+
+			DirectoryReference OSSHeadersDir = DirectoryReference.Combine(PortableToolchain.OSSHeadersPath, "usr", "include");
+			if (!DirectoryReference.Exists(OSSHeadersDir))
+			{
+				Logger.LogError("OSS headers not found at expected path: {Path}", OSSHeadersDir.FullName);
+				return;
+			}
+
+			Arguments.Add("-nostdinc");
+
+			// FreeBSD math.h compatibility
+			Arguments.Add("-D__ISO_C_VISIBLE=1999");
+			Arguments.Add("-D__float_t=float");
+			Arguments.Add("-D__double_t=double");
+			
+			VersionNumber AppleClangVersion = ((AppleToolChainInfo)Info).AppleClangVersion;
+			Arguments.Add($"-fclang-abi-compat={AppleClangVersion.Components[0]}");
+
+			Arguments.Add($"-D{XnuPlatformDefine}");
+
+			// libc++ headers must come first for #include_next to work
+			if (PortableToolchain.LibCxxIncludePath != null && DirectoryReference.Exists(PortableToolchain.LibCxxIncludePath))
+			{
+				Arguments.Add($"-isystem \"{PortableToolchain.LibCxxIncludePath.FullName}\"");
+			}
+			else
+			{
+				Logger.LogError("libc++ headers not found");
+			}
+
+			// Clang builtin headers (stddef.h, stdarg.h, intrinsics)
+			if (PortableToolchain.ClangBuiltinIncludePath != null && DirectoryReference.Exists(PortableToolchain.ClangBuiltinIncludePath))
+			{
+				Arguments.Add($"-isystem \"{PortableToolchain.ClangBuiltinIncludePath.FullName}\"");
+			}
+			else
+			{
+				Logger.LogError("Clang builtin headers not found");
+			}
+
+			// Apple OSS headers
+			Arguments.Add($"-isystem \"{OSSHeadersDir.FullName}\"");
+		}
+
+		/// <summary>
+		/// Gets global compile arguments for portable toolchain (without SDK-specific paths like -isysroot, -F framework paths).
+		/// Used for distributed compilation of non-Apple SDK files on Linux workers.
+		/// Must be implemented by derived classes (IOSToolChain, MacToolChain).
+		/// </summary>
+		protected abstract void GetCompileArguments_PortableToolchain(CppCompileEnvironment CompileEnvironment, List<string> Arguments);
+
 		public override CppCompileEnvironment CreateSharedResponseFile(CppCompileEnvironment CompileEnvironment, FileReference OutResponseFile, IActionGraphBuilder Graph)
 		{
 			// Temporarily turn of shared response files for apple toolchains
diff --git a/Engine/Source/Programs/UnrealBuildTool/Platform/Apple/UEBuildApple.cs b/Engine/Source/Programs/UnrealBuildTool/Platform/Apple/UEBuildApple.cs
--- a/Engine/Source/Programs/UnrealBuildTool/Platform/Apple/UEBuildApple.cs
+++ b/Engine/Source/Programs/UnrealBuildTool/Platform/Apple/UEBuildApple.cs
@@ -63,6 +63,14 @@ namespace UnrealBuildTool
 		/// </summary>
 		[ConfigFile(ConfigHierarchyType.Game, "/Script/AppleTargetPlatform.AppleTargetSettings", "bEnableContactsManager")]
 		public bool bEnableContactsManager = false;
+
+		/// <summary>
+		/// Experimental: Enable portable toolchain for distributed compilation. Separates SDK-dependent
+		/// files (local) from portable files (remote). Requires Auto SDK (UE_SDKS_ROOT).
+		/// </summary>
+		[CommandLine("-UsePortableToolchain")]
+		[XmlConfigFile(Category = "BuildConfiguration")]
+		public bool bUsePortableToolchain = false;
 	}
 
 	/// <summary>
@@ -97,6 +105,7 @@ namespace UnrealBuildTool
 		public bool bEnableThreadSanitizer => Inner.bEnableThreadSanitizer;
 		public bool bEnableUndefinedBehaviorSanitizer => Inner.bEnableUndefinedBehaviorSanitizer;
 		public bool bEnableContactsManager => Inner.bEnableContactsManager;
+		public bool bUsePortableToolchain => Inner.bUsePortableToolchain;
 #pragma warning restore CS1591
 		#endregion
 	}
diff --git a/Engine/Source/Programs/UnrealBuildTool/Platform/IOS/IOSToolChain.cs b/Engine/Source/Programs/UnrealBuildTool/Platform/IOS/IOSToolChain.cs
--- a/Engine/Source/Programs/UnrealBuildTool/Platform/IOS/IOSToolChain.cs
+++ b/Engine/Source/Programs/UnrealBuildTool/Platform/IOS/IOSToolChain.cs
@@ -57,6 +57,8 @@ namespace UnrealBuildTool
 		/// </summary>
 		private const string IOSArchiver = "libtool";
 
+		protected override string XnuPlatformDefine => "XNU_PLATFORM_iPhoneOS";
+
 		protected override ClangToolChainInfo GetToolChainInfo()
 		{
 			FileReference CompilerPath = FileReference.Combine(AppleToolChainSettings.ToolchainDir, IOSCompiler);
@@ -149,6 +151,12 @@ namespace UnrealBuildTool
 
 		/// <inheritdoc/>
 		protected override void GetCompileArguments_Global(CppCompileEnvironment CompileEnvironment, List<string> Arguments)
+		{
+			bool bIsUsingPortableSDK = false;
+			GetCompileArguments_Global(CompileEnvironment, Arguments, bIsUsingPortableSDK);
+		}
+
+		private void GetCompileArguments_Global(CppCompileEnvironment CompileEnvironment, List<string> Arguments, bool bIsUsingPortableSDK)
 		{
 			base.GetCompileArguments_Global(CompileEnvironment, Arguments);
 
@@ -170,28 +178,42 @@ namespace UnrealBuildTool
 			// What architecture(s) to build for
 			Arguments.Add(FormatArchitectureArg(CompileEnvironment.Architectures));
 
-			Arguments.Add($"-isysroot \"{NormalizeCommandLinePath(Settings.GetSDKPath(CompileEnvironment.Architecture), RootPaths)}\"");
-
-			// Add additional frameworks so that their headers can be found
-			foreach (UEBuildFramework Framework in CompileEnvironment.AdditionalFrameworks)
+			// Clang modules depend on SDK paths unavailable on remote workers
+			if (bIsUsingPortableSDK)
+			{
+				AddPortableToolchainIncludePaths(Arguments);
+			}
+			else
 			{
-				if (Framework.bLinkFramework)
+				Arguments.Add($"-isysroot \"{NormalizeCommandLinePath(Settings.GetSDKPath(CompileEnvironment.Architecture), RootPaths)}\"");
+
+				// Add additional frameworks so that their headers can be found
+				foreach (UEBuildFramework Framework in CompileEnvironment.AdditionalFrameworks)
 				{
-					DirectoryReference? FrameworkDirectory = Framework.GetFrameworkDirectory(CompileEnvironment.Platform, CompileEnvironment.Architecture, Logger);
-					if (FrameworkDirectory != null)
+					if (Framework.bLinkFramework)
 					{
-						// embedded frameworks have a framework inside of this directory, so we use this directory. regular frameworks need to go one up to point to the 
-						// directory containing the framework. -F gives a path to look for the -framework
-						if (FrameworkDirectory.FullName.EndsWith(".framework", StringComparison.Ordinal))
+						DirectoryReference? FrameworkDirectory = Framework.GetFrameworkDirectory(CompileEnvironment.Platform, CompileEnvironment.Architecture, Logger);
+						if (FrameworkDirectory != null)
 						{
-							FrameworkDirectory = FrameworkDirectory.ParentDirectory;
+							// embedded frameworks have a framework inside of this directory, so we use this directory. regular frameworks need to go one up to point to the 
+							// directory containing the framework. -F gives a path to look for the -framework
+							if (FrameworkDirectory.FullName.EndsWith(".framework", StringComparison.Ordinal))
+							{
+								FrameworkDirectory = FrameworkDirectory.ParentDirectory;
+							}
+							Arguments.Add($"-F \"{NormalizeCommandLinePath(FrameworkDirectory!, RootPaths)}\"");
 						}
-						Arguments.Add($"-F \"{NormalizeCommandLinePath(FrameworkDirectory!, RootPaths)}\"");
 					}
 				}
 			}
 		}
 
+		protected override void GetCompileArguments_PortableToolchain(CppCompileEnvironment CompileEnvironment, List<string> Arguments)
+		{
+			bool bIsUsingPortableSDK = true;
+			GetCompileArguments_Global(CompileEnvironment, Arguments, bIsUsingPortableSDK);
+		}
+
 		/// <inheritdoc/>
 		protected override FileItem GetCompileArguments_FileType(CppCompileEnvironment CompileEnvironment, FileItem SourceFile, DirectoryReference OutputDir, List<string> Arguments, Action CompileAction, CPPOutput CompileResult)
 		{
diff --git a/Engine/Source/Programs/UnrealBuildTool/Platform/Mac/ApplePortableToolChain.cs b/Engine/Source/Programs/UnrealBuildTool/Platform/Mac/ApplePortableToolChain.cs
new file mode 100644
--- /dev/null
+++ b/Engine/Source/Programs/UnrealBuildTool/Platform/Mac/ApplePortableToolChain.cs
@@ -0,0 +1,301 @@
+// Copyright Epic Games, Inc. All Rights Reserved.
+
+using System;
+using EpicGames.Core;
+using Microsoft.Extensions.Logging;
+using UnrealBuildBase;
+
+namespace UnrealBuildTool
+{
+	internal class ApplePortableToolChainInfo
+	{
+		public bool IsValid { get; }
+		public FileReference? ClangPath { get; }
+		public DirectoryReference? OSSHeadersPath { get; }
+		public DirectoryReference? ClangBasePath { get; }
+		public DirectoryReference? LibCxxIncludePath { get; }
+		public DirectoryReference? ClangBuiltinIncludePath { get; }
+
+		public ApplePortableToolChainInfo(ILogger Logger)
+		{
+			IsValid = false;
+			
+			DirectoryReference? HostAutoSdkDir = TryGetAutoSDKDir();
+			if (HostAutoSdkDir == null)
+			{
+				Logger.LogDebug("Auto SDK HostMac directory not found, portable toolchain unavailable");
+				return;
+			}
+
+			DirectoryReference AppleToolchainDir = DirectoryReference.Combine(HostAutoSdkDir, "Apple");
+			if (!DirectoryReference.Exists(AppleToolchainDir))
+			{
+				Logger.LogDebug("Apple directory not found at {Path}", AppleToolchainDir);
+				return;
+			}
+
+			string? ClangSubdir = GetConfigValue("ClangSubdir");
+			string? OssHeadersSubdir = GetConfigValue("OSSHeadersSubdir");
+
+			if (String.IsNullOrEmpty(ClangSubdir) || String.IsNullOrEmpty(OssHeadersSubdir))
+			{
+				Logger.LogDebug("AppleOpenSource_SDK.json config missing required paths");
+				return;
+			}
+
+			string? Version = ResolveVersion(AppleToolchainDir, Logger);
+			if (Version == null)
+			{
+				Logger.LogDebug("No valid Apple portable toolchain version found in {Path}", AppleToolchainDir);
+				return;
+			}
+
+			DirectoryReference VersionDir = DirectoryReference.Combine(AppleToolchainDir, Version);
+			FileReference Clang = FileReference.Combine(VersionDir, ClangSubdir);
+			DirectoryReference OssHeaders = DirectoryReference.Combine(VersionDir, OssHeadersSubdir);
+
+			if (!FileReference.Exists(Clang))
+			{
+				Logger.LogDebug("Portable toolchain clang not found at {Path}", Clang);
+				return;
+			}
+
+			DirectoryReference OssIncludePath = DirectoryReference.Combine(OssHeaders, "usr", "include");
+			if (!DirectoryReference.Exists(OssIncludePath))
+			{
+				Logger.LogDebug("Apple OSS headers not found at {Path}", OssIncludePath);
+				return;
+			}
+
+			DirectoryReference ClangBase = Clang.Directory.ParentDirectory!;
+
+			DirectoryReference LibCxxInclude = DirectoryReference.Combine(ClangBase, "include", "c++", "v1");
+			if (!DirectoryReference.Exists(LibCxxInclude))
+			{
+				Logger.LogWarning("libc++ headers not found at {Path}", LibCxxInclude);
+				return;
+			}
+
+			DirectoryReference? ClangBuiltinInclude = null;
+			DirectoryReference ClangLibDir = DirectoryReference.Combine(ClangBase, "lib", "clang");
+
+			if (DirectoryReference.Exists(ClangLibDir))
+			{
+				foreach (DirectoryReference VersionSubDir in DirectoryReference.EnumerateDirectories(ClangLibDir))
+				{
+					DirectoryReference IncludePath = DirectoryReference.Combine(VersionSubDir, "include");
+					if (DirectoryReference.Exists(IncludePath))
+					{
+						ClangBuiltinInclude = IncludePath;
+						break;
+					}
+				}
+			}
+
+			if (ClangBuiltinInclude == null)
+			{
+				Logger.LogWarning("Clang builtin headers not found in {Path}", ClangLibDir);
+				return;
+			}
+
+			Logger.LogInformation("Found Apple portable toolchain: {Version}", Version);
+
+			ClangPath = Clang;
+			OSSHeadersPath = OssHeaders;
+			ClangBasePath = ClangBase;
+			LibCxxIncludePath = LibCxxInclude;
+			ClangBuiltinIncludePath = ClangBuiltinInclude;
+			IsValid = true;
+		}
+
+		public static string? ResolveVersion(string AppleToolchainDirectory, ILogger Logger)
+		{
+			DirectoryReference AppleToolchainDir = new DirectoryReference(AppleToolchainDirectory);
+			return ResolveVersion(AppleToolchainDir, Logger);
+		}
+
+		private static string? ResolveVersion(DirectoryReference AppleToolchainDir, ILogger Logger)
+		{
+			string? MainVersion = GetConfigValue("MainVersion");
+			if (!String.IsNullOrEmpty(MainVersion))
+			{
+				DirectoryReference VersionDir = DirectoryReference.Combine(AppleToolchainDir, MainVersion);
+				if (DirectoryReference.Exists(VersionDir))
+				{
+					return MainVersion;
+				}
+				Logger.LogDebug("MainVersion '{Version}' specified but directory not found at {Path}", MainVersion, VersionDir);
+			}
+
+			// Try to select based on Xcode version mapping
+			string? MappedVersion = ResolveVersionFromXcodeMapping(AppleToolchainDir, Logger);
+			if (MappedVersion != null)
+			{
+				return MappedVersion;
+			}
+
+			// Fallback to auto-discovery of latest v* directory
+			string? LatestVersion = null;
+			int LatestNumericVersion = -1;
+
+			if (DirectoryReference.Exists(AppleToolchainDir))
+			{
+				foreach (DirectoryReference Dir in DirectoryReference.EnumerateDirectories(AppleToolchainDir))
+				{
+					string DirName = Dir.GetDirectoryName();
+					if (DirName.StartsWith("v", StringComparison.OrdinalIgnoreCase))
+					{
+						int NumericVersion = ParseVersionNumber(DirName);
+						if (NumericVersion > LatestNumericVersion)
+						{
+							LatestNumericVersion = NumericVersion;
+							LatestVersion = DirName;
+						}
+					}
+				}
+			}
+
+			if (LatestVersion != null)
+			{
+				Logger.LogInformation("Auto-discovered Apple portable toolchain version: {Version}", LatestVersion);
+			}
+
+			return LatestVersion;
+		}
+
+		/// <summary>
+		/// Resolves toolchain version based on XcodeToToolchainVersion mapping.
+		/// </summary>
+		private static string? ResolveVersionFromXcodeMapping(DirectoryReference AppleToolchainDir, ILogger Logger)
+		{
+			// Get the installed Xcode version
+			string? XcodeVersionString = ApplePlatformSDK.InstalledXcodeVersion.Value;
+			if (String.IsNullOrEmpty(XcodeVersionString))
+			{
+				return null;
+			}
+
+			if (!Version.TryParse(XcodeVersionString, out Version? XcodeVersion))
+			{
+				Logger.LogDebug("Failed to parse Xcode version: {Version}", XcodeVersionString);
+				return null;
+			}
+
+			// Get the mapping array from config
+			string[]? Mappings = GetConfigArray("XcodeToToolchainVersion");
+			if (Mappings == null || Mappings.Length == 0)
+			{
+				return null;
+			}
+
+			// Find the matching toolchain prefix (iterate in reverse for highest match)
+			string? ToolchainPrefix = null;
+			for (int MappingIndex = Mappings.Length - 1; MappingIndex >= 0; MappingIndex--)
+			{
+				string[] Parts = Mappings[MappingIndex].Split('-');
+				if (Parts.Length != 2)
+				{
+					continue;
+				}
+
+				if (Version.TryParse(Parts[0], out Version? MappingVersion) && XcodeVersion >= MappingVersion)
+				{
+					ToolchainPrefix = Parts[1];
+					Logger.LogInformation("Xcode {XcodeVersion} mapped to toolchain prefix {Prefix}", XcodeVersion, ToolchainPrefix);
+					break;
+				}
+			}
+
+			if (ToolchainPrefix == null)
+			{
+				return null;
+			}
+
+			// Find directory starting with this prefix
+			if (!DirectoryReference.Exists(AppleToolchainDir))
+			{
+				return null;
+			}
+
+			foreach (DirectoryReference Dir in DirectoryReference.EnumerateDirectories(AppleToolchainDir))
+			{
+				string DirName = Dir.GetDirectoryName();
+				if (DirName.StartsWith(ToolchainPrefix + "_", StringComparison.OrdinalIgnoreCase))
+				{
+					Logger.LogInformation("Selected portable toolchain: {Dir}", DirName);
+					return DirName;
+				}
+			}
+
+			Logger.LogDebug("No directory found with prefix {Prefix}", ToolchainPrefix);
+			return null;
+		}
+
+		private static int ParseVersionNumber(string VersionString)
+		{
+			int UnderscoreIndex = VersionString.IndexOf('_');
+			string NumericPart = UnderscoreIndex > 1
+				? VersionString.Substring(1, UnderscoreIndex - 1)
+				: VersionString.Substring(1);
+
+			return Int32.TryParse(NumericPart, out int Result) ? Result : 0;
+		}
+
+		private static DirectoryReference? TryGetAutoSDKDir()
+		{
+			string? SdksRoot = Environment.GetEnvironmentVariable("UE_SDKS_ROOT");
+			if (String.IsNullOrEmpty(SdksRoot))
+			{
+				return null;
+			}
+
+			DirectoryReference HostDir = DirectoryReference.Combine(new DirectoryReference(SdksRoot), "HostMac");
+			return DirectoryReference.Exists(HostDir) ? HostDir : null;
+		}
+
+		private static JsonObject? CachedConfig;
+		private static bool ConfigLoadAttempted;
+
+		private static string? GetConfigValue(string Key)
+		{
+			EnsureConfigLoaded();
+
+			if (CachedConfig != null && CachedConfig.TryGetStringField(Key, out string? Value))
+			{
+				return Value;
+			}
+
+			return null;
+		}
+
+		private static string[]? GetConfigArray(string Key)
+		{
+			EnsureConfigLoaded();
+
+			if (CachedConfig != null && CachedConfig.TryGetStringArrayField(Key, out string[]? Values))
+			{
+				return Values;
+			}
+
+			return null;
+		}
+
+		private static void EnsureConfigLoaded()
+		{
+			if (!ConfigLoadAttempted)
+			{
+				ConfigLoadAttempted = true;
+				FileReference ConfigFile = FileReference.Combine(
+					Unreal.EngineDirectory,
+					"Config",
+					"Apple",
+					"AppleOpenSource_SDK.json");
+
+				if (FileReference.Exists(ConfigFile))
+				{
+					JsonObject.TryRead(ConfigFile, out CachedConfig);
+				}
+			}
+		}
+	}
+}
diff --git a/Engine/Source/Programs/UnrealBuildTool/Platform/Mac/MacToolChain.cs b/Engine/Source/Programs/UnrealBuildTool/Platform/Mac/MacToolChain.cs
--- a/Engine/Source/Programs/UnrealBuildTool/Platform/Mac/MacToolChain.cs
+++ b/Engine/Source/Programs/UnrealBuildTool/Platform/Mac/MacToolChain.cs
@@ -63,6 +63,8 @@ namespace UnrealBuildTool
 		/// </summary>
 		private const string MacArchiver = "libtool";
 
+		protected override string XnuPlatformDefine => "XNU_PLATFORM_MacOSX";
+
 		protected override ClangToolChainInfo GetToolChainInfo()
 		{
 			FileReference CompilerPath = FileReference.Combine(AppleToolChainSettings.ToolchainDir, MacCompiler);
@@ -187,6 +189,12 @@ namespace UnrealBuildTool
 
 		/// <inheritdoc/>
 		protected override void GetCompileArguments_Global(CppCompileEnvironment CompileEnvironment, List<string> Arguments)
+		{
+			bool bIsUsingPortableSDK = false;
+			GetCompileArguments_Global(CompileEnvironment, Arguments, bIsUsingPortableSDK);
+		}
+
+		private void GetCompileArguments_Global(CppCompileEnvironment CompileEnvironment, List<string> Arguments, bool bIsUsingPortableSDK)
 		{
 			base.GetCompileArguments_Global(CompileEnvironment, Arguments);
 
@@ -206,20 +214,35 @@ namespace UnrealBuildTool
 
 			// Pass through architecture and OS info
 			Arguments.Add("" + FormatArchitectureArg(CompileEnvironment.Architectures));
-			Arguments.Add($"-isysroot \"{NormalizeCommandLinePath(Settings.GetSDKPath(CompileEnvironment.Architecture), RootPaths)}\"");
 
-			List<string> FrameworksSearchPaths = new List<string>();
-			foreach (UEBuildFramework Framework in CompileEnvironment.AdditionalFrameworks)
+			// Clang modules depend on SDK paths unavailable on remote workers
+			if (bIsUsingPortableSDK)
+			{
+				AddPortableToolchainIncludePaths(Arguments);
+			}
+			else
 			{
-				FileReference FrameworkPath = new FileReference(Path.GetFullPath(Framework.Name));
-				if (!FrameworksSearchPaths.Contains(FrameworkPath.Directory.FullName))
+				Arguments.Add($"-isysroot \"{NormalizeCommandLinePath(Settings.GetSDKPath(CompileEnvironment.Architecture), RootPaths)}\"");
+				
+				List<string> FrameworksSearchPaths = new List<string>();
+				foreach (UEBuildFramework Framework in CompileEnvironment.AdditionalFrameworks)
 				{
-					Arguments.Add($"-F \"{NormalizeCommandLinePath(FrameworkPath.Directory, RootPaths)}\"");
-					FrameworksSearchPaths.Add(FrameworkPath.Directory.FullName);
+					FileReference FrameworkPath = new FileReference(Path.GetFullPath(Framework.Name));
+					if (!FrameworksSearchPaths.Contains(FrameworkPath.Directory.FullName))
+					{
+						Arguments.Add($"-F \"{NormalizeCommandLinePath(FrameworkPath.Directory, RootPaths)}\"");
+						FrameworksSearchPaths.Add(FrameworkPath.Directory.FullName);
+					}
 				}
 			}
 		}
 
+		protected override void GetCompileArguments_PortableToolchain(CppCompileEnvironment CompileEnvironment, List<string> Arguments)
+		{
+			bool bIsUsingPortableSDK = true;
+			GetCompileArguments_Global(CompileEnvironment, Arguments, bIsUsingPortableSDK);
+		}
+
 		void AddFrameworkToLinkCommand(LinkEnvironment linkEnvironment, List<string> LinkArguments, string FrameworkName, string Arg = "-framework")
 		{
 			if (FrameworkName.EndsWith(".framework", StringComparison.Ordinal))
diff --git a/Engine/Source/Programs/UnrealBuildTool/Platform/TVOS/TVOSToolChain.cs b/Engine/Source/Programs/UnrealBuildTool/Platform/TVOS/TVOSToolChain.cs
--- a/Engine/Source/Programs/UnrealBuildTool/Platform/TVOS/TVOSToolChain.cs
+++ b/Engine/Source/Programs/UnrealBuildTool/Platform/TVOS/TVOSToolChain.cs
@@ -6,6 +6,8 @@ namespace UnrealBuildTool
 {
 	class TVOSToolChain : IOSToolChain
 	{
+		protected override string XnuPlatformDefine => "XNU_PLATFORM_AppleTVOS";
+
 		public TVOSToolChain(ReadOnlyTargetRules InTarget, TVOSProjectSettings InProjectSettings, ILogger InLogger)
 			: base(InTarget, InProjectSettings, () => new TVOSToolChainSettings(InProjectSettings, InLogger), ClangToolChainOptions.None, InLogger)
 		{
diff --git a/Engine/Source/Programs/UnrealBuildTool/System/CppCompileEnvironment.cs b/Engine/Source/Programs/UnrealBuildTool/System/CppCompileEnvironment.cs
--- a/Engine/Source/Programs/UnrealBuildTool/System/CppCompileEnvironment.cs
+++ b/Engine/Source/Programs/UnrealBuildTool/System/CppCompileEnvironment.cs
@@ -804,6 +804,11 @@ namespace UnrealBuildTool
 		/// </summary>
 		public bool bValidateExperimentalApi = false;
 
+		/// <summary>
+		/// Experimental: Whether this module requires platform specific SDK files. Not supported on all platforms yet.
+		/// </summary>
+		public bool bRequiresPlatformSDK = false;
+
 		/// <summary>
 		/// Directory where to put crash report files for platforms that support it
 		/// </summary>
@@ -955,6 +960,7 @@ namespace UnrealBuildTool
 			bValidateFormatStrings = Other.bValidateFormatStrings;
 			bValidateInternalApi = Other.bValidateInternalApi;
 			bValidateExperimentalApi = Other.bValidateExperimentalApi;
+			bRequiresPlatformSDK = Other.bRequiresPlatformSDK;
 			FPassPlugins = Other.FPassPlugins;
 		}
 
diff --git a/Engine/Source/Programs/UnrealBuildTool/System/TargetMakefile.cs b/Engine/Source/Programs/UnrealBuildTool/System/TargetMakefile.cs
--- a/Engine/Source/Programs/UnrealBuildTool/System/TargetMakefile.cs
+++ b/Engine/Source/Programs/UnrealBuildTool/System/TargetMakefile.cs
@@ -585,7 +585,7 @@ namespace UnrealBuildTool
 			{
 				lock (UbaCrossArchitecturePaths)
 				{
-					UbaCrossArchitecturePaths.Add(DirectoryItem.GetItemByPath(e.Path), DirectoryItem.GetItemByPath(e.CrossPath));
+					UbaCrossArchitecturePaths.TryAdd(DirectoryItem.GetItemByPath(e.Path), DirectoryItem.GetItemByPath(e.CrossPath));
 				}
 			};
 
@@ -593,7 +593,7 @@ namespace UnrealBuildTool
 			{
 				lock (UbaPathHashes)
 				{
-					UbaPathHashes.Add(DirectoryItem.GetItemByPath(e.Path), e.Hash);
+					UbaPathHashes.TryAdd(DirectoryItem.GetItemByPath(e.Path), e.Hash);
 				}
 			};
 		}
