[UBT] Fix two CIS warning in AppleToolChain files.

Mostly about making some methdos I added in a previous CL static

#Synced-CL 50088342

diff --git a/Engine/Source/Programs/UnrealBuildTool/Platform/Apple/AppleToolChain.cs b/Engine/Source/Programs/UnrealBuildTool/Platform/Apple/AppleToolChain.cs
--- a/Engine/Source/Programs/UnrealBuildTool/Platform/Apple/AppleToolChain.cs
+++ b/Engine/Source/Programs/UnrealBuildTool/Platform/Apple/AppleToolChain.cs
@@ -124,7 +124,7 @@ namespace UnrealBuildTool
 			}
 		}
 
-		private void RegisterUBAPathsForDistribution()
+		private static void RegisterUBAPathsForDistribution()
 		{
 			EpicGames.UBA.Utils.RegisterDisallowedPaths(
 				"/usr/bin/",
@@ -427,7 +427,7 @@ namespace UnrealBuildTool
 			base.GetCppStandardCompileArgument(CompileEnvironment, Arguments);
 		}
 
-		protected void GetCompileArguments_ClangModules(List<string> Arguments)
+		protected static void GetCompileArguments_ClangModules(List<string> Arguments)
 		{
 			Arguments.Add("-fmodules");
 			Arguments.Add("-fno-implicit-modules");
diff --git a/Engine/Source/Programs/UnrealBuildTool/Platform/Mac/ApplePortableToolChain.cs b/Engine/Source/Programs/UnrealBuildTool/Platform/Mac/ApplePortableToolChain.cs
--- a/Engine/Source/Programs/UnrealBuildTool/Platform/Mac/ApplePortableToolChain.cs
+++ b/Engine/Source/Programs/UnrealBuildTool/Platform/Mac/ApplePortableToolChain.cs
@@ -233,7 +233,7 @@ namespace UnrealBuildTool
 
 		private static int ParseVersionNumber(string VersionString)
 		{
-			int UnderscoreIndex = VersionString.IndexOf('_');
+			int UnderscoreIndex = VersionString.IndexOf('_', StringComparison.Ordinal);
 			string NumericPart = UnderscoreIndex > 1
 				? VersionString.Substring(1, UnderscoreIndex - 1)
 				: VersionString.Substring(1);
