UnrealBuildTool: Fix issue CA1870 (Use a cached 'SearchValues' instance).

#Synced-CL 50096664

diff --git a/Engine/Source/Programs/UnrealBuildTool/System/CommandLine.cs b/Engine/Source/Programs/UnrealBuildTool/System/CommandLine.cs
--- a/Engine/Source/Programs/UnrealBuildTool/System/CommandLine.cs
+++ b/Engine/Source/Programs/UnrealBuildTool/System/CommandLine.cs
@@ -1,6 +1,7 @@
 // Copyright Epic Games, Inc. All Rights Reserved.
 
 using System;
+using System.Buffers;
 using System.Collections.Generic;
 using System.Diagnostics.CodeAnalysis;
 using System.Linq;
@@ -77,6 +78,8 @@ namespace UnrealBuildTool
 			return type;
 		}
 
+		private static readonly SearchValues<char> ArgumentSeparators = SearchValues.Create("=:");
+
 		/// <summary>
 		/// Parse the given list of arguments, and remove any that are parsed successfully
 		/// </summary>
@@ -126,7 +129,7 @@ namespace UnrealBuildTool
 				if (Argument.Length > 0 && Argument[0] == '-')
 				{
 					// Get the length of the argument prefix
-					int EqualsIdx = Argument.IndexOfAny(new char[] { '=', ':' });
+					int EqualsIdx = Argument.AsSpan().IndexOfAny(ArgumentSeparators);
 					string Prefix = (EqualsIdx == -1) ? Argument : Argument.Substring(0, EqualsIdx + 1);
 
 					// Check if there's a matching argument registered
