[UBT/UBA] Updating auto sdk paths used by Uba cross-platform

#Synced-CL 50088342

diff --git a/Engine/Source/Programs/UnrealBuildTool/Platform/Apple/AppleToolChain.cs b/Engine/Source/Programs/UnrealBuildTool/Platform/Apple/AppleToolChain.cs
--- a/Engine/Source/Programs/UnrealBuildTool/Platform/Apple/AppleToolChain.cs
+++ b/Engine/Source/Programs/UnrealBuildTool/Platform/Apple/AppleToolChain.cs
@@ -156,21 +156,22 @@ namespace UnrealBuildTool
 				return;
 			}
 
-			string MacClangPath = Path.Combine(SdksRoot, "HostMac", "Apple", Version, "Clang");
-			string LinuxClangPath = Path.Combine(SdksRoot, "HostLinux", "Apple", Version, "Clang");
-			string WindowsClangPath = Path.Combine(SdksRoot, "HostWin64", "Apple", Version, "Clang");
+			string ClangBasePath = Path.Combine(SdksRoot, "HostMac", "Apple", Version, "Clang");
+			string MacBinPath = Path.Combine(ClangBasePath, "bin");
+			string LinuxBinPath = Path.Combine(ClangBasePath, "binlinux");
+			string WindowsBinPath = Path.Combine(ClangBasePath, "binwin64");
 
-			if (Directory.Exists(MacClangPath))
+			if (Directory.Exists(MacBinPath))
 			{
-				if (Directory.Exists(LinuxClangPath))
+				if (Directory.Exists(LinuxBinPath))
 				{
-					EpicGames.UBA.Utils.RegisterCrossArchitecturePath(MacClangPath, LinuxClangPath);
-					Logger.LogDebug("[UBA] Early registration: {MacPath} -> {LinuxPath}", MacClangPath, LinuxClangPath);
+					EpicGames.UBA.Utils.RegisterCrossArchitecturePath(MacBinPath, LinuxBinPath);
+					Logger.LogDebug("[UBA] Early registration: {MacPath} -> {LinuxPath}", MacBinPath, LinuxBinPath);
 				}
-				if (Directory.Exists(WindowsClangPath))
+				if (Directory.Exists(WindowsBinPath))
 				{
-					EpicGames.UBA.Utils.RegisterCrossArchitecturePath(MacClangPath, WindowsClangPath);
-					Logger.LogDebug("[UBA] Early registration: {MacPath} -> {WindowsPath}", MacClangPath, WindowsClangPath);
+					EpicGames.UBA.Utils.RegisterCrossArchitecturePath(MacBinPath, WindowsBinPath);
+					Logger.LogDebug("[UBA] Early registration: {MacPath} -> {WindowsPath}", MacBinPath, WindowsBinPath);
 				}
 			}
 		}
