UnrealBuildTool: Fix issue CA1311 (Specify a culture or use an invariant version).

#Synced-CL 50156122

diff --git a/Engine/Source/Programs/UnrealBuildTool/Configuration/UEBuildTarget.cs b/Engine/Source/Programs/UnrealBuildTool/Configuration/UEBuildTarget.cs
--- a/Engine/Source/Programs/UnrealBuildTool/Configuration/UEBuildTarget.cs
+++ b/Engine/Source/Programs/UnrealBuildTool/Configuration/UEBuildTarget.cs
@@ -754,19 +754,19 @@ namespace UnrealBuildTool
 		private static string FixName(string Name)
 		{
 			// allow for some alternate names
-			switch (Name.ToLower())
+			switch (Name.ToUpperInvariant())
 			{
-				case "a8":
-				case "a":
-				case "-arm64":
+				case "A8":
+				case "A":
+				case "-ARM64":
 					Name = "arm64";
 					break;
 
-				case "x86_64":
-				case "x6":
-				case "x":
-				case "-x64":
-				case "intel":
+				case "X86_64":
+				case "X6":
+				case "X":
+				case "-X64":
+				case "INTEL":
 					Name = "x64";
 					break;
 			}
@@ -2179,8 +2179,8 @@ namespace UnrealBuildTool
 					{
 						foreach (string IncludeFileName in Directory.EnumerateFiles(IncludePath, "*", SearchOption.AllDirectories))
 						{
-							string Extension = Path.GetExtension(IncludeFileName).ToLower();
-							if (Extension == ".h" || Extension == ".inl" || Extension == ".hpp" || Extension == ".ipp" || String.IsNullOrEmpty(Extension))
+							string Extension = Path.GetExtension(IncludeFileName).ToUpperInvariant();
+							if (Extension == ".H" || Extension == ".INL" || Extension == ".HPP" || Extension == ".IPP" || String.IsNullOrEmpty(Extension))
 							{
 								Files.Add(new FileReference(IncludeFileName));
 							}
@@ -6624,7 +6624,9 @@ namespace UnrealBuildTool
 			// tell the compiled code the name of the UBT platform (this affects folder on disk, etc that the game may need to know)
 			GlobalCompileEnvironment.Definitions.Add(String.Format("UBT_COMPILED_PLATFORM=" + Platform.ToString()));
 			GlobalCompileEnvironment.Definitions.Add(String.Format("UBT_COMPILED_TARGET=" + TargetType.ToString()));
-			GlobalCompileEnvironment.Definitions.Add(String.Format("UBT_COMPILED_ARCHITECTURE=" + (Architectures.bIsMultiArch ? "MULTI" : Architectures.SingleArchitecture.ToString().ToLower()) ));
+#pragma warning disable CA1308 // Normalize strings to uppercase - preserve casing of compiler define
+			GlobalCompileEnvironment.Definitions.Add(String.Format("UBT_COMPILED_ARCHITECTURE=" + (Architectures.bIsMultiArch ? "MULTI" : Architectures.SingleArchitecture.ToString().ToLowerInvariant()) ));
+#pragma warning restore CA1308
 
 			// Set the global app name
 			GlobalCompileEnvironment.Definitions.Add(String.Format("UE_APP_NAME=\"{0}\"", AppName));
diff --git a/Engine/Source/Programs/UnrealBuildTool/Executors/Experimental/FASTBuild.cs b/Engine/Source/Programs/UnrealBuildTool/Executors/Experimental/FASTBuild.cs
--- a/Engine/Source/Programs/UnrealBuildTool/Executors/Experimental/FASTBuild.cs
+++ b/Engine/Source/Programs/UnrealBuildTool/Executors/Experimental/FASTBuild.cs
@@ -418,13 +418,17 @@ namespace UnrealBuildTool
 			Tuple.Create<string, Func<LinkedAction, string>, FBBuildType>
 			(
 				"apple",
-				Action => Action.CommandArguments.ToLower(),
+#pragma warning disable CA1308 // Normalize strings to uppercase - external command may be case-sensitive
+				Action => Action.CommandArguments.ToLowerInvariant(),
+#pragma warning restore CA1308
 				FBBuildType.Apple
 			),
 			Tuple.Create<string, Func<LinkedAction, string>, FBBuildType>
 			(
 				"/bin/sh",
-				Action => Action.CommandPath.FullName.ToLower(),
+#pragma warning disable CA1308 // Normalize strings to uppercase - external command may be case-sensitive
+				Action => Action.CommandPath.FullName.ToLowerInvariant(),
+#pragma warning restore CA1308
 				FBBuildType.Apple
 			),
 			Tuple.Create<string, Func<LinkedAction, string>, FBBuildType>
diff --git a/Engine/Source/Programs/UnrealBuildTool/Modes/BuildMode.cs b/Engine/Source/Programs/UnrealBuildTool/Modes/BuildMode.cs
--- a/Engine/Source/Programs/UnrealBuildTool/Modes/BuildMode.cs
+++ b/Engine/Source/Programs/UnrealBuildTool/Modes/BuildMode.cs
@@ -219,7 +219,7 @@ namespace UnrealBuildTool
 						FileReference BaseLogFile = Log.OutputFile ?? new FileReference(BaseLogFileName);
 						FileReference RemoteLogFile = FileReference.Combine(BaseLogFile.Directory, BaseLogFile.GetFileNameWithoutExtension() + "_Remote.txt");
 
-						int RemoteIniIndex = TargetDesc.AdditionalArguments.FindIndex(x => x.ToLower().StartsWith("-remoteini=", StringComparison.Ordinal));
+						int RemoteIniIndex = TargetDesc.AdditionalArguments.FindIndex(x => x.StartsWith("-remoteini=", StringComparison.OrdinalIgnoreCase));
 						string? RemoteIni = RemoteIniIndex == -1 ? null : new string(TargetDesc.AdditionalArguments[RemoteIniIndex].Substring(11));
 						RemoteMac RemoteMac = new RemoteMac(TargetDesc.ProjectFile, Logger, RemoteIni);
 						if (!RemoteMac.Build(TargetDesc, RemoteLogFile, bSkipPreBuildTargets, Logger))
diff --git a/Engine/Source/Programs/UnrealBuildTool/Platform/Android/AndroidToolChain.cs b/Engine/Source/Programs/UnrealBuildTool/Platform/Android/AndroidToolChain.cs
--- a/Engine/Source/Programs/UnrealBuildTool/Platform/Android/AndroidToolChain.cs
+++ b/Engine/Source/Programs/UnrealBuildTool/Platform/Android/AndroidToolChain.cs
@@ -1127,7 +1127,9 @@ namespace UnrealBuildTool
 
 		public static string InlineArchName(string Pathname, UnrealArch Arch, bool bUseShortNames = false)
 		{
-			string FinalArch = "-" + Arch.ToString().ToLower();
+#pragma warning disable CA1308 // Normalize strings to uppercase - file path may be case-sensitive
+			string FinalArch = "-" + Arch.ToString().ToLowerInvariant();
+#pragma warning restore CA1308
 			if (bUseShortNames)
 			{
 				FinalArch = ShortArchNames[FinalArch];
diff --git a/Engine/Source/Programs/UnrealBuildTool/Platform/Android/UEDeployAndroid.cs b/Engine/Source/Programs/UnrealBuildTool/Platform/Android/UEDeployAndroid.cs
--- a/Engine/Source/Programs/UnrealBuildTool/Platform/Android/UEDeployAndroid.cs
+++ b/Engine/Source/Programs/UnrealBuildTool/Platform/Android/UEDeployAndroid.cs
@@ -2215,13 +2215,13 @@ namespace UnrealBuildTool
 
 		private static int ConvertDepthBufferIniValue(string IniValue)
 		{
-			switch (IniValue.ToLower())
+			switch (IniValue.ToUpperInvariant())
 			{
-				case "bits16":
+				case "BITS16":
 					return 16;
-				case "bits24":
+				case "BITS24":
 					return 24;
-				case "bits32":
+				case "BITS32":
 					return 32;
 				default:
 					return 0;
@@ -2230,23 +2230,23 @@ namespace UnrealBuildTool
 
 		private static string ConvertOrientationIniValue(string IniValue)
 		{
-			switch (IniValue.ToLower())
+			switch (IniValue.ToUpperInvariant())
 			{
-				case "portrait":
+				case "PORTRAIT":
 					return "portrait";
-				case "reverseportrait":
+				case "REVERSEPORTRAIT":
 					return "reversePortrait";
-				case "sensorportrait":
+				case "SENSORPORTRAIT":
 					return "sensorPortrait";
-				case "landscape":
+				case "LANDSCAPE":
 					return "landscape";
-				case "reverselandscape":
+				case "REVERSELANDSCAPE":
 					return "reverseLandscape";
-				case "sensorlandscape":
+				case "SENSORLANDSCAPE":
 					return "sensorLandscape";
-				case "sensor":
+				case "SENSOR":
 					return "sensor";
-				case "fullsensor":
+				case "FULLSENSOR":
 					return "fullSensor";
 				default:
 					return "landscape";
@@ -2274,33 +2274,33 @@ namespace UnrealBuildTool
 			bNeedLandscape = false;
 			bNeedPortrait = false;
 
-			switch (GetOrientation(Arch).ToLower())
+			switch (GetOrientation(Arch).ToUpperInvariant())
 			{
-				case "portrait":
+				case "PORTRAIT":
 					bNeedPortrait = true;
 					break;
-				case "reverseportrait":
+				case "REVERSEPORTRAIT":
 					bNeedPortrait = true;
 					break;
-				case "sensorportrait":
+				case "SENSORPORTRAIT":
 					bNeedPortrait = true;
 					break;
 
-				case "landscape":
+				case "LANDSCAPE":
 					bNeedLandscape = true;
 					break;
-				case "reverselandscape":
+				case "REVERSELANDSCAPE":
 					bNeedLandscape = true;
 					break;
-				case "sensorlandscape":
+				case "SENSORLANDSCAPE":
 					bNeedLandscape = true;
 					break;
 
-				case "sensor":
+				case "SENSOR":
 					bNeedPortrait = true;
 					bNeedLandscape = true;
 					break;
-				case "fullsensor":
+				case "FULLSENSOR":
 					bNeedPortrait = true;
 					bNeedLandscape = true;
 					break;
@@ -2653,12 +2653,12 @@ namespace UnrealBuildTool
 
 			string InstallLocation;
 			Ini.GetString("/Script/AndroidRuntimeSettings.AndroidRuntimeSettings", "InstallLocation", out InstallLocation);
-			switch (InstallLocation.ToLower())
+			switch (InstallLocation.ToUpperInvariant())
 			{
-				case "preferexternal":
+				case "PREFEREXTERNAL":
 					InstallLocation = "preferExternal";
 					break;
-				case "auto":
+				case "AUTO":
 					InstallLocation = "auto";
 					break;
 				default:
@@ -4635,7 +4635,9 @@ dependencies.implementation('androidx.tracing:tracing:1.2.0')
 					// LLDB (in Android Studio and AGDE) needs to see a file that is matching the filename on the device ("libUnreal.so") to resolve symbols,
 					// otherwise it will try to download libUnreal.so and fail due to adb pull not having permissions to access it.
 					// Newer lldb has "platform.plugin.remote-android.package-name" in lldb to avoid that.
-					string LLDBExtraSymbolsDirectory = Path.Combine(IntermediateAndroidPath, "LLDBSymbolsLibs", Arch.ToString().ToLower());
+#pragma warning disable CA1308 // Normalize strings to uppercase - file path should be created with the correct case
+					string LLDBExtraSymbolsDirectory = Path.Combine(IntermediateAndroidPath, "LLDBSymbolsLibs", Arch.ToString().ToLowerInvariant());
+#pragma warning restore CA1308
 
 					Directory.CreateDirectory(LLDBExtraSymbolsDirectory);
 
@@ -4650,8 +4652,10 @@ dependencies.implementation('androidx.tracing:tracing:1.2.0')
 				{
 					string SOName = AndroidToolChain.InlineArchName(OutputPath, Arch);
 					string SONameStripped = Path.Combine(Path.GetDirectoryName(SOName)!, Path.GetFileNameWithoutExtension(SOName) + "-stripped" + Path.GetExtension(SOName));
-					string LLDBExtraSymbolsDirectory = Path.Combine(IntermediateAndroidPath, "LLDBSymbolsLibs", Arch.ToString().ToLower());
-					
+#pragma warning disable CA1308 // Normalize strings to uppercase - file path should use the correct case
+					string LLDBExtraSymbolsDirectory = Path.Combine(IntermediateAndroidPath, "LLDBSymbolsLibs", Arch.ToString().ToLowerInvariant());
+#pragma warning restore CA1308
+
 					SafeDeleteFile(SONameStripped);
 					SafeDeleteFile(Path.Combine(LLDBExtraSymbolsDirectory, "libUnreal.so"));
 				}
diff --git a/Engine/Source/Programs/UnrealBuildTool/Platform/Apple/AppleToolChain.cs b/Engine/Source/Programs/UnrealBuildTool/Platform/Apple/AppleToolChain.cs
--- a/Engine/Source/Programs/UnrealBuildTool/Platform/Apple/AppleToolChain.cs
+++ b/Engine/Source/Programs/UnrealBuildTool/Platform/Apple/AppleToolChain.cs
@@ -685,8 +685,8 @@ namespace UnrealBuildTool
 
 			// tell swift to define PLATFORM_FOO, and tell C++ to define PLATFORM_FOO=1
 			// swift doesn't pass it's defines along to C++ when importing a C header, and it also doesn't have values for defines
-			Arguments.Add($"-DPLATFORM_{Target.Platform.ToString().ToUpper()}");
-			Arguments.Add($"-Xcc -DPLATFORM_{Target.Platform.ToString().ToUpper()}=1");
+			Arguments.Add($"-DPLATFORM_{Target.Platform.ToString().ToUpperInvariant()}");
+			Arguments.Add($"-Xcc -DPLATFORM_{Target.Platform.ToString().ToUpperInvariant()}=1");
 			if (bUseSwiftUIMain)
 			{
 				Arguments.Add("-DUE_USE_SWIFT_UI_MAIN");
diff --git a/Engine/Source/Programs/UnrealBuildTool/Platform/Apple/AppleToolChainSettings.cs b/Engine/Source/Programs/UnrealBuildTool/Platform/Apple/AppleToolChainSettings.cs
--- a/Engine/Source/Programs/UnrealBuildTool/Platform/Apple/AppleToolChainSettings.cs
+++ b/Engine/Source/Programs/UnrealBuildTool/Platform/Apple/AppleToolChainSettings.cs
@@ -82,14 +82,18 @@ namespace UnrealBuildTool
 		protected AppleToolChainSettings(string OSPrefix, string? SimulatorOSPrefix, string TargetOSName, string TargetOSVersion, bool bVerbose, ILogger Logger)
 		{
 			this.TargetOSName = TargetOSName;
-			PlatformDirName = OSPrefix.ToLower();
+#pragma warning disable CA1308 // Normalize strings to uppercase - file path should use the correct case
+			PlatformDirName = OSPrefix.ToLowerInvariant();
+#pragma warning restore CA1308
 
 			// set up directories
 			SDKDir = ApplePlatformSDK.GetPlatformSDKDirectory(OSPrefix);
 			if (SimulatorOSPrefix != null)
 			{
 				SimulatorSDKDir = ApplePlatformSDK.GetPlatformSDKDirectory(SimulatorOSPrefix);
-				SimulatorPlatformDirName = SimulatorOSPrefix.ToLower();
+#pragma warning disable CA1308 // Normalize strings to uppercase - file path should use the correct case
+				SimulatorPlatformDirName = SimulatorOSPrefix.ToLowerInvariant();
+#pragma warning restore CA1308
 			}
 
 			MinTargetVersion = TargetOSVersion;
diff --git a/Engine/Source/Programs/UnrealBuildTool/Platform/Mac/UEBuildMac.cs b/Engine/Source/Programs/UnrealBuildTool/Platform/Mac/UEBuildMac.cs
--- a/Engine/Source/Programs/UnrealBuildTool/Platform/Mac/UEBuildMac.cs
+++ b/Engine/Source/Programs/UnrealBuildTool/Platform/Mac/UEBuildMac.cs
@@ -200,11 +200,9 @@ namespace UnrealBuildTool
 			EngineIni.GetString("/Script/MacTargetPlatform.MacTargetSettings", SupportKey, out SupportedArchitecture);
 			EngineIni.GetString("/Script/MacTargetPlatform.MacTargetSettings", DefaultKey, out DefaultArchitecture);
 			EngineIni.GetBool("/Script/MacTargetPlatform.MacTargetSettings", "bBuildAllSupportedOnBuildMachine", out bBuildAllSupportedOnBuildMachine);
-			SupportedArchitecture = SupportedArchitecture.ToLower();
-			DefaultArchitecture = DefaultArchitecture.ToLower();
 
-			bool bSupportsArm64 = SupportedArchitecture.Contains("universal", StringComparison.Ordinal) || SupportedArchitecture.Contains("apple", StringComparison.Ordinal);
-			bool bSupportsX86 = SupportedArchitecture.Contains("universal", StringComparison.Ordinal) || SupportedArchitecture.Contains("intel", StringComparison.Ordinal);
+			bool bSupportsArm64 = SupportedArchitecture.Contains("universal", StringComparison.OrdinalIgnoreCase) || SupportedArchitecture.Contains("apple", StringComparison.OrdinalIgnoreCase);
+			bool bSupportsX86 = SupportedArchitecture.Contains("universal", StringComparison.OrdinalIgnoreCase) || SupportedArchitecture.Contains("intel", StringComparison.OrdinalIgnoreCase);
 
 			// make sure we found a good value
 			if (!bSupportsArm64 && !bSupportsX86)
@@ -228,12 +226,12 @@ namespace UnrealBuildTool
 					Architectures.Add(UnrealArch.X64);
 				}
 			}
-			else if (DefaultArchitecture.Contains("host", StringComparison.Ordinal))
+			else if (DefaultArchitecture.Contains("host", StringComparison.OrdinalIgnoreCase))
 			{
 				// if we don't support Arm, then always use X64, otherwise use whatever the host arch is
 				Architectures.Add(bSupportsArm64 ? UnrealArch.Host.Value : UnrealArch.X64);
 			}
-			else if (DefaultArchitecture.Contains("apple", StringComparison.Ordinal))
+			else if (DefaultArchitecture.Contains("apple", StringComparison.OrdinalIgnoreCase))
 			{
 				if (!bSupportsArm64)
 				{
@@ -241,7 +239,7 @@ namespace UnrealBuildTool
 				}
 				Architectures.Add(UnrealArch.Arm64);
 			}
-			else if (DefaultArchitecture.Contains("intel", StringComparison.Ordinal))
+			else if (DefaultArchitecture.Contains("intel", StringComparison.OrdinalIgnoreCase))
 			{
 				if (!bSupportsX86)
 				{
diff --git a/Engine/Source/Programs/UnrealBuildTool/Platform/Windows/AppXManifestGeneratorBase.cs b/Engine/Source/Programs/UnrealBuildTool/Platform/Windows/AppXManifestGeneratorBase.cs
--- a/Engine/Source/Programs/UnrealBuildTool/Platform/Windows/AppXManifestGeneratorBase.cs
+++ b/Engine/Source/Programs/UnrealBuildTool/Platform/Windows/AppXManifestGeneratorBase.cs
@@ -195,9 +195,9 @@ namespace UnrealBuildTool
 				return DefaultValue;
 			}
 
-			ResultStr = ResultStr.Trim().ToLower();
+			ResultStr = ResultStr.Trim().ToUpperInvariant();
 
-			return ResultStr == "true" || ResultStr == "1" || ResultStr == "yes";
+			return ResultStr == "TRUE" || ResultStr == "1" || ResultStr == "YES";
 		}
 
 		/// <summary>
diff --git a/Engine/Source/Programs/UnrealBuildTool/Platform/Windows/VCToolChain.cs b/Engine/Source/Programs/UnrealBuildTool/Platform/Windows/VCToolChain.cs
--- a/Engine/Source/Programs/UnrealBuildTool/Platform/Windows/VCToolChain.cs
+++ b/Engine/Source/Programs/UnrealBuildTool/Platform/Windows/VCToolChain.cs
@@ -1666,7 +1666,9 @@ namespace UnrealBuildTool
 				// Embed a dependency on the SxS assembly manifest. this is for dependent runtime DLLs
 				if (!Target.Architecture.bIsX64)
 				{
-					string AssemblyName = Target.Architecture.ToString().ToLower();
+#pragma warning disable CA1308 // Normalize strings to uppercase - external argument should use the correct case
+					string AssemblyName = Target.Architecture.ToString().ToLowerInvariant();
+#pragma warning restore CA1308
 					Arguments.Add($"/MANIFESTDEPENDENCY:\"name='{AssemblyName}' processorArchitecture='arm64' version='1.0.0.0' type='win32'\"");
 				}
 			}
@@ -2655,7 +2657,9 @@ namespace UnrealBuildTool
 		{
 			if (Target.WindowsPlatform.ManifestFile != null && !Target.bShouldCompileAsDLL && !Target.Architecture.bIsX64)
 			{
-				string AssemblyName = Target.Architecture.ToString().ToLower();
+#pragma warning disable CA1308 // Normalize strings to uppercase - file path should use the correct case
+				string AssemblyName = Target.Architecture.ToString().ToLowerInvariant();
+#pragma warning restore CA1308
 
 				// collect the DLLs that should be included in the SxS manifest
 				DirectoryReference ManifestDir = DirectoryReference.Combine(MakefileBuilder.Makefile.ExecutableFile.Directory, AssemblyName);
@@ -2773,7 +2777,9 @@ namespace UnrealBuildTool
 			// add a dependency on the SxS manifest
 			if (Target.WindowsPlatform.ManifestFile != null && !Target.bShouldCompileAsDLL && !Target.Architecture.bIsX64)
 			{
-				string AssemblyName = Target.Architecture.ToString().ToLower();
+#pragma warning disable CA1308 // Normalize strings to uppercase - file path should use the correct case
+				string AssemblyName = Target.Architecture.ToString().ToLowerInvariant();
+#pragma warning restore CA1308
 				FileReference SrcFile = GetSxSIntermediateManifest(Target, ExeDir);
 				FileReference DstFile = FileReference.Combine(ExeDir, AssemblyName, AssemblyName + ".manifest");
 
diff --git a/Engine/Source/Programs/UnrealBuildTool/ProjectFiles/VisualStudioCode/VSCodeProjectFileGenerator.cs b/Engine/Source/Programs/UnrealBuildTool/ProjectFiles/VisualStudioCode/VSCodeProjectFileGenerator.cs
--- a/Engine/Source/Programs/UnrealBuildTool/ProjectFiles/VisualStudioCode/VSCodeProjectFileGenerator.cs
+++ b/Engine/Source/Programs/UnrealBuildTool/ProjectFiles/VisualStudioCode/VSCodeProjectFileGenerator.cs
@@ -221,7 +221,7 @@ namespace UnrealBuildTool
 
 			public void AddField(string Name, bool Value)
 			{
-				Lines.Add(TabString + Quoted(JsonWriter.EscapeString(Name)) + ": " + Value.ToString().ToLower() + ",");
+				Lines.Add(TabString + Quoted(JsonWriter.EscapeString(Name)) + ": " + (Value ? "true" : "false") + ",");
 			}
 
 			public void AddField(string Name, string Value)
@@ -1039,7 +1039,9 @@ namespace UnrealBuildTool
 							OutFile.AddField("command", MakePathString(Unreal.DotnetPath));
 							OutFile.BeginArray("args");
 							{
-								OutFile.AddUnnamedField(Command.ToLower());
+#pragma warning disable CA1308 // Normalize strings to uppercase - preserve existing behaviour as it may be case-sensitive
+								OutFile.AddUnnamedField(Command.ToLowerInvariant());
+#pragma warning restore CA1308
 
 								OutFile.AddUnnamedField("--configuration");
 								OutFile.AddUnnamedField(BuildProduct.Config.ToString());
diff --git a/Engine/Source/Programs/UnrealBuildTool/ProjectFiles/Xcode/XcodeProject.cs b/Engine/Source/Programs/UnrealBuildTool/ProjectFiles/Xcode/XcodeProject.cs
--- a/Engine/Source/Programs/UnrealBuildTool/ProjectFiles/Xcode/XcodeProject.cs
+++ b/Engine/Source/Programs/UnrealBuildTool/ProjectFiles/Xcode/XcodeProject.cs
@@ -1595,7 +1595,7 @@ namespace UnrealBuildTool.XcodeProjectXcconfig
 				string DefaultStageDir = bIsEngineBuild ? "" : "${UE_PROJECT_DIR}/Saved/StagedBuilds/${UE_TARGET_PLATFORM_NAME}";
 				string SyncSourceSubdir = (Platform == UnrealTargetPlatform.Mac) ? "" : "/cookeddata";
 				string SyncDestSubdir = (Platform == UnrealTargetPlatform.Mac) ? "/UE" : "/cookeddata";
-				string ExecutableKey = $"UE_{Platform.ToString().ToUpper()}_EXECUTABLE_NAME";
+				string ExecutableKey = $"UE_{Platform.ToString().ToUpperInvariant()}_EXECUTABLE_NAME";
 
 				CopyScript.AddRange(new string[]
 				{
@@ -2250,7 +2250,7 @@ namespace UnrealBuildTool.XcodeProjectXcconfig
 				{
 					Xcconfig.AppendLine($"ASSETCATALOG_COMPILER_ALTERNATE_APPICON_NAMES = {AssetCatalogAlternates}");
 				}
-				Xcconfig.AppendLine($"CURRENT_PROJECT_VERSION = $(UE_{Platform.ToString().ToUpper()}_BUILD_VERSION)");
+				Xcconfig.AppendLine($"CURRENT_PROJECT_VERSION = $(UE_{Platform.ToString().ToUpperInvariant()}_BUILD_VERSION)");
 				Xcconfig.AppendLine($"MARKETING_VERSION = {MarketingVersion}");
 				Xcconfig.AppendLine($"INFOPLIST_KEY_LSApplicationCategoryType = {AppCategory}");
 			}
@@ -2311,7 +2311,7 @@ namespace UnrealBuildTool.XcodeProjectXcconfig
 
 				string ExecutableName = AppleExports.MakeBinaryFileName(Config.ExeName, TargetRules.DecoratedSeparator, Platform, Config.BuildConfig, TargetRules.Architectures, TargetRules.UndecoratedConfiguration, null);
 				string ExetuableSubPath = FileReference.Combine(ConfigBuildDir, ExecutableName).MakeRelativeTo(ConfigBuildDir);
-				string ExecutableKey = $"UE_{Platform.ToString().ToUpper()}_EXECUTABLE_NAME";
+				string ExecutableKey = $"UE_{Platform.ToString().ToUpperInvariant()}_EXECUTABLE_NAME";
 
 				// we want to make Foo.app, not FooGame.app, since we added on the target type when making targets
 				string PerTargetTypeProductName = UnrealData.ProductName;
@@ -2739,7 +2739,7 @@ namespace UnrealBuildTool.XcodeProjectXcconfig
 			{
 				// allow Xcode to generate the final plist file from our input, some INFOPLIST settings and other settings 
 				Xcconfig.AppendLine("GENERATE_INFOPLIST_FILE = YES");
-				Xcconfig.AppendLine($"CURRENT_PROJECT_VERSION = $(UE_{Platform.ToString().ToUpper()}_BUILD_VERSION)");
+				Xcconfig.AppendLine($"CURRENT_PROJECT_VERSION = $(UE_{Platform.ToString().ToUpperInvariant()}_BUILD_VERSION)");
 				Xcconfig.AppendLine($"MARKETING_VERSION = {MarketingVersion}");
 			}
 
@@ -2799,7 +2799,7 @@ namespace UnrealBuildTool.XcodeProjectXcconfig
 
 				string ExecutableName = Name;
 				string ExecutableSubPath = FileReference.Combine(ConfigBuildDir, ExecutableName).MakeRelativeTo(ConfigBuildDir);
-				string ExecutableKey = $"UE_{Platform.ToString().ToUpper()}_EXECUTABLE_NAME";
+				string ExecutableKey = $"UE_{Platform.ToString().ToUpperInvariant()}_EXECUTABLE_NAME";
 
 				// we want to make Foo.app, not FooGame.app, since we added on the target type when making targets
 				string PerTargetTypeProductName = UnrealData.ProductName;
diff --git a/Engine/Source/Programs/UnrealBuildTool/System/UnrealPluginLanguage.cs b/Engine/Source/Programs/UnrealBuildTool/System/UnrealPluginLanguage.cs
--- a/Engine/Source/Programs/UnrealBuildTool/System/UnrealPluginLanguage.cs
+++ b/Engine/Source/Programs/UnrealBuildTool/System/UnrealPluginLanguage.cs
@@ -706,7 +706,7 @@ namespace UnrealBuildTool
 			StringBuilder Text = new StringBuilder();
 			foreach (KeyValuePair<string, bool> Variable in Context.BoolVariables)
 			{
-				Text.AppendLine(String.Format("\tbool {0} = {1}", Variable.Key, Variable.Value.ToString().ToLower()));
+				Text.AppendLine(String.Format("\tbool {0} = {1}", Variable.Key, Variable.Value ? "true" : "false"));
 			}
 			foreach (KeyValuePair<string, int> Variable in Context.IntVariables)
 			{
@@ -780,7 +780,7 @@ namespace UnrealBuildTool
 				}
 
 				// Replace the variable, or skip past it
-				Result = Result.Substring(0, Idx) + Value.ToString().ToLower() + Result.Substring(EndIdx + 1);
+				Result = Result.Substring(0, Idx) + (Value ? "true" : "false") + Result.Substring(EndIdx + 1);
 			}
 			for (int Idx = Result.IndexOf("$I(", StringComparison.Ordinal); Idx != -1; Idx = Result.IndexOf("$I(", Idx, StringComparison.Ordinal))
 			{
@@ -2566,7 +2566,9 @@ namespace UnrealBuildTool
 							string Source = GetAttribute(CurrentContext, Node, "source", true, false, "");
 							if (Result != null)
 							{
-								string Value = Source.ToLower();
+#pragma warning disable CA1308 // Normalize strings to uppercase - specifically implementing a method of lowercasing a string
+								string Value = Source.ToLowerInvariant();
+#pragma warning restore CA1308
 								if (Result == "Output")
 								{
 									GlobalContext.StringVariables["Output"] = Value;
@@ -2585,7 +2587,7 @@ namespace UnrealBuildTool
 							string Source = GetAttribute(CurrentContext, Node, "source", true, false, "");
 							if (Result != null)
 							{
-								string Value = Source.ToUpper();
+								string Value = Source.ToUpperInvariant();
 								if (Result == "Output")
 								{
 									GlobalContext.StringVariables["Output"] = Value;
