UnrealBuildTool: Refactor VS solution file building into its own class so it can be used for generating the Automation solution rather than just calling GenerateProjectFiles and getting logic we don't want. Generate the Automation solution for VS2026 if it's found on the system, even if project file generation is requested for VS2022.

#Synced-CL 50040876

#was-backout 50047636

#Synced-CL 50156122

diff --git a/Engine/Source/Programs/UnrealBuildTool/ProjectFiles/ProjectFileGenerator.cs b/Engine/Source/Programs/UnrealBuildTool/ProjectFiles/ProjectFileGenerator.cs
--- a/Engine/Source/Programs/UnrealBuildTool/ProjectFiles/ProjectFileGenerator.cs
+++ b/Engine/Source/Programs/UnrealBuildTool/ProjectFiles/ProjectFileGenerator.cs
@@ -4215,7 +4215,7 @@ namespace UnrealBuildTool
 		/// List of top-level folders in the primary project file
 		protected PrimaryProjectFolder RootFolder;
 
-		private bool AllProjectFilesModified = true;
+		protected bool AllProjectFilesModified = true;
 		private readonly List<ProjectFile> AllProjectFilesCached = [];
 	}
 
diff --git a/Engine/Source/Programs/UnrealBuildTool/ProjectFiles/VisualStudio/VCProjectFileGenerator.cs b/Engine/Source/Programs/UnrealBuildTool/ProjectFiles/VisualStudio/VCProjectFileGenerator.cs
--- a/Engine/Source/Programs/UnrealBuildTool/ProjectFiles/VisualStudio/VCProjectFileGenerator.cs
+++ b/Engine/Source/Programs/UnrealBuildTool/ProjectFiles/VisualStudio/VCProjectFileGenerator.cs
@@ -4,6 +4,7 @@ using System;
 using System.Collections.Generic;
 using System.IO;
 using System.Linq;
+using System.Runtime.Versioning;
 using System.Security.Cryptography;
 using System.Text;
 using System.Xml;
@@ -11,6 +12,7 @@ using System.Xml.Linq;
 using EpicGames.Core;
 using Microsoft.Extensions.Logging;
 using UnrealBuildBase;
+using UnrealBuildTool.ProjectFiles.VisualStudio;
 
 namespace UnrealBuildTool
 {
@@ -446,21 +448,29 @@ namespace UnrealBuildTool
 			if (Settings.ProjectFileFormat == VCProjectFileFormat.VisualStudio2022)
 			{
 				Logger.LogInformation(
-					"Visual Studio 2022 does not support .NET 10.0 C# projects, these projects will not be added to the generated automation solution as the current dotnet runtime version is '{Version}'. Compiling with an older dotnet SDK from Visual Studio can cause errors.",
+					"Visual Studio 2022 does not support .NET 10.0 C# projects, and the current dotnet runtime version is '{Version}'. Compiling with an older dotnet SDK from Visual Studio can cause errors.",
 					System.Runtime.InteropServices.RuntimeInformation.FrameworkDescription);
-				Logger.LogInformation("Please generate the Visual Studio 2026 solution if .NET 10.0 C# project support is required.");
-				return;
+
+				// Some users can't use VS2026 yet (certain platforms are VS2022 only) so we should accomodate for them by generating the Automation solution for 2026 if available.
+				if (WindowsPlatform.HasCompiler(WindowsCompiler.VisualStudio2026, UnrealArch.X64, Logger))
+				{
+					Logger.LogInformation("The {SolutionName} solution will be generated for Visual Studio 2026.", $"Automation_{PrimaryProjectName}");
+				}
+				else
+				{
+					Logger.LogInformation("The {SolutionName} solution will not be generated.", $"Automation_{PrimaryProjectName}");
+					return;
+				}
 			}
 
 			// We expect this to be called as part of a call that also writes a C++ solution too, so we have to make sure to put
 			// this property back to its original value when done.
-			string OriginalPrimaryProjectName = PrimaryProjectName;
 			PrimaryProjectFolder OriginalRootFolder = RootFolder;
 			List<ProjectFile> OriginalAutomationProjectFiles = [.. AutomationProjectFiles];
 			List<ProjectFile> OriginalOtherProjectFiles = [.. OtherProjectFiles];
+
 			try
 			{
-				PrimaryProjectName = $"Automation_{PrimaryProjectName}";
 				RootFolder = new(this, "<Root>");
 				AutomationProjectFiles.Clear();
 				OtherProjectFiles.Clear();
@@ -469,17 +479,36 @@ namespace UnrealBuildTool
 
 				using (Timeline.ScopeEvent("Writing automation project files", Logger))
 				{
-					WriteProjectFiles(PlatformProjectGenerators, Logger);
+					string SolutionFileName = $"Automation_{PrimaryProjectName}.sln";
+					VisualStudioSolutionBuilder SolutionBuilder = new(RootFolder, AllProjectFiles, Logger);
+
+					SolutionBuilder.WritePreamble(Settings.ProjectFileFormat);
+					SolutionBuilder.WriteProjectFolders();
+					SolutionBuilder.WriteProjects();
+
+					SolutionBuilder.WriteGlobalStart();
+					SolutionBuilder.WriteSolutionProperties();
+					SolutionBuilder.WriteNestedProjects();
+					SolutionBuilder.WriteGlobalEnd();
+
+					FileReference SolutionFilePath = FileReference.Combine(PrimaryProjectPath, SolutionFileName);
+					bool success = WriteFileIfChanged(SolutionFilePath.FullName, SolutionBuilder.ToString(), Logger);
+
+					// Generate a slnx for VS2026 until properly supported
+					if (success && Environment.Version.Major >= 10)
+					{
+						ConvertSolutionFileToSlnx(SolutionFileName, SolutionFilePath, Logger);
+					}
 				}
 			}
 			finally
 			{
-				PrimaryProjectName = OriginalPrimaryProjectName;
 				RootFolder = OriginalRootFolder;
 				AutomationProjectFiles.Clear();
 				AutomationProjectFiles.AddRange(OriginalAutomationProjectFiles);
 				OtherProjectFiles.Clear();
 				OtherProjectFiles.AddRange(OriginalOtherProjectFiles);
+				AllProjectFilesModified = true;
 			}
 		}
 
@@ -548,27 +577,6 @@ namespace UnrealBuildTool
 			return SolutionConfigName;
 		}
 
-		static IDictionary<PrimaryProjectFolder, Guid> GenerateProjectFolderGuids(PrimaryProjectFolder RootFolder)
-		{
-			IDictionary<PrimaryProjectFolder, Guid> Guids = new Dictionary<PrimaryProjectFolder, Guid>();
-			foreach (PrimaryProjectFolder Folder in RootFolder.SubFolders)
-			{
-				GenerateProjectFolderGuids("UE5", Folder, Guids);
-			}
-			return Guids;
-		}
-
-		static void GenerateProjectFolderGuids(string ParentPath, PrimaryProjectFolder Folder, IDictionary<PrimaryProjectFolder, Guid> Guids)
-		{
-			string Path = String.Format("{0}/{1}", ParentPath, Folder.FolderName);
-			Guids[Folder] = MakeMd5Guid(Encoding.UTF8.GetBytes(Path));
-
-			foreach (PrimaryProjectFolder SubFolder in Folder.SubFolders)
-			{
-				GenerateProjectFolderGuids(Path, SubFolder, Guids);
-			}
-		}
-
 		static Guid MakeMd5Guid(byte[] Input)
 		{
 #pragma warning disable CA5351 // Not used for cryptography, would be disruptive to users to change hash algo
@@ -769,120 +777,13 @@ namespace UnrealBuildTool
 			string SolutionFileName = PrimaryProjectName + ".sln";
 
 			// Setup solution file content
-			StringBuilder VCSolutionFileContent = new StringBuilder();
+			VisualStudioSolutionBuilder SolutionBuilder = new(RootFolder, AllProjectFiles, Logger);
 
-			// Solution file header. Note that a leading newline is required for file type detection to work correclty in the shell.
-			if (Settings.ProjectFileFormat == VCProjectFileFormat.VisualStudio2022)
-			{
-				VCSolutionFileContent.AppendLine("Microsoft Visual Studio Solution File, Format Version 12.00");
-				VCSolutionFileContent.AppendLine("# Visual Studio Version 17");
-				VCSolutionFileContent.AppendLine("VisualStudioVersion = 17.0.31314.256");
-				VCSolutionFileContent.AppendLine("MinimumVisualStudioVersion = 10.0.40219.1");
-			}
-			else if (Settings.ProjectFileFormat == VCProjectFileFormat.VisualStudio2026)
-			{
-				VCSolutionFileContent.AppendLine("Microsoft Visual Studio Solution File, Format Version 12.00");
-				VCSolutionFileContent.AppendLine("# Visual Studio Version 18");
-				VCSolutionFileContent.AppendLine("VisualStudioVersion = 18.0.11205.157");
-				VCSolutionFileContent.AppendLine("MinimumVisualStudioVersion = 10.0.40219.1");
-			}
-			else
-			{
-				throw new BuildException("Unexpected ProjectFileFormat '{0}'", Settings.ProjectFileFormat);
-			}
-
-			IDictionary<PrimaryProjectFolder, Guid> ProjectFolderGuids = GenerateProjectFolderGuids(RootFolder);
-
-			// check if the give folder has any projects that will be put into the solution
-			System.Func<PrimaryProjectFolder, bool>? HasProjectFunc = null;
-			HasProjectFunc = (ProjectFolder) =>
-			{
-				foreach (MSBuildProjectFile ChildProject in ProjectFolder.ChildProjects)
-				{
-					if (AllProjectFiles.Contains(ChildProject))
-					{
-						return true;
-					}
-				}
-
-				foreach (PrimaryProjectFolder SubFolder in ProjectFolder.SubFolders)
-				{
-					if (HasProjectFunc!(SubFolder))
-					{
-						return true;
-					}
-				}
-				return false;
-			};
+			SolutionBuilder.WritePreamble(Settings.ProjectFileFormat);
 
 			// Solution folders, files and project entries
 			{
-				// This the GUID that Visual Studio uses to identify a solution folder
-				string SolutionFolderEntryGUID = "{2150E333-8FDC-42A3-9474-1A3956D46DE8}";
-
-				// Solution folders
-				{
-					IEnumerable<PrimaryProjectFolder> AllSolutionFolders = ProjectFolderGuids.Keys.OrderBy(Folder => Folder.FolderName).ThenBy(Folder => ProjectFolderGuids[Folder]);
-					foreach (PrimaryProjectFolder CurFolder in AllSolutionFolders)
-					{
-						// if this folder has no projects anywhere under it then skip it
-						if (!HasProjectFunc(CurFolder))
-						{
-							continue;
-						}
-						string FolderGUIDString = ProjectFolderGuids[CurFolder].ToString("B").ToUpperInvariant();
-						VCSolutionFileContent.AppendLine("Project(\"" + SolutionFolderEntryGUID + "\") = \"" + CurFolder.FolderName + "\", \"" + CurFolder.FolderName + "\", \"" + FolderGUIDString + "\"");
-
-						// Add any files that are inlined right inside the solution folder
-						if (CurFolder.Files.Count > 0)
-						{
-							VCSolutionFileContent.AppendLine("	ProjectSection(SolutionItems) = preProject");
-							foreach (string CurFile in CurFolder.Files)
-							{
-								// Syntax is:  <relative file path> = <relative file path>
-								VCSolutionFileContent.AppendLine("		" + CurFile + " = " + CurFile);
-							}
-							VCSolutionFileContent.AppendLine("	EndProjectSection");
-						}
-
-						VCSolutionFileContent.AppendLine("EndProject");
-					}
-				}
-
-				// Project files
-				foreach (MSBuildProjectFile CurProject in AllProjectFiles)
-				{
-					// Visual Studio uses different GUID types depending on the project type
-					string ProjectTypeGUID = CurProject.ProjectTypeGUID;
-
-					// NOTE: The project name in the solution doesn't actually *have* to match the project file name on disk.  However,
-					//       we prefer it when it does match so we use the actual file name here.
-					string ProjectNameInSolution = CurProject.ProjectFilePath.GetFileNameWithoutExtension();
-
-					// Use the existing project's GUID that's already known to us
-					string ProjectGUID = CurProject.ProjectGUID.ToString("B").ToUpperInvariant();
-
-					VCSolutionFileContent.AppendLine($"Project(\"{ProjectTypeGUID}\") = \"{ProjectNameInSolution}\", \"{CurProject.ProjectFilePath.MakeRelativeTo(ProjectFileGenerator.PrimaryProjectPath)}\", \"{ProjectGUID}\"");
-
-					if (!CurProject.IsStubProject)
-					{
-						if (CurProject.DependsOnProjects.Count > 0)
-						{
-							VCSolutionFileContent.AppendLine("\tProjectSection(ProjectDependencies) = postProject");
-
-							// Setup any addition dependencies this project has...
-							foreach (ProjectFile DependsOnProject in CurProject.DependsOnProjects)
-							{
-								string DependsOnProjectGUID = ((MSBuildProjectFile)DependsOnProject).ProjectGUID.ToString("B").ToUpperInvariant();
-								VCSolutionFileContent.AppendLine($"\t\t{DependsOnProjectGUID} = {DependsOnProjectGUID}");
-							}
-
-							VCSolutionFileContent.AppendLine("\tEndProjectSection");
-						}
-					}
-
-					VCSolutionFileContent.AppendLine("EndProject");
-				}
+				SolutionBuilder.WriteProjectFolders();
 
 				// Get the path to the visualizers file. Try to make it relative to the solution directory, but fall back to a full path if it's a foreign project.
 				FileReference VisualizersFile = FileReference.Combine(Unreal.EngineDirectory, "Extras", "VisualStudioDebugging", "Unreal.natvis");
@@ -890,172 +791,65 @@ namespace UnrealBuildTool
 				FileReference VisualizersJustMyCodeFile = FileReference.Combine(Unreal.EngineDirectory, "Extras", "VisualStudioDebugging", "Unreal.natjmc");
 
 				// Add the visualizers at the solution level. Doesn't seem to be picked up from a makefile project in VS2017 15.8.5.
-				VCSolutionFileContent.AppendLine(String.Format("Project(\"{0}\") = \"Visualizers\", \"Visualizers\", \"{{1CCEC849-CC72-4C59-8C36-2F7C38706D4C}}\"", SolutionFolderEntryGUID));
-				VCSolutionFileContent.AppendLine("\tProjectSection(SolutionItems) = preProject");
-				VCSolutionFileContent.AppendLine("\t\t{0} = {0}", VisualizersFile.MakeRelativeTo(PrimaryProjectPath));
-				VCSolutionFileContent.AppendLine("\t\t{0} = {0}", VisualizersStepFile.MakeRelativeTo(PrimaryProjectPath));
-				VCSolutionFileContent.AppendLine("\t\t{0} = {0}", VisualizersJustMyCodeFile.MakeRelativeTo(PrimaryProjectPath));
-				VCSolutionFileContent.AppendLine("\tEndProjectSection");
-				VCSolutionFileContent.AppendLine("EndProject");
+				PrimaryProjectFolder visualizers = new PrimaryProjectFolder(this, "Visualizers");
+				visualizers.Files.AddRange([
+					VisualizersFile.MakeRelativeTo(PrimaryProjectPath),
+					VisualizersStepFile.MakeRelativeTo(PrimaryProjectPath),
+					VisualizersJustMyCodeFile.MakeRelativeTo(PrimaryProjectPath),
+				]);
+				SolutionBuilder.WriteProjectStart(
+					VisualStudioSolutionBuilder.SolutionFolderGuid,
+					"Visualizers",
+					"Visualizers",
+					Guid.Parse("1CCEC849-CC72-4C59-8C36-2F7C38706D4C")
+				);
+				SolutionBuilder.WriteProjectSolutionItems(visualizers);
+				SolutionBuilder.WriteProjectEnd();
+
+				SolutionBuilder.WriteProjects();
 			}
 
-			// Solution configuration platforms.  This is just a list of all of the platforms and configurations that
+			// Solution configuration platforms. This is just a list of all of the platforms and configurations that
 			// appear in Visual Studio's build configuration selector.
 			List<VCSolutionConfigCombination> SolutionConfigCombinations;
 
 			// The "Global" section has source control, solution configurations, project configurations,
 			// preferences, and project hierarchy data
 			{
-				VCSolutionFileContent.AppendLine("Global");
+				SolutionBuilder.WriteGlobalStart();
 				{
 					HashSet<UnrealTargetPlatform> PlatformsValidForProjects;
-					{
-						VCSolutionFileContent.AppendLine("	GlobalSection(SolutionConfigurationPlatforms) = preSolution");
+					CollectSolutionConfigurations(SupportedConfigurations, SupportedPlatforms, AllProjectFiles, bMakeProjectPerTarget,
+						out PlatformsValidForProjects, out SolutionConfigCombinations);
 
-						CollectSolutionConfigurations(SupportedConfigurations, SupportedPlatforms, AllProjectFiles, bMakeProjectPerTarget,
-							out PlatformsValidForProjects, out SolutionConfigCombinations);
-
-						HashSet<string> AppendedSolutionConfigAndPlatformNames = new HashSet<string>(StringComparer.InvariantCultureIgnoreCase);
-						foreach (VCSolutionConfigCombination SolutionConfigCombination in SolutionConfigCombinations)
-						{
-							// We alias "Game" and "Program" to both have the same solution configuration, so we're careful not to add the same combination twice.
-							if (!AppendedSolutionConfigAndPlatformNames.Contains(SolutionConfigCombination.VCSolutionConfigAndPlatformName))
-							{
-								VCSolutionFileContent.AppendLine("		" + SolutionConfigCombination.VCSolutionConfigAndPlatformName + " = " + SolutionConfigCombination.VCSolutionConfigAndPlatformName);
-								AppendedSolutionConfigAndPlatformNames.Add(SolutionConfigCombination.VCSolutionConfigAndPlatformName);
-							}
-						}
-
-						VCSolutionFileContent.AppendLine("	EndGlobalSection");
-					}
+					SolutionBuilder.WriteSolutionConfigurationPlatforms(SolutionConfigCombinations);
 
 					// Embed UnrealVS section, which is parsed by that VSPackage (Extension) to know how to handle this solution
-					{
-						string UnrealVSGuid = "ddbf523f-7eb6-4887-bd51-85a714ff87eb";
-						VCSolutionFileContent.AppendLine("\t# UnrealVS Section");
-						VCSolutionFileContent.AppendLine("\tGlobalSection({0}) = preSolution", UnrealVSGuid);
-						VCSolutionFileContent.AppendLine("\t\tAvailablePlatforms={0}", String.Join(";", PlatformsValidForProjects.Select(platform => platform.ToString())));
-						VCSolutionFileContent.AppendLine("\tEndGlobalSection");
-					}
+					SolutionBuilder.WriteUnrealVs(PlatformsValidForProjects);
 
 					// Assign each project's "project configuration" to our "solution platform + configuration" pairs.  This
 					// also sets up which projects are actually built when building the solution.
-					{
-						VCSolutionFileContent.AppendLine("	GlobalSection(ProjectConfigurationPlatforms) = postSolution");
-
-						foreach (MSBuildProjectFile CurProject in AllProjectFiles)
-						{
-							foreach (VCSolutionConfigCombination SolutionConfigCombination in SolutionConfigCombinations)
-							{
-								// Get the context for the current solution context
-								MSBuildProjectContext? ProjectContext = CurProject.GetMatchingProjectContext(SolutionConfigCombination.TargetConfigurationName, SolutionConfigCombination.Configuration, SolutionConfigCombination.Platform, PlatformProjectGenerators, SolutionConfigCombination.Architecture, Logger);
-
-								if (ProjectContext == null)
-								{
-									continue;
-								}
-
-								// Override the configuration to build for UBT
-								if (Settings.bBuildUBTInDebug && CurProject == UBTProject)
-								{
-									ProjectContext.ConfigurationName = "Debug";
-								}
-
-								// Write the solution mapping (e.g.  "{4232C52C-680F-4850-8855-DC39419B5E9B}.Debug|iOS.ActiveCfg = iOS_Debug|Win32")
-								string CurProjectGUID = CurProject.ProjectGUID.ToString("B").ToUpperInvariant();
-								VCSolutionFileContent.AppendLine("		{0}.{1}.ActiveCfg = {2}", CurProjectGUID, SolutionConfigCombination.VCSolutionConfigAndPlatformName, ProjectContext.Name);
-								if (ProjectContext.bBuildByDefault)
-								{
-									VCSolutionFileContent.AppendLine("		{0}.{1}.Build.0 = {2}", CurProjectGUID, SolutionConfigCombination.VCSolutionConfigAndPlatformName, ProjectContext.Name);
-									if (ProjectContext.bDeployByDefault)
-									{
-										VCSolutionFileContent.AppendLine("		{0}.{1}.Deploy.0 = {2}", CurProjectGUID, SolutionConfigCombination.VCSolutionConfigAndPlatformName, ProjectContext.Name);
-									}
-								}
-							}
-						}
+					bool DebugPredicate(MSBuildProjectFile Project) => Settings.bBuildUBTInDebug && Project == UBTProject;
+					SolutionBuilder.WriteProjectConfigurationPlatforms(SolutionConfigCombinations, PlatformProjectGenerators, DebugPredicate);
 
-						VCSolutionFileContent.AppendLine("	EndGlobalSection");
-					}
-
-					// Setup other solution properties
-					{
-						// HideSolutionNode sets whether or not the top-level solution entry is completely hidden in the UI.
-						// We don't want that, as we need users to be able to right click on the solution tree item.
-						VCSolutionFileContent.AppendLine("	GlobalSection(SolutionProperties) = preSolution");
-						VCSolutionFileContent.AppendLine("		HideSolutionNode = FALSE");
-						VCSolutionFileContent.AppendLine("	EndGlobalSection");
-					}
-
-					// Solution directory hierarchy
-					{
-						VCSolutionFileContent.AppendLine("	GlobalSection(NestedProjects) = preSolution");
+					SolutionBuilder.WriteSolutionProperties();
 
-						// Every entry in this section is in the format "Guid1 = Guid2".  Guid1 is the child project (or solution
-						// filter)'s GUID, and Guid2 is the solution filter directory to parent the child project (or solution
-						// filter) to.  This sets up the hierarchical solution explorer tree for all solution folders and projects.
-
-						void FolderProcessorFunction(StringBuilder LocalVCSolutionFileContent, List<PrimaryProjectFolder> LocalPrimaryProjectFolders)
-							{
-								foreach (PrimaryProjectFolder CurFolder in LocalPrimaryProjectFolders)
-								{
-									string CurFolderGUIDString = ProjectFolderGuids[CurFolder].ToString("B").ToUpperInvariant();
-
-									foreach (MSBuildProjectFile ChildProject in CurFolder.ChildProjects)
-									{
-										if (AllProjectFiles.Contains(ChildProject))
-										{
-											//	e.g. "{BF6FB09F-A2A6-468F-BE6F-DEBE07EAD3EA} = {C43B6BB5-3EF0-4784-B896-4099753BCDA9}"
-											LocalVCSolutionFileContent.AppendLine("		" + ChildProject.ProjectGUID.ToString("B").ToUpperInvariant() + " = " + CurFolderGUIDString);
-										}
-									}
-
-									foreach (PrimaryProjectFolder SubFolder in CurFolder.SubFolders)
-									{
-										// if this folder has no projects anywhere under it then skip it
-										if (HasProjectFunc(SubFolder))
-										{
-											//	e.g. "{BF6FB09F-A2A6-468F-BE6F-DEBE07EAD3EA} = {C43B6BB5-3EF0-4784-B896-4099753BCDA9}"
-											LocalVCSolutionFileContent.AppendLine("		" + ProjectFolderGuids[SubFolder].ToString("B").ToUpperInvariant() + " = " + CurFolderGUIDString);
-										}
-									}
-
-									// Recurse into subfolders
-									FolderProcessorFunction(LocalVCSolutionFileContent, CurFolder.SubFolders);
-								}
-							}
-						FolderProcessorFunction(VCSolutionFileContent, RootFolder.SubFolders);
-
-						VCSolutionFileContent.AppendLine("	EndGlobalSection");
-					}
+					SolutionBuilder.WriteNestedProjects();
 				}
 
-				VCSolutionFileContent.AppendLine("EndGlobal");
+				SolutionBuilder.WriteGlobalEnd();
 			}
 
 			// Save the solution file
 			if (bSuccess)
 			{
 				FileReference SolutionFilePath = FileReference.Combine(PrimaryProjectPath, SolutionFileName);
-				bSuccess = WriteFileIfChanged(SolutionFilePath.FullName, VCSolutionFileContent.ToString(), Logger);
+				bSuccess = WriteFileIfChanged(SolutionFilePath.FullName, SolutionBuilder.ToString(), Logger);
 
 				// Generate a slnx for VS2026 until properly supported
 				if (bSuccess && Settings.ProjectFileFormat == VCProjectFileFormat.VisualStudio2026 && Environment.Version.Major >= 10)
 				{
-					FileReference TempSolutionFilePath = FileReference.Combine(new DirectoryReference(Path.GetTempPath()), SolutionFileName);
-					FileReference.Copy(SolutionFilePath, TempSolutionFilePath, true);
-					using ManagedProcessGroup group = new();
-					using ManagedProcess process = new(group, Unreal.DotnetPath.FullName, $"sln migrate \"{TempSolutionFilePath}\"", TempSolutionFilePath.Directory.FullName, null, System.Diagnostics.ProcessPriorityClass.Normal);
-					process.ReadAllLines();
-					process.WaitForExit();
-					bSuccess = process.ExitCode == 0;
-					if (bSuccess)
-					{
-						string contents = FileReference.ReadAllText(TempSolutionFilePath.ChangeExtension(".slnx"));
-						bSuccess = WriteFileIfChanged(SolutionFilePath.ChangeExtension(".slnx").FullName, contents, Logger);
-						FileReference.Delete(TempSolutionFilePath);
-						FileReference.Delete(TempSolutionFilePath.ChangeExtension(".slnx"));
-					}
+					bSuccess = ConvertSolutionFileToSlnx(SolutionFileName, SolutionFilePath, Logger);
 				}
 			}
 
@@ -1063,68 +857,106 @@ namespace UnrealBuildTool
 			// .suo file writable only on Windows, requires ole32
 			if (bSuccess && Settings.bWriteSolutionOptionFile && OperatingSystem.IsWindows())
 			{
-				// Figure out the filename for the SUO file. VS will automatically import the options from earlier versions if necessary.
-				FileReference SolutionOptionsFileName = Settings.ProjectFileFormat switch
-				{
-					VCProjectFileFormat.VisualStudio2022 => FileReference.Combine(PrimaryProjectPath, ".vs", Path.GetFileNameWithoutExtension(SolutionFileName), "v17", ".suo"),
-					VCProjectFileFormat.VisualStudio2026 => FileReference.Combine(PrimaryProjectPath, ".vs", Path.GetFileNameWithoutExtension(SolutionFileName), "v18", ".suo"),
-					_ => throw new BuildException("Unsupported Visual Studio version"),
-				};
+				WriteSolutionOptionFile(SolutionFileName, SolutionConfigCombinations);
+			}
 
-				// Check it doesn't exist before overwriting it. Since these files store the user's preferences, it'd be bad form to overwrite them.
-				if (!FileReference.Exists(SolutionOptionsFileName))
-				{
-					DirectoryReference.CreateDirectory(SolutionOptionsFileName.Directory);
+			if (bSuccess && Settings.bVsConfigFile && OperatingSystem.IsWindows())
+			{
+				bSuccess = WriteVsConfigFile(Logger);
+			}
 
-					VCSolutionOptions Options = new VCSolutionOptions(Settings.ProjectFileFormat);
+			return bSuccess;
+		}
 
-					// Set the default configuration and startup project
-					VCSolutionConfigCombination? DefaultConfig = SolutionConfigCombinations.Find(x =>
-						x.Configuration == UnrealTargetConfiguration.Development &&
-						x.Platform == UnrealTargetPlatform.Win64 &&
-						(x.Architecture == null || x.Architecture == UnrealArch.X64) &&
-						(bMakeProjectPerTarget || x.TargetConfigurationName == TargetType.Editor));
-					if (DefaultConfig != null)
-					{
-						List<VCBinarySetting> Settings = new List<VCBinarySetting>();
-						Settings.Add(new VCBinarySetting("ActiveCfg", DefaultConfig.VCSolutionConfigAndPlatformName));
-						if (DefaultProject != null)
-						{
-							Settings.Add(new VCBinarySetting("StartupProject", ((MSBuildProjectFile)DefaultProject).ProjectGUID.ToString("B")));
-						}
-						Options.SetConfiguration(Settings);
-					}
+		private static bool ConvertSolutionFileToSlnx(string SolutionFileName, FileReference SolutionFilePath, ILogger Logger)
+		{
+			FileReference TempSolutionFilePath = FileReference.Combine(new DirectoryReference(Path.GetTempPath()), SolutionFileName);
+			FileReference.Copy(SolutionFilePath, TempSolutionFilePath, true);
+			using ManagedProcessGroup group = new();
+			using ManagedProcess process = new(
+				group,
+				Unreal.DotnetPath.FullName,
+				$"sln migrate \"{TempSolutionFilePath}\"",
+				TempSolutionFilePath.Directory.FullName,
+				null,
+				System.Diagnostics.ProcessPriorityClass.Normal
+			);
+			process.ReadAllLines();
+			process.WaitForExit();
+			bool bSuccess = process.ExitCode == 0;
+			if (bSuccess)
+			{
+				string contents = FileReference.ReadAllText(TempSolutionFilePath.ChangeExtension(".slnx"));
+				bSuccess = WriteFileIfChanged(SolutionFilePath.ChangeExtension(".slnx").FullName, contents, Logger);
+				FileReference.Delete(TempSolutionFilePath);
+				FileReference.Delete(TempSolutionFilePath.ChangeExtension(".slnx"));
+			}
+
+			return bSuccess;
+		}
+
+		[SupportedOSPlatform("windows")]
+		private void WriteSolutionOptionFile(string SolutionFileName, List<VCSolutionConfigCombination> SolutionConfigCombinations)
+		{
+			// Figure out the filename for the SUO file. VS will automatically import the options from earlier versions if necessary.
+			FileReference SolutionOptionsFileName = Settings.ProjectFileFormat switch
+			{
+				VCProjectFileFormat.VisualStudio2022 => FileReference.Combine(PrimaryProjectPath, ".vs", Path.GetFileNameWithoutExtension(SolutionFileName), "v17", ".suo"),
+				VCProjectFileFormat.VisualStudio2026 => FileReference.Combine(PrimaryProjectPath, ".vs", Path.GetFileNameWithoutExtension(SolutionFileName), "v18", ".suo"),
+				_ => throw new BuildException("Unsupported Visual Studio version"),
+			};
+
+			// Check it doesn't exist before overwriting it. Since these files store the user's preferences, it'd be bad form to overwrite them.
+			if (!FileReference.Exists(SolutionOptionsFileName))
+			{
+				DirectoryReference.CreateDirectory(SolutionOptionsFileName.Directory);
 
-					// Mark all the projects as closed by default, apart from the startup project
-					VCSolutionExplorerState ExplorerState = new VCSolutionExplorerState();
-					Options.SetExplorerState(ExplorerState);
+				VCSolutionOptions Options = new VCSolutionOptions(Settings.ProjectFileFormat);
 
-					// Write the file
-					if (Options.Sections.Count > 0)
+				// Set the default configuration and startup project
+				VCSolutionConfigCombination? DefaultConfig = SolutionConfigCombinations.Find(x =>
+					x.Configuration == UnrealTargetConfiguration.Development &&
+					x.Platform == UnrealTargetPlatform.Win64 &&
+					(x.Architecture == null || x.Architecture == UnrealArch.X64) &&
+					(bMakeProjectPerTarget || x.TargetConfigurationName == TargetType.Editor));
+				if (DefaultConfig != null)
+				{
+					List<VCBinarySetting> Settings = new List<VCBinarySetting>();
+					Settings.Add(new VCBinarySetting("ActiveCfg", DefaultConfig.VCSolutionConfigAndPlatformName));
+					if (DefaultProject != null)
 					{
-						Options.Write(SolutionOptionsFileName.FullName);
+						Settings.Add(new VCBinarySetting("StartupProject", ((MSBuildProjectFile)DefaultProject).ProjectGUID.ToString("B")));
 					}
+					Options.SetConfiguration(Settings);
 				}
-			}
 
-			if (bSuccess && Settings.bVsConfigFile && OperatingSystem.IsWindows())
-			{
-				StringBuilder VsConfigFileContent = new StringBuilder();
-
-				VsConfigFileContent.AppendLine("{");
-				VsConfigFileContent.AppendLine("  \"version\": \"1.0\",");
-				VsConfigFileContent.AppendLine("  \"components\": [");
-				IEnumerable<string> Components = MicrosoftPlatformSDK.GetVisualStudioSuggestedComponents(Settings.ProjectFileFormat);
-				string ComponentsString = String.Join($",{Environment.NewLine}    ", Components.Select(x => $"\"{x}\""));
-				VsConfigFileContent.AppendLine($"    {ComponentsString}");
-				VsConfigFileContent.AppendLine("  ]");
-				VsConfigFileContent.AppendLine("}");
-
-				FileReference VsConfigFileName = FileReference.Combine(PrimaryProjectPath, ".vsconfig");
-				bSuccess = WriteFileIfChanged(VsConfigFileName.FullName, VsConfigFileContent.ToString(), Logger);
+				// Mark all the projects as closed by default, apart from the startup project
+				VCSolutionExplorerState ExplorerState = new VCSolutionExplorerState();
+				Options.SetExplorerState(ExplorerState);
+
+				// Write the file
+				if (Options.Sections.Count > 0)
+				{
+					Options.Write(SolutionOptionsFileName.FullName);
+				}
 			}
+		}
 
-			return bSuccess;
+		private bool WriteVsConfigFile(ILogger Logger)
+		{
+			StringBuilder VsConfigFileContent = new StringBuilder();
+
+			VsConfigFileContent.AppendLine("{");
+			VsConfigFileContent.AppendLine("  \"version\": \"1.0\",");
+			VsConfigFileContent.AppendLine("  \"components\": [");
+			IEnumerable<string> Components = MicrosoftPlatformSDK.GetVisualStudioSuggestedComponents(Settings.ProjectFileFormat);
+			string ComponentsString = String.Join($",{Environment.NewLine}    ", Components.Select(x => $"\"{x}\""));
+			VsConfigFileContent.AppendLine($"    {ComponentsString}");
+			VsConfigFileContent.AppendLine("  ]");
+			VsConfigFileContent.AppendLine("}");
+
+			FileReference VsConfigFileName = FileReference.Combine(PrimaryProjectPath, ".vsconfig");
+			return WriteFileIfChanged(VsConfigFileName.FullName, VsConfigFileContent.ToString(), Logger);
 		}
 
 		public static void CollectSolutionConfigurations(List<UnrealTargetConfiguration> AllConfigurations,
diff --git a/Engine/Source/Programs/UnrealBuildTool/ProjectFiles/VisualStudio/VisualStudioSolutionBuilder.cs b/Engine/Source/Programs/UnrealBuildTool/ProjectFiles/VisualStudio/VisualStudioSolutionBuilder.cs
new file mode 100644
--- /dev/null
+++ b/Engine/Source/Programs/UnrealBuildTool/ProjectFiles/VisualStudio/VisualStudioSolutionBuilder.cs
@@ -0,0 +1,331 @@
+// Copyright Epic Games, Inc. All Rights Reserved.
+
+using System;
+using System.Collections.Generic;
+using System.Linq;
+using System.Text;
+using Microsoft.Extensions.Logging;
+using static UnrealBuildTool.VCProjectFileGenerator;
+
+namespace UnrealBuildTool.ProjectFiles.VisualStudio;
+
+internal sealed class VisualStudioSolutionBuilder
+{
+	public VisualStudioSolutionBuilder(PrimaryProjectFolder rootFolder, List<ProjectFile> allProjectFiles, ILogger logger)
+	{
+		_rootFolder = rootFolder;
+		_logger = logger;
+		_allProjectFiles = [.. allProjectFiles.OfType<MSBuildProjectFile>()];
+		_projectFolderGuids = GenerateProjectFolderGuids();
+	}
+
+	// This the GUID that Visual Studio uses to identify a solution folder
+	public const string SolutionFolderGuid = "{2150E333-8FDC-42A3-9474-1A3956D46DE8}";
+
+	public static string FormatGuid(Guid guid) => guid.ToString("B").ToUpperInvariant();
+
+	public void WritePreamble(VCProjectFileFormat format)
+	{
+		// Solution file header. Note that a leading newline is required for file type detection to work correctly in the shell.
+
+		_sb.AppendLine("Microsoft Visual Studio Solution File, Format Version 12.00");
+
+		switch (format)
+		{
+			case VCProjectFileFormat.VisualStudio2022:
+				_sb.AppendLine("# Visual Studio Version 17");
+				_sb.AppendLine("VisualStudioVersion = 17.0.31314.256");
+				break;
+			case VCProjectFileFormat.VisualStudio2026:
+				_sb.AppendLine("# Visual Studio Version 18");
+				_sb.AppendLine("VisualStudioVersion = 18.0.11205.157");
+				break;
+			default:
+				throw new BuildException("Unexpected ProjectFileFormat '{0}'", format);
+		}
+
+		_sb.AppendLine("MinimumVisualStudioVersion = 10.0.40219.1");
+	}
+
+	public void WriteProjectStart(string typeGuid, string name, string path, Guid guid)
+	{
+		string guidString = FormatGuid(guid);
+		_sb.AppendLine($"Project(\"{typeGuid}\") = \"{name}\", \"{path}\", \"{guidString}\"");
+	}
+
+	public void WriteProjectEnd()
+	{
+		_sb.AppendLine("EndProject");
+	}
+
+	public void WriteProjectFolders()
+	{
+		IEnumerable<PrimaryProjectFolder> allSolutionFolders = _projectFolderGuids.Keys.OrderBy(folder => folder.FolderName).ThenBy(folder => _projectFolderGuids[folder]);
+		foreach (PrimaryProjectFolder folder in allSolutionFolders)
+		{
+			Guid guid = _projectFolderGuids[folder];
+
+			WriteProjectStart(SolutionFolderGuid, folder.FolderName, folder.FolderName, guid);
+			WriteProjectSolutionItems(folder);
+			WriteProjectEnd();
+		}
+	}
+
+	public void WriteProjects()
+	{
+		// Project files
+		foreach (MSBuildProjectFile project in _allProjectFiles)
+		{
+			// Visual Studio uses different GUID types depending on the project type
+			string typeGuid = project.ProjectTypeGUID;
+
+			// NOTE: The project name in the solution doesn't actually *have* to match the project file name on disk. However,
+			// we prefer it when it does match so we use the actual file name here.
+			string name = project.ProjectFilePath.GetFileNameWithoutExtension();
+			string relativePath = project.ProjectFilePath.MakeRelativeTo(ProjectFileGenerator.PrimaryProjectPath);
+
+			WriteProjectStart(typeGuid, name, relativePath, project.ProjectGUID);
+
+			if (!project.IsStubProject)
+			{
+				if (project.DependsOnProjects.Count > 0)
+				{
+					_sb.AppendLine("\tProjectSection(ProjectDependencies) = postProject");
+
+					// Setup any addition dependencies this project has
+					foreach (MSBuildProjectFile dependsOnProject in project.DependsOnProjects.Cast<MSBuildProjectFile>())
+					{
+						string dependsOnProjectGuid = FormatGuid(dependsOnProject.ProjectGUID);
+						_sb.AppendLine($"\t\t{dependsOnProjectGuid} = {dependsOnProjectGuid}");
+					}
+
+					_sb.AppendLine("\tEndProjectSection");
+				}
+			}
+
+			WriteProjectEnd();
+		}
+	}
+
+	public void WriteGlobalStart()
+	{
+		_sb.AppendLine("Global");
+	}
+
+	public void WriteGlobalEnd()
+	{
+		_sb.AppendLine("EndGlobal");
+	}
+
+	public void WriteSolutionConfigurationPlatforms(List<VCSolutionConfigCombination> solutionConfigCombinations)
+	{
+		_sb.AppendLine("	GlobalSection(SolutionConfigurationPlatforms) = preSolution");
+
+		HashSet<string> AppendedSolutionConfigAndPlatformNames = new HashSet<string>(StringComparer.OrdinalIgnoreCase);
+		foreach (VCSolutionConfigCombination SolutionConfigCombination in solutionConfigCombinations)
+		{
+			// We alias "Game" and "Program" to both have the same solution configuration, so we're careful not to add the same combination twice.
+			if (!AppendedSolutionConfigAndPlatformNames.Contains(SolutionConfigCombination.VCSolutionConfigAndPlatformName))
+			{
+				_sb.AppendLine($"		{SolutionConfigCombination.VCSolutionConfigAndPlatformName} = {SolutionConfigCombination.VCSolutionConfigAndPlatformName}");
+				AppendedSolutionConfigAndPlatformNames.Add(SolutionConfigCombination.VCSolutionConfigAndPlatformName);
+			}
+		}
+
+		_sb.AppendLine("	EndGlobalSection");
+	}
+
+	public void WriteProjectConfigurationPlatforms(
+		List<VCSolutionConfigCombination> solutionConfigCombinations,
+		PlatformProjectGeneratorCollection platformProjectGenerators,
+		Predicate<MSBuildProjectFile> debugPredicate)
+	{
+		_sb.AppendLine("	GlobalSection(ProjectConfigurationPlatforms) = postSolution");
+
+		foreach (MSBuildProjectFile project in _allProjectFiles)
+		{
+			foreach (VCSolutionConfigCombination solutionConfigCombination in solutionConfigCombinations)
+			{
+				// Get the context for the current solution context
+				MSBuildProjectContext? projectContext = project.GetMatchingProjectContext(
+					solutionConfigCombination.TargetConfigurationName,
+					solutionConfigCombination.Configuration,
+					solutionConfigCombination.Platform,
+					platformProjectGenerators,
+					solutionConfigCombination.Architecture,
+					_logger
+				);
+
+				if (projectContext == null)
+				{
+					continue;
+				}
+
+				// Override the configuration to build for UBT
+				if (debugPredicate(project))
+				{
+					projectContext.ConfigurationName = "Debug";
+				}
+
+				// Write the solution mapping (e.g.  "{4232C52C-680F-4850-8855-DC39419B5E9B}.Debug|iOS.ActiveCfg = iOS_Debug|Win32")
+				string guid = FormatGuid(project.ProjectGUID);
+				string configAndPlatformName = solutionConfigCombination.VCSolutionConfigAndPlatformName;
+
+				_sb.AppendLine($"		{guid}.{configAndPlatformName}.ActiveCfg = {projectContext.Name}");
+				if (projectContext.bBuildByDefault)
+				{
+					_sb.AppendLine($"		{guid}.{configAndPlatformName}.Build.0 = {projectContext.Name}");
+					if (projectContext.bDeployByDefault)
+					{
+						_sb.AppendLine($"		{guid}.{configAndPlatformName}.Deploy.0 = {projectContext.Name}");
+					}
+				}
+			}
+		}
+
+		_sb.AppendLine("	EndGlobalSection");
+	}
+
+	public void WriteUnrealVs(HashSet<UnrealTargetPlatform> validPlatforms)
+	{
+		const string UnrealVSGuid = "ddbf523f-7eb6-4887-bd51-85a714ff87eb";
+		string validPlatformsString = String.Join(";", validPlatforms.Select(platform => platform.ToString()));
+
+		_sb.AppendLine("\t# UnrealVS Section");
+		_sb.AppendLine($"\tGlobalSection({UnrealVSGuid}) = preSolution");
+		_sb.AppendLine($"\t\tAvailablePlatforms={validPlatformsString}");
+		_sb.AppendLine("\tEndGlobalSection");
+	}
+
+	public void WriteSolutionProperties()
+	{
+		// HideSolutionNode sets whether or not the top-level solution entry is completely hidden in the UI.
+		// We don't want that, as we need users to be able to right click on the solution tree item.
+		_sb.AppendLine("	GlobalSection(SolutionProperties) = preSolution");
+		_sb.AppendLine("		HideSolutionNode = FALSE");
+		_sb.AppendLine("	EndGlobalSection");
+	}
+
+	public void WriteNestedProjects()
+	{
+		_sb.AppendLine("	GlobalSection(NestedProjects) = preSolution");
+
+		// Every entry in this section is in the format "Guid1 = Guid2".  Guid1 is the child project (or solution
+		// filter)'s GUID, and Guid2 is the solution filter directory to parent the child project (or solution
+		// filter) to. This sets up the hierarchical solution explorer tree for all solution folders and projects.
+
+		void FolderProcessorFunction(List<PrimaryProjectFolder> localFolders)
+		{
+			foreach (PrimaryProjectFolder folder in localFolders)
+			{
+				if (!_projectFolderGuids.TryGetValue(folder, out Guid folderGuid))
+				{
+					continue;
+				}
+
+				string folderGuidString = FormatGuid(_projectFolderGuids[folder]);
+
+				foreach (MSBuildProjectFile childProject in folder.ChildProjects)
+				{
+					if (_allProjectFiles.Contains(childProject))
+					{
+						//	e.g. "{BF6FB09F-A2A6-468F-BE6F-DEBE07EAD3EA} = {C43B6BB5-3EF0-4784-B896-4099753BCDA9}"
+						_sb.AppendLine($"		{FormatGuid(childProject.ProjectGUID)} = {folderGuidString}");
+					}
+				}
+
+				foreach (PrimaryProjectFolder subFolder in folder.SubFolders)
+				{
+					// if this folder has no projects anywhere under it then skip it
+					if (_projectFolderGuids.TryGetValue(subFolder, out Guid subFolderGuid))
+					{
+						//	e.g. "{BF6FB09F-A2A6-468F-BE6F-DEBE07EAD3EA} = {C43B6BB5-3EF0-4784-B896-4099753BCDA9}"
+						_sb.AppendLine($"		{FormatGuid(subFolderGuid)} = {folderGuidString}");
+					}
+				}
+
+				// Recurse into subfolders
+				FolderProcessorFunction(folder.SubFolders);
+			}
+		}
+
+		FolderProcessorFunction(_rootFolder.SubFolders);
+
+		_sb.AppendLine("	EndGlobalSection");
+	}
+
+	/// <summary>
+	/// Add any files that are inlined right inside the solution folder
+	/// </summary>
+	public void WriteProjectSolutionItems(PrimaryProjectFolder folder)
+	{
+		if (folder.Files.Count > 0)
+		{
+			_sb.AppendLine("	ProjectSection(SolutionItems) = preProject");
+			foreach (string CurFile in folder.Files)
+			{
+				// Syntax is:  <relative file path> = <relative file path>
+				_sb.AppendLine("		" + CurFile + " = " + CurFile);
+			}
+			_sb.AppendLine("	EndProjectSection");
+		}
+	}
+
+	private bool FolderHasAnyProjects(PrimaryProjectFolder ProjectFolder)
+	{
+		foreach (MSBuildProjectFile ChildProject in ProjectFolder.ChildProjects)
+		{
+			if (_allProjectFiles.Contains(ChildProject))
+			{
+				return true;
+			}
+		}
+
+		foreach (PrimaryProjectFolder SubFolder in ProjectFolder.SubFolders)
+		{
+			if (FolderHasAnyProjects(SubFolder))
+			{
+				return true;
+			}
+		}
+
+		return false;
+	}
+
+	private Dictionary<PrimaryProjectFolder, Guid> GenerateProjectFolderGuids()
+	{
+		Dictionary<PrimaryProjectFolder, Guid> guids = [];
+		foreach (PrimaryProjectFolder folder in _rootFolder.SubFolders)
+		{
+			GenerateProjectFolderGuids("UE5", folder, guids);
+		}
+		return guids;
+	}
+
+	private void GenerateProjectFolderGuids(
+		string parentPath,
+		PrimaryProjectFolder folder,
+		Dictionary<PrimaryProjectFolder, Guid> guids)
+	{
+		if (!FolderHasAnyProjects(folder))
+		{
+			return;
+		}
+
+		string path = String.Format("{0}/{1}", parentPath, folder.FolderName);
+		guids[folder] = MakeMd5Guid(path);
+
+		foreach (PrimaryProjectFolder subFolder in folder.SubFolders)
+		{
+			GenerateProjectFolderGuids(path, subFolder, guids);
+		}
+	}
+
+	public override string ToString() => _sb.ToString();
+
+	private readonly StringBuilder _sb = new();
+	private readonly PrimaryProjectFolder _rootFolder;
+	private readonly ILogger _logger;
+	private readonly List<MSBuildProjectFile> _allProjectFiles;
+	private readonly Dictionary<PrimaryProjectFolder, Guid> _projectFolderGuids;
+}
