UnrealBuildTool: Fix issue CA1849 (Call async methods when in an async method).

#Synced-CL 50156122

diff --git a/Engine/Source/Programs/UnrealBuildTool/Configuration/UEBuildTarget.cs b/Engine/Source/Programs/UnrealBuildTool/Configuration/UEBuildTarget.cs
--- a/Engine/Source/Programs/UnrealBuildTool/Configuration/UEBuildTarget.cs
+++ b/Engine/Source/Programs/UnrealBuildTool/Configuration/UEBuildTarget.cs
@@ -2659,7 +2659,10 @@ namespace UnrealBuildTool
 		{
 			await Task.CompletedTask;
 			(TargetMakefile Makefile, Task? UHTTimestampTask) = BuildMakefile(BuildConfiguration, WorkingSet, TargetDescriptor, Logger, bInitOnly);
-			UHTTimestampTask?.Wait();
+			if (UHTTimestampTask is not null)
+			{
+				await UHTTimestampTask;
+			}
 			return Makefile;
 
 		}
diff --git a/Engine/Source/Programs/UnrealBuildTool/Executors/UnrealBuildAccelerator/UBAAgentCoordinatorHorde.cs b/Engine/Source/Programs/UnrealBuildTool/Executors/UnrealBuildAccelerator/UBAAgentCoordinatorHorde.cs
--- a/Engine/Source/Programs/UnrealBuildTool/Executors/UnrealBuildAccelerator/UBAAgentCoordinatorHorde.cs
+++ b/Engine/Source/Programs/UnrealBuildTool/Executors/UnrealBuildAccelerator/UBAAgentCoordinatorHorde.cs
@@ -882,9 +882,9 @@ namespace UnrealBuildTool
 			_updateStatus!(0, 6, "Connecting", EpicGames.UBA.LogEntryType.Info, null);
 
 			_hordeSessionTask = UBAHordeSession.TryCreateHordeSessionAsync(_hordeConfig, _ubaConfig, executor.Crypto, _ubaConfig.bStrict, _logger, cancellationToken);
-			await _hordeSessionTask;
-			_updateStatus!(0, 6, _hordeSessionTask.Result.Item2, EpicGames.UBA.LogEntryType.Info, null);
-			executor.AgentCoordinatorInitialized(_hordeSessionTask.IsCompletedSuccessfully && _hordeSessionTask.Result.Item1 != null);
+			(UBAHordeSession? session, string status) = await _hordeSessionTask;
+			_updateStatus!(0, 6, status, EpicGames.UBA.LogEntryType.Info, null);
+			executor.AgentCoordinatorInitialized(_hordeSessionTask.IsCompletedSuccessfully && session is not null);
 		}
 
 		public UnrealBuildAcceleratorCacheConfig? RequestCacheServer(CancellationToken cancellationToken)
diff --git a/Engine/Source/Programs/UnrealBuildTool/Executors/UnrealBuildAccelerator/UBAExecutor.cs b/Engine/Source/Programs/UnrealBuildTool/Executors/UnrealBuildAccelerator/UBAExecutor.cs
--- a/Engine/Source/Programs/UnrealBuildTool/Executors/UnrealBuildAccelerator/UBAExecutor.cs
+++ b/Engine/Source/Programs/UnrealBuildTool/Executors/UnrealBuildAccelerator/UBAExecutor.cs
@@ -579,10 +579,10 @@ namespace UnrealBuildTool
 
 			if (!inputActions.Any())
 			{
-				_cancellationSource.Cancel();
+				await _cancellationSource.CancelAsync();
 			}
 
-			_initTask.Wait();
+			await _initTask;
 
 			_trace?.UpdateStatus(1, 6, "Hits 0 Misses 0", LogEntryType.Info);
 
@@ -635,7 +635,7 @@ namespace UnrealBuildTool
 			}
 			finally
 			{
-				_cancellationSource.Cancel();
+				await _cancellationSource.CancelAsync();
 				_agentCoordinators.ForEach(ac => ac.Stop());
 
 				DisposeCacheClients();
@@ -1130,7 +1130,7 @@ namespace UnrealBuildTool
 				{
 					if (shouldCancel && !_cancellationSource.IsCancellationRequested)
 					{
-						_cancellationSource.Cancel();
+						await _cancellationSource.CancelAsync();
 
 						scheduler.Cancel();
 						foreach (ICacheClient client in cache)
@@ -1151,7 +1151,7 @@ namespace UnrealBuildTool
 							}
 						}
 					}
-					Thread.Sleep(100);
+					await Task.Delay(100);
 				}
 
 				bool noFailedActions = failedActions == 0;
diff --git a/Engine/Source/Programs/UnrealBuildTool/Modes/AggregateClangTimingInfo.cs b/Engine/Source/Programs/UnrealBuildTool/Modes/AggregateClangTimingInfo.cs
--- a/Engine/Source/Programs/UnrealBuildTool/Modes/AggregateClangTimingInfo.cs
+++ b/Engine/Source/Programs/UnrealBuildTool/Modes/AggregateClangTimingInfo.cs
@@ -169,10 +169,10 @@ namespace UnrealBuildTool
 
 		private readonly ConcurrentDictionary<FileReference, ClangTrace> ClangTraceCache = new();
 
-		public override Task<int> ExecuteAsync(CommandLineArguments Arguments, ILogger Logger)
+		public override async Task<int> ExecuteAsync(CommandLineArguments Arguments, ILogger Logger)
 		{
 			FileReference ManifestFile = Arguments.GetFileReference("-ManifestFile=");
-			IEnumerable<FileReference> SourceFiles = FileReference.ReadAllLines(ManifestFile).Select(x => new FileReference(x));
+			IEnumerable<FileReference> SourceFiles = (await FileReference.ReadAllLinesAsync(ManifestFile)).Select(x => new FileReference(x));
 
 			// Create aggregate summary.
 			FileReference? AggregateFile = Arguments.GetFileReferenceOrDefault("-AggregateFile=", null);
@@ -181,16 +181,16 @@ namespace UnrealBuildTool
 			if (AggregateFile != null)
 			{
 				Task<TraceData[]> Tasks = Task.WhenAll(SourceFiles.Select(x => GetTraceData(x, Logger)));
-				Tasks.Wait();
-				List<TraceData> TraceDatas = Tasks.Result.OrderBy(x => x.Module).ThenBy(x => x.Name).ToList();
+				TraceData[] RawTraceData = await Tasks;
+				List<TraceData> TraceDatas = RawTraceData.OrderBy(x => x.Module).ThenBy(x => x.Name).ToList();
 
 				string TempFilePath = Path.Join(Path.GetTempPath(), AggregateFile.GetFileName() + ".tmp");
 				using (StreamWriter Writer = new StreamWriter(TempFilePath))
 				{
-					Writer.WriteLine(TraceData.CsvHeader);
+					await Writer.WriteLineAsync(TraceData.CsvHeader);
 					foreach (TraceData Data in TraceDatas)
 					{
-						Writer.WriteLine(Data.CsvLine);
+						await Writer.WriteLineAsync(Data.CsvLine);
 					}
 				}
 				File.Move(TempFilePath, AggregateFile.FullName, true);
@@ -199,8 +199,8 @@ namespace UnrealBuildTool
 			if (HeadersFile != null)
 			{
 				Task<ClangTrace[]> Tasks = Task.WhenAll(SourceFiles.Select(x => ParseTimingDataFileAsync(x, Logger)));
-				Tasks.Wait();
-				List<ClangTrace> ClangTraces = Tasks.Result.ToList();
+				ClangTrace[] RawTraceData = await Tasks;
+				List<ClangTrace> ClangTraces = RawTraceData.ToList();
 
 				Dictionary<FileReference, List<long>> Sources = new Dictionary<FileReference, List<long>>();
 				foreach (ClangTrace ClangTrace in ClangTraces.Where(x => x.traceEvents != null))
@@ -220,10 +220,10 @@ namespace UnrealBuildTool
 				string TempFilePath = Path.Join(Path.GetTempPath(), HeadersFile.GetFileName() + ".tmp");
 				using (StreamWriter Writer = new StreamWriter(TempFilePath))
 				{
-					Writer.WriteLine("Source,Count,Total,Min,Max,Average");
+					await Writer.WriteLineAsync("Source,Count,Total,Min,Max,Average");
 					foreach (KeyValuePair<FileReference, List<long>> Data in Sources.OrderBy(x => x.Key.FullName).Where(x => x.Key.HasExtension(".h") || x.Key.HasExtension(".inl")))
 					{
-						Writer.WriteLine($"{Data.Key.FullName},{Data.Value.Count},{Data.Value.Sum()},{Data.Value.Min()},{Data.Value.Max()},{Data.Value.Average()}");
+						await Writer.WriteLineAsync($"{Data.Key.FullName},{Data.Value.Count},{Data.Value.Sum()},{Data.Value.Min()},{Data.Value.Max()},{Data.Value.Average()}");
 					}
 				}
 				File.Move(TempFilePath, HeadersFile.FullName, true);
@@ -255,7 +255,7 @@ namespace UnrealBuildTool
 				File.Move(TempFilePath, ArchiveFile.FullName, true);
 			}
 
-			return Task.FromResult(0);
+			return 0;
 		}
 
 		private async Task<ClangTrace> ParseTimingDataFileAsync(FileReference SourceFile, ILogger Logger)
diff --git a/Engine/Source/Programs/UnrealBuildTool/Modes/AggregateParsedTimingInfo.cs b/Engine/Source/Programs/UnrealBuildTool/Modes/AggregateParsedTimingInfo.cs
--- a/Engine/Source/Programs/UnrealBuildTool/Modes/AggregateParsedTimingInfo.cs
+++ b/Engine/Source/Programs/UnrealBuildTool/Modes/AggregateParsedTimingInfo.cs
@@ -29,10 +29,10 @@ namespace UnrealBuildTool
 		private readonly ConcurrentBag<TimingData> AggregateFunctions = new ConcurrentBag<TimingData>();
 		private readonly object AddFileLock = new object();
 
-		public override Task<int> ExecuteAsync(CommandLineArguments Arguments, ILogger Logger)
+		public override async Task<int> ExecuteAsync(CommandLineArguments Arguments, ILogger Logger)
 		{
 			FileReference ManifestFile = Arguments.GetFileReference("-ManifestFile=");
-			string[] ParsedFileNames = FileReference.ReadAllLines(ManifestFile);
+			string[] ParsedFileNames = await FileReference.ReadAllLinesAsync(ManifestFile);
 
 			// Load all the file timing data and summarize them for the aggregate.
 			FileTimingData = new TimingData("Files", TimingDataType.Summary);
@@ -84,10 +84,10 @@ namespace UnrealBuildTool
 					CompileTimes.Add(JsonWriter.EscapeString(TimingData.Key), TimingData.Value.InclusiveDuration);
 				}
 				string JsonCompileTimes = JsonSerializer.Serialize(CompileTimes);
-				FileReference.WriteAllText(CompileTimingFile, JsonCompileTimes);
+				await FileReference.WriteAllTextAsync(CompileTimingFile, JsonCompileTimes);
 			}
 
-			return Task.FromResult(0);
+			return 0;
 		}
 
 		private static void GroupTimingDataOnName(TimingData AggregateData, string TimingName, IEnumerable<TimingData> UngroupedData)
diff --git a/Engine/Source/Programs/UnrealBuildTool/Modes/BuildMode.cs b/Engine/Source/Programs/UnrealBuildTool/Modes/BuildMode.cs
--- a/Engine/Source/Programs/UnrealBuildTool/Modes/BuildMode.cs
+++ b/Engine/Source/Programs/UnrealBuildTool/Modes/BuildMode.cs
@@ -296,7 +296,7 @@ namespace UnrealBuildTool
 				// Check if anything failed during our run, and act accordingly.
 				ProcessCoreDumps(SaveCrashDumpDirectory, Logger);
 
-				Task.WaitAll(tasks.ToArray());
+				await Task.WhenAll(tasks.ToArray());
 
 				// Save all the caches
 				using (Timeline.ScopeEvent("Saving caches"))
@@ -437,7 +437,7 @@ namespace UnrealBuildTool
 			}
 			finally
 			{
-				Task.WaitAll(tasks.ToArray());
+				await Task.WhenAll(tasks.ToArray());
 			}
 		}
 
@@ -1074,7 +1074,7 @@ namespace UnrealBuildTool
 			}
 			finally
 			{
-				Task.WaitAll(tasks.ToArray());
+				await Task.WhenAll(tasks.ToArray());
 			}
 		}
 
diff --git a/Engine/Source/Programs/UnrealBuildTool/Modes/ClReproMode.cs b/Engine/Source/Programs/UnrealBuildTool/Modes/ClReproMode.cs
--- a/Engine/Source/Programs/UnrealBuildTool/Modes/ClReproMode.cs
+++ b/Engine/Source/Programs/UnrealBuildTool/Modes/ClReproMode.cs
@@ -50,11 +50,11 @@ namespace UnrealBuildTool.Modes
 			return path;
 		}
 
-		private Task<int> PrepareFullRepro(ILogger logger)
+		private async Task<int> PrepareFullRepro(ILogger logger)
 		{
 			if (InputCpp == null || OutputDir == null)
 			{
-				return Task.FromResult(1);
+				return 1;
 			}
 
 			FileReference inputDepFile = InputCpp.ChangeExtension(".cpp.dep.json");
@@ -62,19 +62,19 @@ namespace UnrealBuildTool.Modes
 			ClDepends? depends;
 			using (System.IO.FileStream fstream = FileReference.Open(inputDepFile, System.IO.FileMode.Open))
 			{
-				depends = JsonSerializer.Deserialize<ClDepends>(fstream);
+				depends = await JsonSerializer.DeserializeAsync<ClDepends>(fstream);
 			}
 
 			if (depends == null || depends.Data == null || String.IsNullOrEmpty(depends.Data.Source))
 			{
 				logger.LogError("Unable to deserialize {Dep}", inputDepFile);
-				return Task.FromResult(1);
+				return 1;
 			}
 
 			if (!String.IsNullOrEmpty(depends.Data.PCH))
 			{
 				logger.LogError("PCH repro not currently supported, please use -NoPCH when compiling ({Dep})", inputDepFile);
-				return Task.FromResult(1);
+				return 1;
 			}
 
 			FileReference inputRspFile = InputCpp.ChangeExtension(".cpp.obj.rsp");
@@ -131,13 +131,13 @@ namespace UnrealBuildTool.Modes
 			}
 
 			// Update include paths in the rsp files
-			FileReference.WriteAllLines(outputRspFile, FileReference.ReadAllLines(inputRspFile).Select(x => FixPaths(x)));
+			await FileReference.WriteAllLinesAsync(outputRspFile, (await FileReference.ReadAllLinesAsync(inputRspFile)).Select(x => FixPaths(x)));
 			logger.LogInformation("Rewrote absolute paths for {Input} to {Output}", inputRspFile, outputRspFile);
-			FileReference.WriteAllLines(outputSharedRspFile, FileReference.ReadAllLines(inputSharedRspFile).Select(x => FixPaths(x)));
+			await FileReference.WriteAllLinesAsync(outputSharedRspFile, (await FileReference.ReadAllLinesAsync(inputSharedRspFile)).Select(x => FixPaths(x)));
 			logger.LogInformation("Rewrote absolute paths for {Input} to {Output}", inputSharedRspFile, outputSharedRspFile);
 
 			DirectoryReference outputEngineSourceDir = DirectoryReference.Combine(OutputDir, "Engine", "Source");
-			FileReference.WriteAllLines(batchFile, new string[]
+			await FileReference.WriteAllLinesAsync(batchFile, new string[]
 			{
 				$"@echo off",
 				$"setlocal",
@@ -148,7 +148,7 @@ namespace UnrealBuildTool.Modes
 			});
 			logger.LogInformation("Wrote batch file {Batch}", batchFile);
 
-			return Task.FromResult(0);
+			return 0;
 		}
 
 		public override Task<int> ExecuteAsync(CommandLineArguments arguments, ILogger logger)
diff --git a/Engine/Source/Programs/UnrealBuildTool/Modes/FixIncludePathsMode.cs b/Engine/Source/Programs/UnrealBuildTool/Modes/FixIncludePathsMode.cs
--- a/Engine/Source/Programs/UnrealBuildTool/Modes/FixIncludePathsMode.cs
+++ b/Engine/Source/Programs/UnrealBuildTool/Modes/FixIncludePathsMode.cs
@@ -53,7 +53,7 @@ namespace UnrealBuildTool
 		/// <param name="Arguments">Command line arguments</param>
 		/// <returns>Exit code</returns>
 		/// <param name="Logger"></param>
-		public override Task<int> ExecuteAsync(CommandLineArguments Arguments, ILogger Logger)
+		public override async Task<int> ExecuteAsync(CommandLineArguments Arguments, ILogger Logger)
 		{
 			Arguments.ApplyTo(this);
 
@@ -163,7 +163,7 @@ namespace UnrealBuildTool
 							foreach (FileReference InputFile in FileList)
 							{
 								List<int> LinesUpdated = new();
-								string[] Text = FileReference.ReadAllLines(InputFile);
+								string[] Text = await FileReference.ReadAllLinesAsync(InputFile);
 								bool UpdatedText = false;
 
 								for (int i = 0; i < Text.Length; i++)
@@ -327,7 +327,7 @@ namespace UnrealBuildTool
 												Process.StartInfo.FileName = "p4.exe";
 												Process.StartInfo.Arguments = $"edit {InputFile.FullName}";
 												Process.Start();
-												Process.WaitForExit();
+												await Process.WaitForExitAsync();
 											}
 											System.IO.File.WriteAllLines(InputFile.FullName, Text);
 										}
@@ -347,7 +347,7 @@ namespace UnrealBuildTool
 				}
 			}
 
-			return Task.FromResult(0);
+			return 0;
 		}
 
 		class HeaderSortComparison : IComparer<string>
diff --git a/Engine/Source/Programs/UnrealBuildTool/Modes/IWYUMode.cs b/Engine/Source/Programs/UnrealBuildTool/Modes/IWYUMode.cs
--- a/Engine/Source/Programs/UnrealBuildTool/Modes/IWYUMode.cs
+++ b/Engine/Source/Programs/UnrealBuildTool/Modes/IWYUMode.cs
@@ -631,7 +631,7 @@ namespace UnrealBuildTool
 					Logger.LogInformation("Reading TOC from {TOCPath}...", TocReference.FullName);
 					if (FileReference.Exists(TocReference))
 					{
-						TocContent = FileReference.ReadAllLines(TocReference).ToList();
+						TocContent = [.. await FileReference.ReadAllLinesAsync(TocReference)];
 					}
 				}
 
@@ -943,7 +943,7 @@ namespace UnrealBuildTool
 					DirectoryReference.CreateDirectory(TocDir);
 					string TocName = TargetName + ".txt";
 					FileReference TocReference = FileReference.Combine(TocDir, TocName);
-					FileReference.WriteAllLines(TocReference, TocContent);
+					await FileReference.WriteAllLinesAsync(TocReference, TocContent);
 				}
 			}
 
diff --git a/Engine/Source/Programs/UnrealBuildTool/Modes/InlineGeneratedCpps.cs b/Engine/Source/Programs/UnrealBuildTool/Modes/InlineGeneratedCpps.cs
--- a/Engine/Source/Programs/UnrealBuildTool/Modes/InlineGeneratedCpps.cs
+++ b/Engine/Source/Programs/UnrealBuildTool/Modes/InlineGeneratedCpps.cs
@@ -54,7 +54,7 @@ namespace UnrealBuildTool
 		/// <param name="Arguments">Command line arguments</param>
 		/// <returns>Exit code</returns>
 		/// <param name="Logger"></param>
-		public override Task<int> ExecuteAsync(CommandLineArguments Arguments, ILogger Logger)
+		public override async Task<int> ExecuteAsync(CommandLineArguments Arguments, ILogger Logger)
 		{
 			Arguments.ApplyTo(this);
 
@@ -129,7 +129,7 @@ namespace UnrealBuildTool
 									}
 									else
 									{
-										string HeaderText = FileReference.ReadAllText(FoundHeader.Location);
+										string HeaderText = await FileReference.ReadAllTextAsync(FoundHeader.Location);
 										if (!IncludeGenHeaderRegex.IsMatch(HeaderText))
 										{
 											continue;
@@ -144,7 +144,7 @@ namespace UnrealBuildTool
 										}
 									}
 
-									List<string> TextLines = FileReference.ReadAllLines(InputFile.Location).ToList();
+									List<string> TextLines = [.. await FileReference.ReadAllLinesAsync(InputFile.Location)];
 									if (TextLines.Any(Text => InlineReflectionMarkupRegex.IsMatch(Text)))
 									{
 										continue;
@@ -239,9 +239,9 @@ namespace UnrealBuildTool
 												Process.StartInfo.FileName = "p4.exe";
 												Process.StartInfo.Arguments = $"edit {InputFile.FullName}";
 												Process.Start();
-												Process.WaitForExit();
+												await Process.WaitForExitAsync();
 											}
-											System.IO.File.WriteAllLines(InputFile.FullName, TextLines);
+											await System.IO.File.WriteAllLinesAsync(InputFile.FullName, TextLines);
 										}
 									}
 									catch (Exception ex)
@@ -257,7 +257,7 @@ namespace UnrealBuildTool
 				}
 			}
 
-			return Task.FromResult(0);
+			return 0;
 		}
 	}
 }
diff --git a/Engine/Source/Programs/UnrealBuildTool/Modes/ParseMsvcTimingInfo.cs b/Engine/Source/Programs/UnrealBuildTool/Modes/ParseMsvcTimingInfo.cs
--- a/Engine/Source/Programs/UnrealBuildTool/Modes/ParseMsvcTimingInfo.cs
+++ b/Engine/Source/Programs/UnrealBuildTool/Modes/ParseMsvcTimingInfo.cs
@@ -20,7 +20,7 @@ namespace UnrealBuildTool
 	{
 		const string TimingDataRegex = @"^\t\t(?<Indent>\t*)(?<Name>[^\t]+):\s*(?<Duration>[0-9\.]+)s$";
 
-		public override Task<int> ExecuteAsync(CommandLineArguments Arguments, ILogger Logger)
+		public override async Task<int> ExecuteAsync(CommandLineArguments Arguments, ILogger Logger)
 		{
 			FileReference InputFile = Arguments.GetFileReference("-TimingFile=");
 
@@ -29,11 +29,11 @@ namespace UnrealBuildTool
 			if (Arguments.HasOption("-Tracing"))
 			{
 				ParseTimingDataToTracingFiles(InputFile);
-				return Task.FromResult(0);
+				return 0;
 			}
 
 			// Break the input file into the various sections for processing.
-			string[] AllLines = FileReference.ReadAllLines(InputFile);
+			string[] AllLines = await FileReference.ReadAllLinesAsync(InputFile);
 			List<string> Includes = new List<string>();
 			List<string> Classes = new List<string>();
 			List<string> Functions = new List<string>();
@@ -112,7 +112,7 @@ namespace UnrealBuildTool
 				Writer.Write(Summary);
 			}
 
-			return Task.FromResult(0);
+			return 0;
 		}
 
 		private static TimingData SummarizeParsedTimingData(string SummaryName, TimingDataType TimingType, IEnumerable<string> Lines)
diff --git a/Engine/Source/Programs/UnrealBuildTool/Modes/WriteMetadataMode.cs b/Engine/Source/Programs/UnrealBuildTool/Modes/WriteMetadataMode.cs
--- a/Engine/Source/Programs/UnrealBuildTool/Modes/WriteMetadataMode.cs
+++ b/Engine/Source/Programs/UnrealBuildTool/Modes/WriteMetadataMode.cs
@@ -75,14 +75,14 @@ namespace UnrealBuildTool
 		/// <param name="Arguments">Command line arguments</param>
 		/// <param name="Logger">Logger for output</param>
 		/// <returns>Exit code</returns>
-		private static Task<int> ExecuteInternal(CommandLineArguments Arguments, ILogger Logger)
+		private static async Task<int> ExecuteInternal(CommandLineArguments Arguments, ILogger Logger)
 		{
 			// Read the target info
 			FileReference FilePath = Arguments.GetFileReference("-Input=");
 			WriteMetadataTargetInfo TargetInfo;
 			using (FileStream FileStream = FileReference.Open(FilePath, FileMode.Open, FileAccess.Read, FileShare.Read))
 			{
-				WriteMetadataTargetInfo? MaybeTargetInfo = JsonSerializer.Deserialize<WriteMetadataTargetInfo>(FileStream, SerializerOptions);
+				WriteMetadataTargetInfo? MaybeTargetInfo = await JsonSerializer.DeserializeAsync<WriteMetadataTargetInfo>(FileStream, SerializerOptions);
 				TargetInfo = MaybeTargetInfo ?? throw new BuildLogEventException("Target info file '{FilePath}' was invalid", FilePath);
 			}
 
@@ -155,7 +155,7 @@ namespace UnrealBuildTool
 						// Check if the manifest has changed. Note that if a manifest is out of date, we should have generated a new build id causing the contents to differ.
 						if (bNoManifestChanges)
 						{
-							string CurrentText = FileReference.ReadAllText(ManifestFile);
+							string CurrentText = await FileReference.ReadAllTextAsync(ManifestFile);
 							if (CurrentText != OutputText)
 							{
 								Logger.LogError(
@@ -167,7 +167,7 @@ namespace UnrealBuildTool
 						}
 
 						// Write it to disk
-						FileReference.WriteAllText(ManifestFile, OutputText);
+						await FileReference.WriteAllTextAsync(ManifestFile, OutputText);
 					}
 				}
 			}
@@ -199,7 +199,7 @@ namespace UnrealBuildTool
 						// Check if the manifest has changed. Note that if a manifest is out of date, we should have generated a new build id causing the contents to differ.
 						if (bNoManifestChanges)
 						{
-							string CurrentText = FileReference.ReadAllText(ManifestFile);
+							string CurrentText = await FileReference.ReadAllTextAsync(ManifestFile);
 							if (CurrentText != OutputText)
 							{
 								Logger.LogError(
@@ -211,7 +211,7 @@ namespace UnrealBuildTool
 						}
 
 						// Write it to disk
-						FileReference.WriteAllText(ManifestFile, OutputText);
+						await FileReference.WriteAllTextAsync(ManifestFile, OutputText);
 					}
 				}
 			}
@@ -232,7 +232,7 @@ namespace UnrealBuildTool
 				TargetInfo.Receipt.Write(TargetInfo.ReceiptFile);
 			}
 
-			return Task.FromResult(0);
+			return 0;
 		}
 
 		/// <summary>
diff --git a/Engine/Source/Programs/UnrealBuildTool/Platform/Apple/AppleToolChain.cs b/Engine/Source/Programs/UnrealBuildTool/Platform/Apple/AppleToolChain.cs
--- a/Engine/Source/Programs/UnrealBuildTool/Platform/Apple/AppleToolChain.cs
+++ b/Engine/Source/Programs/UnrealBuildTool/Platform/Apple/AppleToolChain.cs
@@ -1189,7 +1189,7 @@ namespace UnrealBuildTool
 		[CommandLine("-noEntitlements")]
 		public bool bNoEntitlements;
 
-		public override Task<int> ExecuteAsync(CommandLineArguments Arguments, ILogger Logger)
+		public override async Task<int> ExecuteAsync(CommandLineArguments Arguments, ILogger Logger)
 		{
 			Arguments.ApplyTo(this);
 			Arguments.CheckAllArgumentsUsed();
@@ -1198,13 +1198,11 @@ namespace UnrealBuildTool
 			ApplePostBuildSyncTarget Target;
 			using (FileStream FileStream = FileReference.Open(InputFile!, FileMode.Open, FileAccess.Read, FileShare.Read))
 			{
-				ApplePostBuildSyncTarget? MaybeTarget = JsonSerializer.Deserialize<ApplePostBuildSyncTarget>(FileStream, AppleToolChain.SerializerOptions);
+				ApplePostBuildSyncTarget? MaybeTarget = await JsonSerializer.DeserializeAsync<ApplePostBuildSyncTarget>(FileStream, AppleToolChain.SerializerOptions);
 				Target = MaybeTarget ?? throw new BuildLogEventException("Post-build sync target file '{FilePath}' was invalid", InputFile!);
 			}
 
-			int ExitCode = PostBuildSync(Target, Logger);
-
-			return Task.FromResult(ExitCode);
+			return PostBuildSync(Target, Logger);
 		}
 
 		private static int PostBuildSync(ApplePostBuildSyncTarget Target, ILogger Logger)
diff --git a/Engine/Source/Programs/UnrealBuildTool/Platform/IOS/IOSPostBuildSyncMode.cs b/Engine/Source/Programs/UnrealBuildTool/Platform/IOS/IOSPostBuildSyncMode.cs
--- a/Engine/Source/Programs/UnrealBuildTool/Platform/IOS/IOSPostBuildSyncMode.cs
+++ b/Engine/Source/Programs/UnrealBuildTool/Platform/IOS/IOSPostBuildSyncMode.cs
@@ -83,7 +83,7 @@ namespace UnrealBuildTool
 		[CommandLine("-LegacyXcode")]
 		public bool bLegacyXcode;
 
-		public override Task<int> ExecuteAsync(CommandLineArguments Arguments, ILogger Logger)
+		public override async Task<int> ExecuteAsync(CommandLineArguments Arguments, ILogger Logger)
 		{
 			Arguments.ApplyTo(this);
 			Arguments.CheckAllArgumentsUsed();
@@ -92,13 +92,13 @@ namespace UnrealBuildTool
 			IOSPostBuildSyncTarget Target;
 			using (FileStream FileStream = FileReference.Open(InputFile!, FileMode.Open, FileAccess.Read, FileShare.Read))
 			{
-				IOSPostBuildSyncTarget? MaybeTarget = JsonSerializer.Deserialize<IOSPostBuildSyncTarget>(FileStream, AppleToolChain.SerializerOptions);
+				IOSPostBuildSyncTarget? MaybeTarget = await JsonSerializer.DeserializeAsync<IOSPostBuildSyncTarget>(FileStream, AppleToolChain.SerializerOptions);
 				Target = MaybeTarget ?? throw new BuildLogEventException("Post-build sync target file '{FilePath}' was invalid", InputFile!);
 			}
 
 			IOSToolChain.PostBuildSync(Target, Logger);
 
-			return Task.FromResult(0);
+			return 0;
 		}
 	}
 }
