UnrealBuildTool: Fix issue CA1850 (Prefer static HashData method over ComputeHash).

#Synced-CL 50156122
#Synced-CL 50156122

diff --git a/Engine/Source/Programs/UnrealBuildTool/ProjectFiles/VisualStudio/VCProjectFileGenerator.cs b/Engine/Source/Programs/UnrealBuildTool/ProjectFiles/VisualStudio/VCProjectFileGenerator.cs
--- a/Engine/Source/Programs/UnrealBuildTool/ProjectFiles/VisualStudio/VCProjectFileGenerator.cs
+++ b/Engine/Source/Programs/UnrealBuildTool/ProjectFiles/VisualStudio/VCProjectFileGenerator.cs
@@ -580,7 +580,7 @@ namespace UnrealBuildTool
 		static Guid MakeMd5Guid(byte[] Input)
 		{
 #pragma warning disable CA5351 // Not used for cryptography, would be disruptive to users to change hash algo
-			byte[] Hash = MD5.Create().ComputeHash(Input);
+			byte[] Hash = MD5.HashData(Input);
 #pragma warning restore CA5351
 			Hash[6] = (byte)(0x30 | (Hash[6] & 0x0f)); // 0b0011'xxxx Version 3 UUID (MD5)
 			Hash[8] = (byte)(0x80 | (Hash[8] & 0x3f)); // 0b10xx'xxxx RFC 4122 UUID
