UnrealBuildTool: Enable enum types conformance for BuildSettingsVersion.V7

#rnx

#Synced-CL 50170575

diff --git a/Engine/Source/Programs/UnrealBuildTool/Configuration/Rules/TargetRules.cs b/Engine/Source/Programs/UnrealBuildTool/Configuration/Rules/TargetRules.cs
--- a/Engine/Source/Programs/UnrealBuildTool/Configuration/Rules/TargetRules.cs
+++ b/Engine/Source/Programs/UnrealBuildTool/Configuration/Rules/TargetRules.cs
@@ -3701,6 +3701,7 @@ namespace UnrealBuildTool
 
 				if (BuildSettingsVersion.V7 <= latestVersion && DefaultBuildSettings < BuildSettingsVersion.V7)
 				{
+					modifiedSettings.Add(Tuple.Create($"{nameof(WindowsPlatform)}.{nameof(WindowsPlatform.bStrictEnumTypesConformance)} = true", "Updates MSVC enum types conformance to true (Previously: false)."));
 					modifiedSettings.Add(Tuple.Create(String.Format("{0} = WarningLevel.Error", nameof(CppCompileWarningSettings.ReturnTypeWarningLevel)), "Validates that non-void functions return a value. (Previously: Off for VCClang, Error for MSVC and non-VC Clang)."));
 				}
 
diff --git a/Engine/Source/Programs/UnrealBuildTool/Platform/Windows/UEBuildWindows.cs b/Engine/Source/Programs/UnrealBuildTool/Platform/Windows/UEBuildWindows.cs
--- a/Engine/Source/Programs/UnrealBuildTool/Platform/Windows/UEBuildWindows.cs
+++ b/Engine/Source/Programs/UnrealBuildTool/Platform/Windows/UEBuildWindows.cs
@@ -655,7 +655,13 @@ namespace UnrealBuildTool
 		[RequiresUniqueBuildEnvironment]
 		[XmlConfigFile(Category = "WindowsPlatform")]
 		[CommandLine("-StrictEnumTypes")]
-		public bool bStrictEnumTypesConformance { get; set; } = false;
+		[CommandLine("-NoStrictEnumTypes", Value = "false")]
+		public bool bStrictEnumTypesConformance
+		{
+			get => bStrictEnumTypesConformancePrivate ?? Target.DefaultBuildSettings >= BuildSettingsVersion.V7;
+			set => bStrictEnumTypesConformancePrivate = value;
+		}
+		private bool? bStrictEnumTypesConformancePrivate;
 
 		/// <summary>
 		/// Enables enforcing standard C++ ODR violations (/Zc:checkGwOdr).
