UnrealBuildTool: Changes to try and fix the first-run project generation issue.

#Synced-CL 50156122

diff --git a/Engine/Source/Programs/Shared/EpicGames.Build/DotnetProjectInfo.cs b/Engine/Source/Programs/Shared/EpicGames.Build/DotnetProjectInfo.cs
--- a/Engine/Source/Programs/Shared/EpicGames.Build/DotnetProjectInfo.cs
+++ b/Engine/Source/Programs/Shared/EpicGames.Build/DotnetProjectInfo.cs
@@ -131,14 +131,14 @@ public sealed record class DotnetProjectInfo : IDotnetProjectInfo
 		ProjectOptions projectOptions = new()
 		{
 			GlobalProperties = globalPropertiesWithDefaults,
-			LoadSettings = ProjectLoadSettings.IgnoreMissingImports | ProjectLoadSettings.IgnoreInvalidImports | ProjectLoadSettings.DoNotEvaluateElementsWithFalseCondition,
+			LoadSettings = ProjectLoadSettings.DoNotEvaluateElementsWithFalseCondition,
 			ProjectCollection = projectCollection,
 		};
 		Project project = Project.FromFile(file.FullName, projectOptions);
 
 		Dictionary<string, string> properties = project.Properties.ToDictionary(p => p.Name, p => p.EvaluatedValue);
 
-		DirectoryReference projectDirectory = new(project.DirectoryPath);
+		DirectoryReference projectDirectory = file.Directory;
 		FileReference GetFullPath(string path) =>
 			Path.IsPathFullyQualified(path) ? new FileReference(path) : FileReference.Combine(projectDirectory, path);
 
diff --git a/Engine/Source/Programs/UnrealBuildTool/ProjectFiles/ProjectFileGenerator.cs b/Engine/Source/Programs/UnrealBuildTool/ProjectFiles/ProjectFileGenerator.cs
--- a/Engine/Source/Programs/UnrealBuildTool/ProjectFiles/ProjectFileGenerator.cs
+++ b/Engine/Source/Programs/UnrealBuildTool/ProjectFiles/ProjectFileGenerator.cs
@@ -1530,7 +1530,7 @@ namespace UnrealBuildTool
 			// Currently only implemented for Visual Studio.
 		}
 
-		protected void AddAutomationProjects(List<FileReference> AllGameProjects)
+		protected void AddAutomationProjects(List<FileReference> AllGameProjects, ILogger Logger)
 		{
 			using ITimelineEvent AutomationProjectTimer = Timeline.ScopeEvent("Adding Automation Projects");
 
@@ -1550,7 +1550,7 @@ namespace UnrealBuildTool
 
 			while (DotnetProjectTasks.Count != 0)
 			{
-				DotnetProjectTasks = GetDotnetDependencies(RootFolder, DependenciesFolder);
+				DotnetProjectTasks = GetDotnetDependencies(RootFolder, DependenciesFolder, Logger);
 				Task.WaitAll([.. DotnetProjectTasks]);
 			}
 
@@ -1562,7 +1562,7 @@ namespace UnrealBuildTool
 		/// <summary>
 		/// Get a single level of dotnet dependency projects. Should be called recursively until it yields no more results.
 		/// </summary>
-		protected List<Task> GetDotnetDependencies(PrimaryProjectFolder FolderToScan, PrimaryProjectFolder FolderToAddTo)
+		protected List<Task> GetDotnetDependencies(PrimaryProjectFolder FolderToScan, PrimaryProjectFolder FolderToAddTo, ILogger Logger)
 		{
 			HashSet<FileReference> ProjectFileReferences = [];
 			HashSet<FileReference> NewProjects = [];
@@ -1581,6 +1581,16 @@ namespace UnrealBuildTool
 					if (!ProjectFileReferences.Contains(ProjectReference) && !NewProjects.Contains(ProjectReference))
 					{
 						NewProjects.Add(ProjectReference);
+
+						if (!FileReference.Exists(ProjectReference))
+						{
+							Logger.LogWarning(
+								"Project file '{Dependency}' is depended on by '{Depender}' but does not exist on disk, skipping.",
+								ProjectReference,
+								Project.ProjectFilePath);
+							continue;
+						}
+
 						ProjectTasks.Add(Task.Run(() =>
 						{
 							using ITimelineEvent _ = Timeline.ScopeEvent(ProjectReference.GetFileName());
diff --git a/Engine/Source/Programs/UnrealBuildTool/ProjectFiles/VisualStudio/VCProjectFileGenerator.cs b/Engine/Source/Programs/UnrealBuildTool/ProjectFiles/VisualStudio/VCProjectFileGenerator.cs
--- a/Engine/Source/Programs/UnrealBuildTool/ProjectFiles/VisualStudio/VCProjectFileGenerator.cs
+++ b/Engine/Source/Programs/UnrealBuildTool/ProjectFiles/VisualStudio/VCProjectFileGenerator.cs
@@ -475,7 +475,7 @@ namespace UnrealBuildTool
 				AutomationProjectFiles.Clear();
 				OtherProjectFiles.Clear();
 
-				AddAutomationProjects(AllGameProjects);
+				AddAutomationProjects(AllGameProjects, Logger);
 
 				using (Timeline.ScopeEvent("Writing automation project files", Logger))
 				{
diff --git a/Engine/Source/Programs/UnrealBuildTool/ProjectFiles/VisualStudio/VCSharpProjectFile.cs b/Engine/Source/Programs/UnrealBuildTool/ProjectFiles/VisualStudio/VCSharpProjectFile.cs
--- a/Engine/Source/Programs/UnrealBuildTool/ProjectFiles/VisualStudio/VCSharpProjectFile.cs
+++ b/Engine/Source/Programs/UnrealBuildTool/ProjectFiles/VisualStudio/VCSharpProjectFile.cs
@@ -154,7 +154,8 @@ internal class VCSharpProjectFile : MSBuildProjectFile
 		Dictionary<string, string> Properties = [];
 		Properties.Add("Platform", "AnyCPU");
 		Properties.Add("Configuration", InConfiguration.ToString());
-		Properties.Add("EngineDirectory", Unreal.EngineDirectory.FullName);
+		Properties.Add("EngineDir", Unreal.EngineDirectory.FullName);
+		Properties.Add("EnginePath", Unreal.EngineDirectory.FullName);
 
 		DotnetProjectInfo Info = DotnetProjectInfo.Read(ProjectFilePath, Properties);
 		_cachedProjectInfo.Add(InConfiguration, Info);
