UnrealBuildTool: Suppress logger output from TargetRules when generating project files

#rnx

#Synced-CL 50213197

diff --git a/Engine/Source/Programs/UnrealBuildTool/Configuration/Rules/RulesAssembly.cs b/Engine/Source/Programs/UnrealBuildTool/Configuration/Rules/RulesAssembly.cs
--- a/Engine/Source/Programs/UnrealBuildTool/Configuration/Rules/RulesAssembly.cs
+++ b/Engine/Source/Programs/UnrealBuildTool/Configuration/Rules/RulesAssembly.cs
@@ -11,6 +11,7 @@ using System.Runtime.CompilerServices;
 using System.Threading.Tasks;
 using EpicGames.Core;
 using Microsoft.Extensions.Logging;
+using Microsoft.Extensions.Logging.Abstractions;
 using UnrealBuildBase;
 
 namespace UnrealBuildTool
@@ -911,18 +912,21 @@ namespace UnrealBuildTool
 		/// <param name="Architectures">Architectures being built</param>
 		/// <param name="ProjectFile">Path to the project file for this target</param>
 		/// <param name="Arguments">Command line arguments for this target</param>
-		/// <param name="Logger"></param>
+		/// <param name="InLogger"></param>
 		/// <param name="IsTestTarget">If building a low level test target</param>
 		/// <param name="ValidationOptions">Controls validation of target and SDK</param>
 		/// <param name="IntermediateEnvironment">Intermediate environment to use</param>
 		/// <returns>The build target rules for the specified target</returns>
-		public TargetRules CreateTargetRules(string TargetName, UnrealTargetPlatform Platform, UnrealTargetConfiguration Configuration, UnrealArchitectures? Architectures, FileReference? ProjectFile, CommandLineArguments? Arguments, ILogger Logger, bool IsTestTarget = false, TargetRulesValidationOptions ValidationOptions = TargetRulesValidationOptions.ValidateTargetAndSDK, UnrealIntermediateEnvironment IntermediateEnvironment = UnrealIntermediateEnvironment.Default)
+		public TargetRules CreateTargetRules(string TargetName, UnrealTargetPlatform Platform, UnrealTargetConfiguration Configuration, UnrealArchitectures? Architectures, FileReference? ProjectFile, CommandLineArguments? Arguments, ILogger InLogger, bool IsTestTarget = false, TargetRulesValidationOptions ValidationOptions = TargetRulesValidationOptions.ValidateTargetAndSDK, UnrealIntermediateEnvironment IntermediateEnvironment = UnrealIntermediateEnvironment.Default)
 		{
 			if (IsTestTarget)
 			{
 				TargetName = TargetDescriptor.GetTestedName(TargetName);
 			}
 
+			// Suppress TargetRules logging during project gen
+			ILogger Logger = ProjectFileGenerator.bGenerateProjectFiles ? NullLogger.Instance : InLogger;
+
 			Architectures ??= UnrealArchitectureConfig.ForPlatform(Platform).ActiveArchitectures(ProjectFile, TargetName);
 
 			bool bFoundTargetName = TargetNameToTargetFile.ContainsKey(TargetName);
