[UBT]
* Added -fno-validate-pch to portable toolchain path since pch is compatible between platforms but clang has checks that fails

#Synced-CL 50216544

diff --git a/Engine/Source/Programs/UnrealBuildTool/Platform/Apple/AppleToolChain.cs b/Engine/Source/Programs/UnrealBuildTool/Platform/Apple/AppleToolChain.cs
--- a/Engine/Source/Programs/UnrealBuildTool/Platform/Apple/AppleToolChain.cs
+++ b/Engine/Source/Programs/UnrealBuildTool/Platform/Apple/AppleToolChain.cs
@@ -606,6 +606,8 @@ namespace UnrealBuildTool
 				return;
 			}
 
+			Arguments.Add("-Xclang -fno-validate-pch"); // Needed since pch is actually portable but there is a check in clang that prevents it
+			
 			Arguments.Add("-nostdinc");
 
 			// FreeBSD math.h compatibility
