UnrealBuildTool: Minor cleanup of CleanStaleModules

#rnx

#Synced-CL 50213197

diff --git a/Engine/Source/Programs/UnrealBuildTool/Configuration/UEBuildTarget.cs b/Engine/Source/Programs/UnrealBuildTool/Configuration/UEBuildTarget.cs
--- a/Engine/Source/Programs/UnrealBuildTool/Configuration/UEBuildTarget.cs
+++ b/Engine/Source/Programs/UnrealBuildTool/Configuration/UEBuildTarget.cs
@@ -3705,46 +3705,38 @@ namespace UnrealBuildTool
 		}
 
 		/// <summary>
-		/// Cleans any stale modules that have changed moved output folder.
+		/// Cleans any stale modules that have changed output folder.
 		/// 
 		/// On Windows, the loader reads imported DLLs from the first location it finds them. If modules are moved from one place to another, we have to be sure to clean up the old versions
 		/// so that they're not loaded accidentally causing unintuitive import errors.
 		/// </summary>
-		void CleanStaleModules(ILogger Logger)
+		void CleanStaleModules(ILogger logger)
 		{
 			// Find all the output files
-			HashSet<FileReference> OutputFiles = new HashSet<FileReference>();
-			foreach (UEBuildBinary Binary in Binaries)
-			{
-				OutputFiles.UnionWith(Binary.OutputFilePaths);
-			}
+			HashSet<FileReference> outputFiles = [.. Binaries.SelectMany(x => x.OutputFilePaths)];
 
 			// Build a map of base filenames to their full path
-			Dictionary<string, FileReference> OutputNameToLocation = new Dictionary<string, FileReference>(StringComparer.InvariantCultureIgnoreCase);
-			foreach (FileReference OutputFile in OutputFiles)
+			Dictionary<string, FileReference> outputNameToLocation = new(StringComparer.InvariantCultureIgnoreCase);
+			foreach (FileReference outputFile in outputFiles)
 			{
-				OutputNameToLocation[OutputFile.GetFileName()] = OutputFile;
+				outputNameToLocation[outputFile.GetFileName()] = outputFile;
 			}
 
 			// Search all the output directories for files with a name matching one of our output files
-			Parallel2.ForEach(OutputFiles.Select(x => x.Directory).Distinct(), OutputDirectory =>
+			Parallel2.ForEach(outputFiles.Select(x => x.Directory).Distinct().Where(DirectoryReference.Exists), outputDirectory =>
 			{
-				if (DirectoryReference.Exists(OutputDirectory))
+				foreach (FileReference existingFile in DirectoryReference.EnumerateFiles(outputDirectory))
 				{
-					foreach (FileReference ExistingFile in DirectoryReference.EnumerateFiles(OutputDirectory))
+					if (outputNameToLocation.TryGetValue(existingFile.GetFileName(), out FileReference? OutputFile) && !outputFiles.Contains(existingFile))
 					{
-						FileReference? OutputFile;
-						if (OutputNameToLocation.TryGetValue(ExistingFile.GetFileName(), out OutputFile) && !OutputFiles.Contains(ExistingFile))
+						logger.LogInformation("Deleting '{ExistingFile}' to avoid ambiguity with '{OutputFile}'", existingFile, OutputFile);
+						try
 						{
-							Logger.LogInformation("Deleting '{ExistingFile}' to avoid ambiguity with '{OutputFile}'", ExistingFile, OutputFile);
-							try
-							{
-								FileReference.Delete(ExistingFile);
-							}
-							catch (Exception Ex)
-							{
-								Logger.LogError("Unable to delete {File} ({Ex})", ExistingFile, Ex.Message.TrimEnd());
-							}
+							FileReference.Delete(existingFile);
+						}
+						catch (Exception Ex)
+						{
+							logger.LogError("Unable to delete {File} ({Ex})", existingFile, Ex.Message.TrimEnd());
 						}
 					}
 				}
