UnrealBuildTool: Generic VisualStudio compiler enum which will select the preferred compiler from any available version of VS

#rnx

#Synced-CL 50213197

diff --git a/Engine/Source/Programs/UnrealBuildTool/Platform/Windows/UEBuildWindows.cs b/Engine/Source/Programs/UnrealBuildTool/Platform/Windows/UEBuildWindows.cs
--- a/Engine/Source/Programs/UnrealBuildTool/Platform/Windows/UEBuildWindows.cs
+++ b/Engine/Source/Programs/UnrealBuildTool/Platform/Windows/UEBuildWindows.cs
@@ -148,10 +148,15 @@ namespace UnrealBuildTool
 		VisualStudio2022,
 
 		/// <summary>
-		/// Visual Studio 2026 (Visual C++ 19.0)
+		/// Visual Studio 2026 (Visual C++ 19.0) or older
 		/// </summary>
 		VisualStudio2026,
 
+		/// <summary>
+		/// Any version of Visual Studio available
+		/// </summary>
+		VisualStudio = VisualStudio2026, // Should be set to the latest version of VisualStudio
+
 		/// <summary>
 		/// Unsupported Visual Studio
 		/// Must be the last entry in WindowsCompiler enum and should only be used in limited circumstances
