UnrealBuildTool: Format build results report JSON with an encoder that's easier for humans to read the output of.

#Synced-CL 50246058

diff --git a/Engine/Source/Programs/UnrealBuildTool/Modes/ReportBuildResultsMode.cs b/Engine/Source/Programs/UnrealBuildTool/Modes/ReportBuildResultsMode.cs
--- a/Engine/Source/Programs/UnrealBuildTool/Modes/ReportBuildResultsMode.cs
+++ b/Engine/Source/Programs/UnrealBuildTool/Modes/ReportBuildResultsMode.cs
@@ -5,6 +5,7 @@ using System.Collections.Generic;
 using System.Collections.Immutable;
 using System.IO;
 using System.Reflection;
+using System.Text.Encodings.Web;
 using System.Text.Json;
 using System.Threading;
 using System.Threading.Tasks;
@@ -31,7 +32,10 @@ internal class ReportBuildResultsMode : ToolMode
 
 		JsonSerializerOptions options = new JsonSerializerOptions()
 		{
-			WriteIndented = true
+			WriteIndented = true,
+			// It's very unlikely that the report would be consumed in a way that would be unsafe by using the below encoding,
+			// and the default "safe" behaviour severely hurts readability.
+			Encoder = JavaScriptEncoder.UnsafeRelaxedJsonEscaping,
 		};
 
 		foreach (TargetDescriptor targetDescriptor in targetDescriptors)
