UnrealBuildTool: Capture response file contents as part of Action and ActionResult objects, to allow easier introspection of what UBT is doing.

#Synced-CL 50246058

diff --git a/Engine/Source/Programs/UnrealBuildTool/Actions/Action.cs b/Engine/Source/Programs/UnrealBuildTool/Actions/Action.cs
--- a/Engine/Source/Programs/UnrealBuildTool/Actions/Action.cs
+++ b/Engine/Source/Programs/UnrealBuildTool/Actions/Action.cs
@@ -2,6 +2,7 @@
 
 using System;
 using System.Collections.Generic;
+using System.Collections.Immutable;
 using System.Diagnostics;
 using System.Linq;
 using System.Threading;
@@ -120,6 +121,12 @@ namespace UnrealBuildTool
 		/// </summary>
 		string CommandArguments { get; }
 
+		/// <summary>
+		/// Any arguments used in a response file to specify arguments to the executable.
+		/// The response file is expected to be passed as part of <see cref="CommandArguments"/>.
+		/// </summary>
+		ImmutableArray<string> ResponseFileContents { get; }
+
 		/// <summary>
 		/// Version of the command used for this action. This will be considered a dependency.
 		/// </summary>
@@ -279,6 +286,8 @@ namespace UnrealBuildTool
 		/// </summary>
 		public string CommandArguments { get; set; } = null!;
 
+		public ImmutableArray<string> ResponseFileContents { get; set; } = ImmutableArray<string>.Empty;
+
 		/// <summary>
 		/// Version of the command used for this action. This will be considered a dependency.
 		/// </summary>
@@ -366,7 +375,7 @@ namespace UnrealBuildTool
 		IEnumerable<FileItem> IExternalAction.PrerequisiteItems => PrerequisiteItems;
 		IEnumerable<FileItem> IExternalAction.ProducedItems => ProducedItems;
 		IEnumerable<FileItem> IExternalAction.DeleteItems => DeleteItems;
-        public Action(ActionType InActionType)
+		public Action(ActionType InActionType)
 		{
 			ActionType = InActionType;
 
@@ -389,6 +398,7 @@ namespace UnrealBuildTool
 			WorkingDirectory = InOther.WorkingDirectory;
 			CommandPath = InOther.CommandPath;
 			CommandArguments = InOther.CommandArguments;
+			ResponseFileContents = InOther.ResponseFileContents;
 			CommandVersion = InOther.CommandVersion;
 			CommandDescription = InOther.CommandDescription;
 			StatusDescription = InOther.StatusDescription;
@@ -416,6 +426,7 @@ namespace UnrealBuildTool
 			WorkingDirectory = Reader.ReadDirectoryReferenceNotNull();
 			CommandPath = Reader.ReadFileReference();
 			CommandArguments = Reader.ReadString()!;
+			ResponseFileContents = [.. Reader.ReadList(() => Reader.ReadString()!)!];
 			CommandVersion = Reader.ReadString()!;
 			CommandDescription = Reader.ReadString()!;
 			StatusDescription = Reader.ReadString()!;
@@ -451,6 +462,7 @@ namespace UnrealBuildTool
 			Writer.WriteDirectoryReference(WorkingDirectory);
 			Writer.WriteFileReference(CommandPath);
 			Writer.WriteString(CommandArguments);
+			Writer.WriteList(ResponseFileContents, Writer.WriteString);
 			Writer.WriteString(CommandVersion);
 			Writer.WriteString(CommandDescription);
 			Writer.WriteString(StatusDescription);
@@ -504,6 +516,11 @@ namespace UnrealBuildTool
 				Action.CommandArguments = CommandArguments;
 			}
 
+			if (Object.TryGetStringArrayField("ResponseFileContents", out string[]? ResponseFileContents))
+			{
+				Action.ResponseFileContents = [.. ResponseFileContents];
+			}
+
 			if (Object.TryGetStringField("CommandVersion", out string? CommandVersion))
 			{
 				Action.CommandVersion = CommandVersion;
@@ -663,6 +680,7 @@ namespace UnrealBuildTool
 			Writer.WriteValue("WorkingDirectory", Action.WorkingDirectory.FullName);
 			Writer.WriteValue("CommandPath", Action.CommandPath.FullName);
 			Writer.WriteValue("CommandArguments", Action.CommandArguments);
+			Writer.WriteStringArrayField("ResponseFileContents", Action.ResponseFileContents);
 			Writer.WriteValue("CommandVersion", Action.CommandVersion);
 			Writer.WriteValue("CommandDescription", Action.CommandDescription);
 			Writer.WriteValue("StatusDescription", Action.StatusDescription);
@@ -787,6 +805,7 @@ namespace UnrealBuildTool
 		public DirectoryReference WorkingDirectory => Inner.WorkingDirectory;
 		public FileReference CommandPath => Inner.CommandPath;
 		public string CommandArguments => Inner.CommandArguments;
+		public ImmutableArray<string> ResponseFileContents => Inner.ResponseFileContents;
 		public string CommandVersion => Inner.CommandVersion;
 		public string CommandDescription => Inner.CommandDescription;
 		public string StatusDescription => Inner.StatusDescription;
diff --git a/Engine/Source/Programs/UnrealBuildTool/Actions/ActionResult.cs b/Engine/Source/Programs/UnrealBuildTool/Actions/ActionResult.cs
--- a/Engine/Source/Programs/UnrealBuildTool/Actions/ActionResult.cs
+++ b/Engine/Source/Programs/UnrealBuildTool/Actions/ActionResult.cs
@@ -1,6 +1,7 @@
 // Copyright Epic Games, Inc. All Rights Reserved.
 
 using System.Collections.Generic;
+using System.Collections.Immutable;
 using EpicGames.Core;
 
 namespace UnrealBuildTool;
@@ -10,6 +11,7 @@ namespace UnrealBuildTool;
 /// </summary>
 /// <param name="CommandPath">The file path of the command that was ran.</param>
 /// <param name="CommandArguments">The arguments to the command that was ran.</param>
+/// <param name="ResponseFileContents">Contents of a response file, if one was passed as an argument.</param>
 /// <param name="CommandVersion">The version of the command that was ran.</param>
 /// <param name="Description">A human-readable description of the action.</param>
 /// <param name="ExitCode">The exit code the action returned after executing.</param>
@@ -19,6 +21,7 @@ namespace UnrealBuildTool;
 internal record class ActionResult(
 	FileReference CommandPath,
 	string CommandArguments,
+	ImmutableArray<string> ResponseFileContents,
 	string CommandVersion,
 	string Description,
 	int ExitCode,
@@ -32,6 +35,7 @@ internal record class ActionResult(
 		return new(
 			executeResults.Action.CommandPath,
 			executeResults.Action.CommandArguments,
+			executeResults.Action.ResponseFileContents,
 			executeResults.Action.CommandVersion,
 			executeResults.GetDescription(),
 			executeResults.ExitCode,
diff --git a/Engine/Source/Programs/UnrealBuildTool/Platform/Android/AndroidToolChain.cs b/Engine/Source/Programs/UnrealBuildTool/Platform/Android/AndroidToolChain.cs
--- a/Engine/Source/Programs/UnrealBuildTool/Platform/Android/AndroidToolChain.cs
+++ b/Engine/Source/Programs/UnrealBuildTool/Platform/Android/AndroidToolChain.cs
@@ -1372,7 +1372,8 @@ namespace UnrealBuildTool
 			FileReference ResponseFileName = GetResponseFileName(LinkEnvironment, OutputFile);
 			InputFileNames.Add(LinkResponseArguments.Replace("\\", "/", StringComparison.Ordinal));
 
-			FileItem ResponseFileItem = Graph.CreateIntermediateTextFile(ResponseFileName, InputFileNames);
+			LinkAction.ResponseFileContents = [.. InputFileNames];
+			FileItem ResponseFileItem = Graph.CreateIntermediateTextFile(ResponseFileName, LinkAction.ResponseFileContents);
 
 			LinkAction.CommandArguments += String.Format(" @\"{0}\"", ResponseFileName);
 			LinkAction.PrerequisiteItems.Add(ResponseFileItem);
diff --git a/Engine/Source/Programs/UnrealBuildTool/Platform/Clang/ClangToolChain.cs b/Engine/Source/Programs/UnrealBuildTool/Platform/Clang/ClangToolChain.cs
--- a/Engine/Source/Programs/UnrealBuildTool/Platform/Clang/ClangToolChain.cs
+++ b/Engine/Source/Programs/UnrealBuildTool/Platform/Clang/ClangToolChain.cs
@@ -1462,7 +1462,8 @@ namespace UnrealBuildTool
 			}
 
 			// Adds the response file to the compiler input.
-			FileItem CompilerResponseFileItem = Graph.CreateIntermediateTextFile(ResponseFileName, ResponseFileContents);
+			CompileAction.ResponseFileContents = [.. ResponseFileContents];
+			FileItem CompilerResponseFileItem = Graph.CreateIntermediateTextFile(ResponseFileName, CompileAction.ResponseFileContents);
 			string CommandArguments = GetResponseFileArgument(CompilerResponseFileItem, CompileEnvironment.RootPaths);
 
 			if (bMergeModules && CompileEnvironment.PrecompiledHeaderAction != PrecompiledHeaderAction.Create)
diff --git a/Engine/Source/Programs/UnrealBuildTool/Platform/IOS/IOSToolChain.cs b/Engine/Source/Programs/UnrealBuildTool/Platform/IOS/IOSToolChain.cs
--- a/Engine/Source/Programs/UnrealBuildTool/Platform/IOS/IOSToolChain.cs
+++ b/Engine/Source/Programs/UnrealBuildTool/Platform/IOS/IOSToolChain.cs
@@ -488,13 +488,14 @@ namespace UnrealBuildTool
 			{
 				ResponseFileName = GetResponseFileName(LinkEnvironment, OutputFile);
 			}
-			
-			FileItem ResponseFileItem = Graph.CreateIntermediateTextFile(ResponseFileName, LinkArguments);
-			string ResponsFileArgument = GetResponseFileArgument(ResponseFileItem, LinkEnvironment.RootPaths);
+
+			LinkAction.ResponseFileContents = [.. LinkArguments];
+			FileItem ResponseFileItem = Graph.CreateIntermediateTextFile(ResponseFileName, LinkAction.ResponseFileContents);
+			string ResponseFileArgument = GetResponseFileArgument(ResponseFileItem, LinkEnvironment.RootPaths);
 
 			LinkAction.PrerequisiteItems.Add(ResponseFileItem);
 
-			LinkAction.CommandArguments = ResponsFileArgument;
+			LinkAction.CommandArguments = ResponseFileArgument;
 
 			// Only execute linking on the local PC.
 			LinkAction.bCanExecuteRemotely = false;
diff --git a/Engine/Source/Programs/UnrealBuildTool/Platform/Linux/LinuxToolChain.cs b/Engine/Source/Programs/UnrealBuildTool/Platform/Linux/LinuxToolChain.cs
--- a/Engine/Source/Programs/UnrealBuildTool/Platform/Linux/LinuxToolChain.cs
+++ b/Engine/Source/Programs/UnrealBuildTool/Platform/Linux/LinuxToolChain.cs
@@ -929,7 +929,8 @@ namespace UnrealBuildTool
 			FileReference ResponsePath = GetResponseFileName(LinkEnvironment, OutputFile);
 			if (!ProjectFileGenerator.bGenerateProjectFiles)
 			{
-				FileItem ResponseFileItem = Graph.CreateIntermediateTextFile(ResponsePath, InputFileNames);
+				ArchiveAction.ResponseFileContents = [.. InputFileNames];
+				FileItem ResponseFileItem = Graph.CreateIntermediateTextFile(ResponsePath, ArchiveAction.ResponseFileContents);
 				ArchiveAction.PrerequisiteItems.Add(ResponseFileItem);
 			}
 			ArchiveAction.CommandArguments += $" @\"{ResponsePath.FullName}\"";
@@ -1246,7 +1247,8 @@ namespace UnrealBuildTool
 			ResponseLines.Add("--end-group");
 
 			FileReference ResponseFileName = GetResponseFileName(LinkEnvironment, OutputFile);
-			FileItem ResponseFileItem = Graph.CreateIntermediateTextFile(ResponseFileName, ResponseLines);
+			LinkAction.ResponseFileContents = [.. ResponseLines];
+			FileItem ResponseFileItem = Graph.CreateIntermediateTextFile(ResponseFileName, LinkAction.ResponseFileContents);
 
 			LinkArguments.Add($"-Wl,@\"{NormalizeCommandLinePath(ResponseFileName, LinkEnvironment.RootPaths)}\"");
 			LinkAction.PrerequisiteItems.Add(ResponseFileItem);
@@ -1312,7 +1314,8 @@ namespace UnrealBuildTool
 
 			{
 				FileReference LinkResponseFileName = GetResponseFileName(LinkEnvironment, FileItem.GetItemByFileReference(FileReference.Combine(OutputFile.Location.Directory, $"{OutputFile.Name}.link")));
-				FileItem LinkResponseFileItem = Graph.CreateIntermediateTextFile(LinkResponseFileName, LinkArguments);
+				LinkAction.ResponseFileContents = [.. LinkArguments];
+				FileItem LinkResponseFileItem = Graph.CreateIntermediateTextFile(LinkResponseFileName, LinkAction.ResponseFileContents);
 				LinkAction.PrerequisiteItems.Add(LinkResponseFileItem);
 
 				StringWriter LinkWriter = new();
@@ -1411,7 +1414,8 @@ namespace UnrealBuildTool
 					RelinkArguments = [.. RelinkArguments.Select(x => x.Replace(LinkOutputFileString, RelinkedFileString, StringComparison.Ordinal))];
 
 					FileReference RelinkResponseFileName = GetResponseFileName(LinkEnvironment, FileItem.GetItemByFileReference(FileReference.Combine(OutputFile.Location.Directory, $"{OutputFile.Name}.relink")));
-					FileItem RelinkResponseFileItem = Graph.CreateIntermediateTextFile(RelinkResponseFileName, RelinkArguments);
+					RelinkAction.ResponseFileContents = [.. RelinkArguments];
+					FileItem RelinkResponseFileItem = Graph.CreateIntermediateTextFile(RelinkResponseFileName, RelinkAction.ResponseFileContents);
 					RelinkAction.PrerequisiteItems.Add(RelinkResponseFileItem);
 
 					if (bUseCmdExe)
diff --git a/Engine/Source/Programs/UnrealBuildTool/Platform/Mac/MacToolChain.cs b/Engine/Source/Programs/UnrealBuildTool/Platform/Mac/MacToolChain.cs
--- a/Engine/Source/Programs/UnrealBuildTool/Platform/Mac/MacToolChain.cs
+++ b/Engine/Source/Programs/UnrealBuildTool/Platform/Mac/MacToolChain.cs
@@ -795,13 +795,14 @@ namespace UnrealBuildTool
 			{
 				ResponseFileName = GetResponseFileName(LinkEnvironment, OutputFile);
 			}
-			
-			FileItem ResponseFileItem = Graph.CreateIntermediateTextFile(ResponseFileName, LinkArguments);
-			string ResponsFileArgument = GetResponseFileArgument(ResponseFileItem, LinkEnvironment.RootPaths);
+
+			LinkAction.ResponseFileContents = [.. LinkArguments];
+			FileItem ResponseFileItem = Graph.CreateIntermediateTextFile(ResponseFileName, LinkAction.ResponseFileContents);
+			string ResponseFileArgument = GetResponseFileArgument(ResponseFileItem, LinkEnvironment.RootPaths);
 
 			LinkAction.PrerequisiteItems.Add(ResponseFileItem);
 
-			LinkAction.CommandArguments = ResponsFileArgument;
+			LinkAction.CommandArguments = ResponseFileArgument;
 
 			// Only execute linking on the local Mac.
 			LinkAction.bCanExecuteRemotely = false;
diff --git a/Engine/Source/Programs/UnrealBuildTool/Platform/Windows/VCCompileAction.cs b/Engine/Source/Programs/UnrealBuildTool/Platform/Windows/VCCompileAction.cs
--- a/Engine/Source/Programs/UnrealBuildTool/Platform/Windows/VCCompileAction.cs
+++ b/Engine/Source/Programs/UnrealBuildTool/Platform/Windows/VCCompileAction.cs
@@ -2,6 +2,7 @@
 
 using System;
 using System.Collections.Generic;
+using System.Collections.Immutable;
 using EpicGames.Core;
 using UnrealBuildBase;
 
@@ -339,6 +340,9 @@ namespace UnrealBuildTool
 		/// <inheritdoc/>
 		string IExternalAction.CommandArguments => UsingClFilter ? GetClFilterArguments() : GetClArguments();
 
+		/// <inheritdoc/>
+		ImmutableArray<string> IExternalAction.ResponseFileContents => GetResponseFileContents();
+
 		/// <inheritdoc/>
 		string IExternalAction.CommandVersion =>
 			CompilerType.IsMSVC()
@@ -619,6 +623,15 @@ namespace UnrealBuildTool
 			return $"@{Utils.MakePathSafeToUseWithCommandLine(responseFileString)} {String.Join(' ', Arguments)}";
 		}
 
+		private ImmutableArray<string> GetResponseFileContents()
+		{
+			if (ResponseFile != null)
+			{
+				return [.. GetCompilerArguments()];
+			}
+			return [];
+		}
+
 		string GetClFilterArguments()
 		{
 			List<string> arguments = [];
diff --git a/Engine/Source/Programs/UnrealBuildTool/Platform/Windows/VCToolChain.cs b/Engine/Source/Programs/UnrealBuildTool/Platform/Windows/VCToolChain.cs
--- a/Engine/Source/Programs/UnrealBuildTool/Platform/Windows/VCToolChain.cs
+++ b/Engine/Source/Programs/UnrealBuildTool/Platform/Windows/VCToolChain.cs
@@ -3041,7 +3041,10 @@ namespace UnrealBuildTool
 
 			AddIncludePathArguments(CompileEnvironment, EnvVars.ToolChain, Arguments);
 
-			FileItem ResponseFile = Graph.CreateIntermediateTextFile(OutputFile.ChangeExtension(ResponseExt), Arguments, false);
+			// Create the batch file action that will compile then touch the output file
+			Action CompileAction = Graph.CreateAction(ActionType.Compile);
+			CompileAction.ResponseFileContents = [.. Arguments];
+			FileItem ResponseFile = Graph.CreateIntermediateTextFile(OutputFile.ChangeExtension(ResponseExt), CompileAction.ResponseFileContents, false);
 
 			// Build the command for touching the output file(s) to update the time stamp
 			string OutfileTlhFullName = NormalizeCommandLinePath(OutputFile, RootPaths).Replace('/', '\\'); // Must be backward slash
@@ -3058,8 +3061,6 @@ namespace UnrealBuildTool
 
 			FileItem BatchFile = Graph.CreateIntermediateTextFile(OutputFile.ChangeExtension(".bat"), BatchFileContents.ToString(), false);
 
-			// Create the batch file action that will compile then touch the output file
-			Action CompileAction = Graph.CreateAction(ActionType.Compile);
 			CompileAction.CommandDescription = "GenerateTLH";
 			CompileAction.PrerequisiteItems.Add(FileItem.GetItemByFileReference(EnvVars.ToolchainCompilerPath));
 			CompileAction.PrerequisiteItems.Add(TypeLibraryFile);
@@ -3491,16 +3492,18 @@ namespace UnrealBuildTool
 				ProducedItems.AddRange(additional);
 			}
 
+			// Create an action that invokes the linker.
+			Action LinkAction = Graph.CreateAction(ActionType.Link);
+
 			// Create a response file for the linker, unless we're generating IntelliSense data
 			FileReference ResponseFileName = GetResponseFileName(LinkEnvironment, OutputFile);
 			if (!ProjectFileGenerator.bGenerateProjectFiles)
 			{
-				FileItem ResponseFile = Graph.CreateIntermediateTextFile(ResponseFileName, Arguments);
+				LinkAction.ResponseFileContents = [.. Arguments];
+				FileItem ResponseFile = Graph.CreateIntermediateTextFile(ResponseFileName, LinkAction.ResponseFileContents);
 				PrerequisiteItems.Add(ResponseFile);
 			}
 
-			// Create an action that invokes the linker.
-			Action LinkAction = Graph.CreateAction(ActionType.Link);
 			LinkAction.RootPaths = LinkEnvironment.RootPaths;
 			string ReadableArch = UnrealArchitectureConfig.ForPlatform(LinkEnvironment.Platform).ConvertToReadableArchitecture(LinkEnvironment.Architecture);
 			LinkAction.CommandDescription += $"Link [{ReadableArch}]";
diff --git a/Engine/Source/Programs/UnrealBuildTool/ToolChain/ISPCToolChain.cs b/Engine/Source/Programs/UnrealBuildTool/ToolChain/ISPCToolChain.cs
--- a/Engine/Source/Programs/UnrealBuildTool/ToolChain/ISPCToolChain.cs
+++ b/Engine/Source/Programs/UnrealBuildTool/ToolChain/ISPCToolChain.cs
@@ -625,7 +625,8 @@ namespace UnrealBuildTool
 				Arguments.AddRange(GlobalArguments);
 
 				FileReference ResponseFileName = GetResponseFileName(ISPCIncludeHeaderFile);
-				FileItem ResponseFileItem = Graph.CreateIntermediateTextFile(ResponseFileName, Arguments.Select(x => Utils.ExpandVariables(x)));
+				CompileAction.ResponseFileContents = [.. Arguments.Select(x => Utils.ExpandVariables(x))];
+				FileItem ResponseFileItem = Graph.CreateIntermediateTextFile(ResponseFileName, CompileAction.ResponseFileContents);
 				CompileAction.CommandArguments = $"@\"{NormalizeCommandLinePath(ResponseFileName, RootPaths)}\"";
 				CompileAction.PrerequisiteItems.Add(ResponseFileItem);
 
@@ -863,7 +864,8 @@ namespace UnrealBuildTool
 				CompileAction.ProducedItems.UnionWith(CompiledISPCObjFiles);
 
 				FileReference ResponseFileName = GetResponseFileName(CompiledISPCObjFileNoISA);
-				FileItem ResponseFileItem = Graph.CreateIntermediateTextFile(ResponseFileName, Arguments.Select(x => Utils.ExpandVariables(x)));
+				CompileAction.ResponseFileContents = [.. Arguments.Select(x => Utils.ExpandVariables(x))];
+				FileItem ResponseFileItem = Graph.CreateIntermediateTextFile(ResponseFileName, CompileAction.ResponseFileContents);
 
 				string AdditionalArguments = "";
 				// Must be added after response file is created just to make sure it ends up on the command line and not in the response file
@@ -921,7 +923,8 @@ namespace UnrealBuildTool
 
 							// Write the args to a response file
 							FileReference PostCompileResponseFileName = GetResponseFileName(FinalCompiledISPCObjFile);
-							FileItem PostCompileResponseFileItem = Graph.CreateIntermediateTextFile(PostCompileResponseFileName, PostCompileArgs.Select(x => Utils.ExpandVariables(x)));
+							PostCompileAction.ResponseFileContents = [.. PostCompileArgs.Select(x => Utils.ExpandVariables(x))];
+							FileItem PostCompileResponseFileItem = Graph.CreateIntermediateTextFile(PostCompileResponseFileName, PostCompileAction.ResponseFileContents);
 							PostCompileAction.CommandArguments = $"@\"{NormalizeCommandLinePath(PostCompileResponseFileName, RootPaths)}\"{AdditionalArguments}";
 							PostCompileAction.PrerequisiteItems.Add(PostCompileResponseFileItem);
 
