UnrealBuildTool: Fix issue CA1853 (Unnecessary call to 'Dictionary.ContainsKey(key)').

#Synced-CL 50246058

diff --git a/Engine/Source/Programs/UnrealBuildTool/Actions/ActionGraph.cs b/Engine/Source/Programs/UnrealBuildTool/Actions/ActionGraph.cs
--- a/Engine/Source/Programs/UnrealBuildTool/Actions/ActionGraph.cs
+++ b/Engine/Source/Programs/UnrealBuildTool/Actions/ActionGraph.cs
@@ -558,10 +558,7 @@ namespace UnrealBuildTool
 						{
 							ActionIsNonCyclical.Add(Action, true);
 							bFoundNewNonCyclicalAction = true;
-							if (CyclicActions.ContainsKey(Action))
-							{
-								CyclicActions.Remove(Action);
-							}
+							CyclicActions.Remove(Action);
 						}
 					}
 				}
