[UBT] Updated apple portable toolchain to register its paths as root paths so they work with VFS enabled.

Also tweaked the portable tool chain initialization logs to be more specific.

#Synced-CL 50261506

diff --git a/Engine/Source/Programs/UnrealBuildTool/Platform/Apple/AppleToolChain.cs b/Engine/Source/Programs/UnrealBuildTool/Platform/Apple/AppleToolChain.cs
--- a/Engine/Source/Programs/UnrealBuildTool/Platform/Apple/AppleToolChain.cs
+++ b/Engine/Source/Programs/UnrealBuildTool/Platform/Apple/AppleToolChain.cs
@@ -120,6 +120,7 @@ namespace UnrealBuildTool
 			if (ToolChainInfo.IsValid)
 			{
 				PortableToolchain = ToolChainInfo;
+				Logger.LogInformation("Portable toolchain initialized: {Version}", ToolChainInfo.SDKVersion);
 				RegisterUBAPathsForDistribution();
 			}
 		}
@@ -196,6 +197,15 @@ namespace UnrealBuildTool
 				// Xcode sdk files are not allowed to be transferred over the network, ensure UBA will not do so
 				EpicGames.UBA.Utils.RegisterDisallowedPaths(XcodeAppDir.FullName);
 			}
+
+			if (PortableToolchain != null)
+			{
+				// In the current setup, the parent directory of the clang binary, contains both clang, and apple oss header directories 
+				GlobalCompileEnvironment.RootPaths.AddExtraPath(("PortableToolchain", PortableToolchain.ClangBasePath!.ParentDirectory!.FullName));
+				GlobalLinkEnvironment.RootPaths.AddExtraPath(("PortableToolchain", PortableToolchain.ClangBasePath!.ParentDirectory!.FullName));
+				
+				EpicGames.UBA.Utils.RegisterPathHash(PortableToolchain.ClangBasePath!.ParentDirectory!.FullName, PortableToolchain.SDKVersion!);
+			}
 		}
 
 		/// <summary>
@@ -592,7 +602,7 @@ namespace UnrealBuildTool
 		/// Adds include paths for portable toolchain compilation with Apple OSS headers.
 		/// Used for distributed compilation on non-Apple workers.
 		/// </summary>
-		protected void AddPortableToolchainIncludePaths(List<string> Arguments)
+		protected void AddPortableToolchainIncludePaths(CppCompileEnvironment CompileEnvironment, List<string> Arguments)
 		{
 			if (PortableToolchain == null || PortableToolchain.OSSHeadersPath == null)
 			{
@@ -623,7 +633,7 @@ namespace UnrealBuildTool
 			// libc++ headers must come first for #include_next to work
 			if (PortableToolchain.LibCxxIncludePath != null && DirectoryReference.Exists(PortableToolchain.LibCxxIncludePath))
 			{
-				Arguments.Add($"-isystem \"{PortableToolchain.LibCxxIncludePath.FullName}\"");
+				Arguments.Add($"-isystem \"{NormalizeCommandLinePath(PortableToolchain.LibCxxIncludePath, CompileEnvironment.RootPaths)}\"");
 			}
 			else
 			{
@@ -633,7 +643,7 @@ namespace UnrealBuildTool
 			// Clang builtin headers (stddef.h, stdarg.h, intrinsics)
 			if (PortableToolchain.ClangBuiltinIncludePath != null && DirectoryReference.Exists(PortableToolchain.ClangBuiltinIncludePath))
 			{
-				Arguments.Add($"-isystem \"{PortableToolchain.ClangBuiltinIncludePath.FullName}\"");
+				Arguments.Add($"-isystem \"{NormalizeCommandLinePath(PortableToolchain.ClangBuiltinIncludePath,  CompileEnvironment.RootPaths)}\"");
 			}
 			else
 			{
@@ -641,7 +651,7 @@ namespace UnrealBuildTool
 			}
 
 			// Apple OSS headers
-			Arguments.Add($"-isystem \"{OSSHeadersDir.FullName}\"");
+			Arguments.Add($"-isystem \"{NormalizeCommandLinePath(OSSHeadersDir,CompileEnvironment.RootPaths)}\"");
 		}
 
 		/// <summary>
diff --git a/Engine/Source/Programs/UnrealBuildTool/Platform/IOS/IOSToolChain.cs b/Engine/Source/Programs/UnrealBuildTool/Platform/IOS/IOSToolChain.cs
--- a/Engine/Source/Programs/UnrealBuildTool/Platform/IOS/IOSToolChain.cs
+++ b/Engine/Source/Programs/UnrealBuildTool/Platform/IOS/IOSToolChain.cs
@@ -181,7 +181,7 @@ namespace UnrealBuildTool
 			// Clang modules depend on SDK paths unavailable on remote workers
 			if (bIsUsingPortableSDK)
 			{
-				AddPortableToolchainIncludePaths(Arguments);
+				AddPortableToolchainIncludePaths(CompileEnvironment, Arguments);
 			}
 			else
 			{
diff --git a/Engine/Source/Programs/UnrealBuildTool/Platform/Mac/ApplePortableToolChain.cs b/Engine/Source/Programs/UnrealBuildTool/Platform/Mac/ApplePortableToolChain.cs
--- a/Engine/Source/Programs/UnrealBuildTool/Platform/Mac/ApplePortableToolChain.cs
+++ b/Engine/Source/Programs/UnrealBuildTool/Platform/Mac/ApplePortableToolChain.cs
@@ -15,6 +15,7 @@ namespace UnrealBuildTool
 		public DirectoryReference? ClangBasePath { get; }
 		public DirectoryReference? LibCxxIncludePath { get; }
 		public DirectoryReference? ClangBuiltinIncludePath { get; }
+		public string? SDKVersion { get; }
 
 		public ApplePortableToolChainInfo(ILogger Logger)
 		{
@@ -43,14 +44,14 @@ namespace UnrealBuildTool
 				return;
 			}
 
-			string? Version = ResolveVersion(AppleToolchainDir, Logger);
-			if (Version == null)
+			SDKVersion = ResolveVersion(AppleToolchainDir, Logger);
+			if (SDKVersion == null)
 			{
 				Logger.LogDebug("No valid Apple portable toolchain version found in {Path}", AppleToolchainDir);
 				return;
 			}
 
-			DirectoryReference VersionDir = DirectoryReference.Combine(AppleToolchainDir, Version);
+			DirectoryReference VersionDir = DirectoryReference.Combine(AppleToolchainDir, SDKVersion);
 			FileReference Clang = FileReference.Combine(VersionDir, ClangSubdir);
 			DirectoryReference OssHeaders = DirectoryReference.Combine(VersionDir, OssHeadersSubdir);
 
@@ -98,7 +99,7 @@ namespace UnrealBuildTool
 				return;
 			}
 
-			Logger.LogInformation("Found Apple portable toolchain: {Version}", Version);
+			Logger.LogInformation("Found Apple portable toolchain: {Version}", SDKVersion);
 
 			ClangPath = Clang;
 			OSSHeadersPath = OssHeaders;
@@ -222,7 +223,6 @@ namespace UnrealBuildTool
 				string DirName = Dir.GetDirectoryName();
 				if (DirName.StartsWith(ToolchainPrefix + "_", StringComparison.OrdinalIgnoreCase))
 				{
-					Logger.LogInformation("Selected portable toolchain: {Dir}", DirName);
 					return DirName;
 				}
 			}
diff --git a/Engine/Source/Programs/UnrealBuildTool/Platform/Mac/MacToolChain.cs b/Engine/Source/Programs/UnrealBuildTool/Platform/Mac/MacToolChain.cs
--- a/Engine/Source/Programs/UnrealBuildTool/Platform/Mac/MacToolChain.cs
+++ b/Engine/Source/Programs/UnrealBuildTool/Platform/Mac/MacToolChain.cs
@@ -218,7 +218,7 @@ namespace UnrealBuildTool
 			// Clang modules depend on SDK paths unavailable on remote workers
 			if (bIsUsingPortableSDK)
 			{
-				AddPortableToolchainIncludePaths(Arguments);
+				AddPortableToolchainIncludePaths(CompileEnvironment, Arguments);
 			}
 			else
 			{
