UnrealBuildTool: Add ini config for bStoreObjFilesCompressed

#rnx

#Synced-CL 50271719

diff --git a/Engine/Source/Programs/UnrealBuildTool/Configuration/Rules/TargetRules.cs b/Engine/Source/Programs/UnrealBuildTool/Configuration/Rules/TargetRules.cs
--- a/Engine/Source/Programs/UnrealBuildTool/Configuration/Rules/TargetRules.cs
+++ b/Engine/Source/Programs/UnrealBuildTool/Configuration/Rules/TargetRules.cs
@@ -2203,6 +2203,7 @@ namespace UnrealBuildTool
 		/// Warning, this option is not currently compatitable with PGO or the the cl-clang linker as those are not detoured and linking will fail.
 		/// </summary>
 		[XmlConfigFile(Category = "UnrealBuildAccelerator", Name = "bStoreObjFilesCompressed")]
+		[ConfigFile(ConfigHierarchyType.Engine, "/Script/BuildSettings.BuildSettings", "bStoreObjFilesCompressed")]
 		[CommandLine("-UBAStoreObjFilesCompressed", Value = "true", MarkUsed = false)]
 		[RequiresUniqueBuildEnvironment]
 		public bool bAllowUbaCompression { get; set; }
diff --git a/Engine/Source/Programs/UnrealBuildTool/Executors/UnrealBuildAccelerator/UnrealBuildAcceleratorConfig.cs b/Engine/Source/Programs/UnrealBuildTool/Executors/UnrealBuildAccelerator/UnrealBuildAcceleratorConfig.cs
--- a/Engine/Source/Programs/UnrealBuildTool/Executors/UnrealBuildAccelerator/UnrealBuildAcceleratorConfig.cs
+++ b/Engine/Source/Programs/UnrealBuildTool/Executors/UnrealBuildAccelerator/UnrealBuildAcceleratorConfig.cs
@@ -244,6 +244,7 @@ namespace UnrealBuildTool
 		/// Store object (.obj) compressed on disk. Requires uba to do link step where it will decompress obj files again
 		/// </summary>
 		[XmlConfigFile(Category = "UnrealBuildAccelerator")]
+		[ConfigFile(ConfigHierarchyType.Engine, "/Script/BuildSettings.BuildSettings", "bStoreObjFilesCompressed")]
 		[CommandLine("-UBAStoreObjFilesCompressed", Value = "true")]
 		[CommandLine("-UBADisableStoreObjFilesCompressed", Value = "false")]
 		[CommandLine("-nocompress", Value = "false")]
