[UBT]
* Changed so a UBAAgentCoordinatorHorde can have multiple sessions if the sessions have the same server url/access token.
* Moved horde/okta logins over to UBAAgentCoordinatorHorde to prevent sessions from doing this multiple times.
* Split up UBAAgentCoordinatorHorde and UBAHordeSession into two files
* Added UnrealBuildAcceleratorHordeConfig.AgentPlatform that can be used to explicitly say what agent platform to use when using cluster auto.

#Synced-CL 50261005

diff --git a/Engine/Source/Programs/UnrealBuildTool/Executors/UnrealBuildAccelerator/UBAAgentCoordinatorHorde.cs b/Engine/Source/Programs/UnrealBuildTool/Executors/UnrealBuildAccelerator/UBAAgentCoordinatorHorde.cs
--- a/Engine/Source/Programs/UnrealBuildTool/Executors/UnrealBuildAccelerator/UBAAgentCoordinatorHorde.cs
+++ b/Engine/Source/Programs/UnrealBuildTool/Executors/UnrealBuildAccelerator/UBAAgentCoordinatorHorde.cs
@@ -2,24 +2,17 @@
 
 using System;
 using System.Collections.Generic;
-using System.Diagnostics;
 using System.Linq;
 using System.Runtime.InteropServices;
 using System.Threading;
 using System.Threading.Tasks;
 using EpicGames.Core;
 using EpicGames.Horde;
-using EpicGames.Horde.Agents;
-using EpicGames.Horde.Common;
 using EpicGames.Horde.Compute;
 using EpicGames.Horde.Compute.Clients;
 using EpicGames.Horde.Server;
-using EpicGames.Horde.Storage;
-using EpicGames.Horde.Storage.Bundles;
-using EpicGames.Horde.Storage.Nodes;
 using Microsoft.Extensions.DependencyInjection;
 using Microsoft.Extensions.Logging;
-using Microsoft.Extensions.Logging.Abstractions;
 using UnrealBuildBase;
 
 namespace UnrealBuildTool
@@ -27,1014 +20,363 @@ namespace UnrealBuildTool
 	// statusRow, statusColumn, statusText, statusType, statusLink
 	using StatusUpdateAction = Action<uint, uint, string, EpicGames.UBA.LogEntryType, string?>;
 
-	class UBAHordeSession : IAsyncDisposable
+	internal class UBAAgentCoordinatorHorde : IUBAAgentCoordinator, IDisposable
 	{
-		/// <summary>
-		/// ID for this Horde session
-		/// </summary>
-		readonly Guid _id = Guid.NewGuid();
-
-		const string ResourceLogicalCores = "LogicalCores";
-		ClusterId? _clusterId;
-
-		readonly string? _pool;
-		readonly bool _allowWine;
-		readonly int _maxCores;
-		readonly bool _strict;
-		readonly ConnectionMode? _connectionMode;
-		readonly Encryption? _encryption;
-		readonly CancellationTokenSource _cancellationTokenSource = new();
-		readonly ILogger _logger;
-
-		readonly UnrealBuildAcceleratorConfig _config;
-		internal StatusUpdateAction? _updateStatus;
-		internal IUbaAgentCoordinatorScheduler? _scheduler;
-
-		readonly BundleStorageNamespace _storage = BundleStorageNamespace.CreateInMemory(NullLogger.Instance);
-		BlobLocator _ubaAgentLocator;
-
-		readonly ServiceProvider _serviceProvider;
-		readonly string _crypto;
-		IComputeClient? _client;
-		HordeHttpClient? _hordeHttpClient;
-
-		OSPlatform _targetWorkerPlatform;
-
-		public class Worker
-		{
-			public Stopwatch? StartTime { get; init; }
-			public string? Ip { get; init; }
-			public ConnectionMetadataPort? Port { get; init; }
-			public ConnectionMetadataPort? ProxyPort { get; init; }
-
-			public int NumLogicalCores { get; set; }
-			public bool Active { get; set; }
-			public Task? BackgroundTask { get; set; }
-		}
-
-		readonly List<Worker> _workers = new();
+		public static string ProviderPrefix = "Uba.Provider.Horde";
 
-		public UBAHordeSession(UnrealBuildAcceleratorConfig config, string crypto, Uri? serverUrl, string? accessToken, string? pool, bool allowWine, int maxCores, bool strict, ConnectionMode? connectionMode, Encryption? encryption, ILogger logger)
+		public static UnrealBuildAcceleratorHordeConfig CreateHordeConfig(string provider, CommandLineArguments? additionalArguments = null, DirectoryReference? projectDir = null)
 		{
-			_config = config;
-			_pool = pool;
-			_allowWine = allowWine;
-			_maxCores = maxCores;
-			_strict = strict;
-			_connectionMode = connectionMode;
-			_encryption = encryption;
-			_logger = logger;
-			_crypto = crypto;
+			UnrealBuildAcceleratorHordeConfig config = new();
 
-			ServiceCollection services = new();
-			services.AddLogging(builder => builder.AddEpicDefault());
-			services.AddHorde(options =>
-			{
-				options.ServerUrl = serverUrl;
-				options.AccessToken = accessToken;
-			});
-			_serviceProvider = services.BuildServiceProvider();
+			ConfigHierarchy engineIni = ConfigCache.ReadHierarchy(ConfigHierarchyType.Engine, projectDir, BuildHostPlatform.Current.Platform);
+			config.LoadConfigProvider(engineIni, provider);
 
-			if (connectionMode == ConnectionMode.Relay && String.IsNullOrEmpty(_crypto))
+			if (provider == ProviderPrefix)
 			{
-				_crypto = UBAExecutor.CreateCrypto();
-				_encryption = Encryption.Ssl;
+				XmlConfig.ApplyTo(config);
 			}
-		}
 
-		public async ValueTask DisposeAsync()
-		{
-			if (_workers.Count > 0) // Should handle double-dispose, prevent cancelling twice
+			if (Unreal.IsBuildMachine())
 			{
-				await _cancellationTokenSource.CancelAsync();
-
-				for (int idx = _workers.Count - 1; idx >= 0; idx--)
+				string? hordeUrl = System.Environment.GetEnvironmentVariable("UE_HORDE_URL");
+				if (!String.IsNullOrEmpty(hordeUrl))
 				{
-					await _workers[idx].BackgroundTask!;
-					lock (_workers)
-					{
-						_workers.RemoveAt(idx);
-					}
+					config.HordeServer = hordeUrl;
 				}
 			}
 
-			if (_client != null)
+			if (config.HordeEnabled?.Equals("False", StringComparison.OrdinalIgnoreCase) == true || (config.HordeEnabled?.Equals("BuildMachineOnly", StringComparison.OrdinalIgnoreCase) == true && !Unreal.IsBuildMachine()))
 			{
-				_client = null;
-				// Set CPU resource need to zero (will also expire on the server if not updated)
-				await UpdateCpuCoreNeedAsync(0);
+				config.bDisableHorde = true;
 			}
 
-			_hordeHttpClient?.Dispose();
-
-			await _serviceProvider.DisposeAsync();
-			_cancellationTokenSource.Dispose();
-		}
-
-		public async Task InitAsync(bool useSentry, string? hordeCondition, CancellationToken cancellationToken)
-		{
-			if (_client != null)
-			{
-				throw new InvalidOperationException("Session has already been initialized");
-			}
+			additionalArguments?.ApplyTo(config);
 
-			IHordeClient hordeClient;
-			using (ITimelineEvent _ = Timeline.ScopeEvent($"HordeGetService"))
+			// Sentry is currently unsupported for non-Windows and non-x64
+			if (!OperatingSystem.IsWindows() || RuntimeInformation.ProcessArchitecture != Architecture.X64)
 			{
-				hordeClient = _serviceProvider.GetRequiredService<IHordeClient>();
+				config.UBASentryUrl = null;
 			}
 
-			_client = hordeClient.Compute;
-			_hordeHttpClient = hordeClient.CreateHttpClient();
+			// Normalize URLs
+#pragma warning disable CA1308 // Normalize strings to uppercase - URLs are conventionally lowercased
+			config.HordeServer = config.HordeServer?.TrimEnd('/').ToLowerInvariant();
+			config.UBASentryUrl = config.UBASentryUrl?.TrimEnd('/').ToLowerInvariant();
+#pragma warning restore CA1308
 
-			{
-				using ITimelineEvent _ = Timeline.ScopeEvent($"HordeConnect");
-				GetServerInfoResponse serverInfo = await _hordeHttpClient.GetServerInfoAsync(cancellationToken);
-				_logger.LogInformation("Horde server: {ServerVersion}, agent: {AgentVersion}", serverInfo.ServerVersion, serverInfo.AgentVersion);
-			}
+			return config;
+		}
 
-			_targetWorkerPlatform = GetTargetWorkerPlatform(hordeCondition);
-			bool bCrossPlatformCompilation = _targetWorkerPlatform != GetHostPlatform();
+		public static IEnumerable<IUBAAgentCoordinator> Init(ILogger logger, IEnumerable<TargetDescriptor> targetDescriptors, UnrealBuildAcceleratorConfig ubaConfig, ref string remoteConnectionMode)
+		{
+			IEnumerable<string> providers = Unreal.IsBuildMachine()
+				? [.. ubaConfig.BuildMachineProviders, .. ubaConfig.IniBuildMachineProviders]
+				: [.. ubaConfig.Providers, .. ubaConfig.IniProviders];
+			providers = providers.Distinct().Where(x => !String.IsNullOrEmpty(x) && x.StartsWith(ProviderPrefix, StringComparison.Ordinal));
 
-			DirectoryReference ubaDir;
-			if (bCrossPlatformCompilation && _targetWorkerPlatform == OSPlatform.Linux)
-			{
-				ubaDir = DirectoryReference.Combine(Unreal.EngineDirectory, "Binaries", "Linux", "UnrealBuildAccelerator");
-			}
-			else
+			if (!providers.Any())
 			{
-				ubaDir = UBAExecutor.UbaBinariesDir;
+				providers = [ProviderPrefix];
 			}
 
-			List<string> agentFiles = new();
-			if (bCrossPlatformCompilation && _targetWorkerPlatform == OSPlatform.Linux)
-			{
-				agentFiles.Add("UbaAgent");
-			}
-			else if (OperatingSystem.IsWindows())
-			{
-				agentFiles.Add("UbaAgent.exe");
-				if (FileReference.Exists(FileReference.Combine(ubaDir, "UbaWine.dll.so")))
-				{
-					agentFiles.Add("UbaWine.dll.so"); // Not needed but used by linux/wine helpers for more detailed logging or optimizations
-				}
-				if (useSentry)
-				{
-					agentFiles.Add("crashpad_handler.exe");
-					agentFiles.Add("sentry.dll");
-				}
-			}
-			else if (OperatingSystem.IsLinux())
-			{
-				agentFiles.Add("UbaAgent");
-				bool bIsDebug = false;
-				if (bIsDebug)
-				{
-					agentFiles.Add("UbaAgent.debug");
-					agentFiles.Add("UbaAgent.sym");
-					agentFiles.Add("libclang_rt.tsan.so"); // Needs to be copied from autosdk
-					agentFiles.Add("llvm-symbolizer"); // Needs to be copied from autosdk
-				}
-			}
-			else if (OperatingSystem.IsMacOS())
-			{
-				agentFiles.Add("UbaAgent");
-			}
-			else
-			{
-				throw new PlatformNotSupportedException();
-			}
+			// We group configs based on hordeserver/token to be able to have only one connection for configs having matching server/token.
 
-			_ubaAgentLocator = await CreateToolAsync(ubaDir, agentFiles.Select(x => FileReference.Combine(ubaDir, x)), cancellationToken);
-			_logger.LogInformation("Created tool bundle with locator {UbaAgentLocator}", _ubaAgentLocator.ToString());
-		}
+			IEnumerable<IGrouping<string, UnrealBuildAcceleratorHordeConfig>> configGroups = providers
+				.SelectMany(provider => targetDescriptors.Select(targetDescriptor => CreateHordeConfig(provider, targetDescriptor.AdditionalArguments, targetDescriptor.ProjectFile?.Directory)))
+				.Where(x => !x.bDisableHorde && !String.IsNullOrWhiteSpace(x.HordeServer))
+				.DistinctBy(x => HashCode.Combine(x.HordeServer, x.HordeCondition, x.HordeCluster, x.HordePool, x.HordeConnectionMode, x.AgentPlatform))
+				.GroupBy(x => String.Join(";", x.HordeServer, x.HordeToken));
 
-		public UnrealBuildAcceleratorCacheConfig? RequestCacheServer(int desiredConnectionCount, CancellationToken cancellationToken)
-		{
-			// REALLY UGLY. Remove this loop. We know InitAsync has been called before this one so right now it is a hack to just loop.
-			while (_client == null)
+			List<UBAAgentCoordinatorHorde> coordinators = new();
+			foreach (IGrouping<string, UnrealBuildAcceleratorHordeConfig> configGroup in configGroups)
 			{
-				if (cancellationToken.IsCancellationRequested)
-				{
-					return null;
-				}
-				Thread.Sleep(10);
+				coordinators.Add(new UBAAgentCoordinatorHorde(logger, ubaConfig, configGroup));
 			}
 
-			ClusterId clusterId = new(UnrealBuildAcceleratorHordeConfig.ClusterDefault);
-			using ITimelineEvent timelineEvent = Timeline.ScopeEvent($"HordeGetCache");
-			UbaConfig ubaConfig = _client.AllocateUbaCacheServerAsync(clusterId, cancellationToken).Result;
-			timelineEvent.Finish();
-			if (String.IsNullOrEmpty(ubaConfig.CacheEndpoint))
-			{
-				return null;
-			}
-			bool writeAccess = ubaConfig.WriteAccess;
-			return new UnrealBuildAcceleratorCacheConfig() { CacheServer = ubaConfig.CacheEndpoint, bRequireVfs = true, WriteCache = writeAccess.ToString(), Crypto = ubaConfig.CacheSessionKey, DesiredConnectionCount = desiredConnectionCount };
+			remoteConnectionMode = (Enum.TryParse(configGroups.First().First().HordeConnectionMode, true, out ConnectionMode cm) ? cm : ConnectionMode.Direct).ToString();
+			return coordinators;
 		}
 
-		async Task<BlobLocator> CreateToolAsync(DirectoryReference baseDir, IEnumerable<FileReference> files, CancellationToken cancellationToken)
+		public UBAAgentCoordinatorHorde(ILogger logger, UnrealBuildAcceleratorConfig ubaConfig, IEnumerable<UnrealBuildAcceleratorHordeConfig> hordeConfigs)
 		{
-			using ITimelineEvent __ = Timeline.ScopeEvent($"HordeCreateTool");
-
-			// TODO: should drive this off API version reported by server
-			BlobSerializerOptions serializerOptions = BlobSerializerOptions.Create(HordeApiVersion.Initial);
-
-			await using IBlobWriter writer = _storage.CreateBlobWriter(serializerOptions: serializerOptions, cancellationToken: cancellationToken);
-			DirectoryNode sandbox = new();
-			await sandbox.AddFilesAsync(baseDir, files, writer, cancellationToken: cancellationToken);
-			IHashedBlobRef<DirectoryNode> handle = await writer.WriteBlobAsync(sandbox, cancellationToken);
-			await writer.FlushAsync(cancellationToken);
-			return handle.GetLocator();
+			_logger = logger;
+			_title = "Horde";
+			_ubaConfig = ubaConfig;
+			_hordeConfigs = hordeConfigs.ToArray();
+			_timers = new Timer?[_hordeConfigs.Length];
 		}
 
-		public async void RemoveCompleteWorkersAsync()
+		public DirectoryReference? GetUBARootDir()
 		{
-			for (int idx = 0; idx < _workers.Count; idx++)
+			DirectoryReference? hordeSharedDir = DirectoryReference.FromString(Environment.GetEnvironmentVariable("UE_HORDE_SHARED_DIR"));
+			if (hordeSharedDir != null)
 			{
-				Worker worker = _workers[idx];
-				if (worker.BackgroundTask!.IsCompleted)
-				{
-					await worker.BackgroundTask;
-					lock (_workers)
-					{
-						_workers.RemoveAt(idx--);
-					}
-				}
+				return DirectoryReference.Combine(hordeSharedDir, "UbaHost");
 			}
+			return null;
 		}
 
-		public void GetCoreCount(out int active, out int queued, out int activeAgents, out int queuedAgents)
+		public async Task InitAsync(UBAExecutor executor, StatusUpdateAction updateStatus, CancellationToken cancellationToken)
 		{
-			active = 0;
-			queued = 0;
-			activeAgents = 0;
-			queuedAgents = 0;
-			lock (_workers)
+			_updateStatus = updateStatus;
+			_cancellationToken = cancellationToken;
+
+			try
 			{
-				foreach (Worker worker in _workers)
+				if (_ubaConfig.bDisableRemote)
 				{
-					if (worker.NumLogicalCores == 0)
-					{
-						continue;
-					}
-					if (worker.Active)
-					{
-						++activeAgents;
-						active += worker.NumLogicalCores;
-					}
-					else
-					{
-						++queuedAgents;
-						queued += worker.NumLogicalCores;
-					}
+					return;
 				}
-			}
-		}
-
-		int _workerId = 0;
-
-		static readonly HashSet<StringView> s_logProperties = new()
-			{
-				"ComputeIp",
-				"CPU",
-				"RAM",
-				"DiskFreeSpace",
-				"PhysicalCores",
-				"LogicalCores",
-				"EC2",
-				"LeaseId",
-				"aws-instance-type",
-			};
-
-		public async Task<bool> AddWorkerAsync(Requirements requirements, UnrealBuildAcceleratorHordeConfig hordeConfig, int activeCores, CancellationToken cancellationToken)
-		{
-			if (_client == null)
-			{
-				throw new InvalidOperationException($"Session not initialized. Call {nameof(InitAsync)} first");
-			}
 
-			int workerId = _workerId;
-			
-			PrefixLogger workerLogger = new($"[Worker{workerId}]", _logger);
+				// All configs have the same server/token
+				UnrealBuildAcceleratorHordeConfig firstConfig = _hordeConfigs.First()!;
+				Uri? serverUrl = (firstConfig.HordeServer == null) ? null : new Uri(firstConfig.HordeServer);
+				string? accessToken = firstConfig.HordeToken;
 
-			const string UbaPortName = "UbaPort";
-			const string UbaProxyPortName = "UbaProxyPort";
-			const int UbaPort = 7001;
-			const int UbaProxyPort = 7002;
-
-			// Request ID that is unique per attempt to acquire the same compute lease/worker
-			// Primarily for tracking worker demand on Horde server as UBAExecutor will repeatedly try adding a new worker
-			string requestId = $"{_id}-worker-{workerId}";
-			IComputeLease? lease = null;
-			try
-			{
-				Stopwatch stopwatch = Stopwatch.StartNew();
-				ConnectionMetadataRequest cmr = new()
+				ServiceCollection services = new();
+				services.AddLogging(builder => builder.AddEpicDefault());
+				services.AddHorde(options =>
 				{
-					ModePreference = _connectionMode,
-					Encryption = _encryption,
-					Ports = { { UbaPortName, UbaPort } },
-					InactivityTimeoutMs = 90000
-				};
-				
-				await ResolveClusterIdAsync(hordeConfig, requirements, requestId, cmr, workerLogger, cancellationToken);
-				lease = await _client.TryAssignWorkerAsync(_clusterId, requirements, requestId, cmr, false, workerLogger, cancellationToken);
-				if (lease == null)
-				{
-					_logger.LogDebug("Unable to assign a remote worker");
+					options.ServerUrl = serverUrl;
+					options.AccessToken = accessToken;
+				});
+				_serviceProvider = services.BuildServiceProvider();
 
-					int missingNumCores = Math.Max(0, _maxCores - activeCores);
-					await UpdateCpuCoreNeedAsync(missingNumCores, cancellationToken);
-					return false;
+				IHordeClient hordeClient;
+				using (ITimelineEvent _ = Timeline.ScopeEvent($"HordeGetService"))
+				{
+					hordeClient = _serviceProvider.GetRequiredService<IHordeClient>();
 				}
 
-				++_workerId;
+				IComputeClient computeClient = hordeClient.Compute;
+				using (HordeHttpClient httpClient = hordeClient.CreateHttpClient())
+				{
+					using ITimelineEvent _ = Timeline.ScopeEvent($"HordeConnect");
+					GetServerInfoResponse serverInfo = await httpClient.GetServerInfoAsync(cancellationToken);
+					_logger.LogInformation("Horde server: {ServerVersion}, agent: {AgentVersion}", serverInfo.ServerVersion, serverInfo.AgentVersion);
+				}
 
-				workerLogger.LogDebug("Agent properties:");
+				List<Task<(UBAHordeSession?, string)>> sessionTasks = [];
 
-				string agentName = String.Empty;
-				int numLogicalCores = 24; // Assume 24 if something goes wrong here and property is not found
-				string computeIp = String.Empty;
-				string leaseId = String.Empty;
-				string instanceType = String.Empty;
-				foreach (string property in lease.Properties)
+				for (uint i = 0; i != _hordeConfigs.Length; ++i)
 				{
-					int equalsIdx = property.IndexOf('=', StringComparison.OrdinalIgnoreCase);
-					StringView propertyName = new(property, 0, equalsIdx);
-					if (s_logProperties.Contains(propertyName))
+					UnrealBuildAcceleratorHordeConfig hordeConfig = _hordeConfigs[i];
+					string? link = null;
+					if (hordeConfig.HordeServer != null)
 					{
-						_logger.LogDebug("  {Property}", property);
-
-						if (propertyName == ResourceLogicalCores && Int32.TryParse(property.AsSpan(equalsIdx + 1), out int value))
-						{
-							numLogicalCores = value;
-						}
-						else if (propertyName == "ComputeIp")
-						{
-							computeIp = property[(equalsIdx + 1)..];
-						}
-						else if (propertyName == "EC2")
-						{
-							agentName = $"Worker{workerId}";
-						}
-						else if (propertyName == "LeaseId")
-						{
-							leaseId = property[(equalsIdx + 1)..];
-						}
-						else if (propertyName == "aws-instance-type")
+						link = $"{hordeConfig.HordeServer!.TrimEnd('/')}/agents";
+						if (!String.IsNullOrEmpty(hordeConfig.HordePool))
 						{
-							instanceType = property[(equalsIdx + 1)..];
+							link += $"?agent={hordeConfig.HordePool}";
 						}
 					}
-				}
-
-				string desc = instanceType;
-				if (!String.IsNullOrEmpty(leaseId))
-				{
-					string link = $" {hordeConfig.HordeServer!.TrimEnd('/')}/lease/{leaseId}";
-					desc += link;
-				}
-				// When using relay connection mode, the IP will be relay server's IP
-				string ip = String.IsNullOrEmpty(lease.Ip) ? computeIp : lease.Ip;
+					_updateStatus(i, 1, _title, EpicGames.UBA.LogEntryType.Info, link);
+					_updateStatus!(i, 6, "Connecting", EpicGames.UBA.LogEntryType.Info, null);
 
-				if (!lease.Ports.TryGetValue(UbaPortName, out ConnectionMetadataPort? ubaPort))
-				{
-					ubaPort = new ConnectionMetadataPort(UbaPort, UbaPort);
+					sessionTasks.Add(UBAHordeSession.TryCreateHordeSessionAsync(_ubaConfig, hordeConfig, computeClient, executor.Crypto, _ubaConfig.bStrict, _logger, cancellationToken));
 				}
+				_hordeSessionTasks = sessionTasks.ToArray();
 
-				if (!lease.Ports.TryGetValue(UbaProxyPortName, out ConnectionMetadataPort? ubaProxyPort))
+				for (uint i = 0; i != _hordeSessionTasks.Length; ++i)
 				{
-					ubaProxyPort = new ConnectionMetadataPort(UbaProxyPort, UbaProxyPort);
+					(UBAHordeSession? session, string status) = await _hordeSessionTasks[i];
+					_updateStatus!(i, 6, status, EpicGames.UBA.LogEntryType.Info, null);
+					executor.AgentCoordinatorInitialized(_hordeSessionTasks[i].IsCompletedSuccessfully && session is not null);
 				}
-
-				string exeName = OperatingSystem.IsWindows() ? "UbaAgent.exe" : "UbaAgent";
-				BlobLocator locator = _ubaAgentLocator;
-				Worker worker = new()
-				{
-					StartTime = stopwatch,
-					NumLogicalCores = numLogicalCores,
-					Ip = ip,
-					Port = ubaPort,
-					ProxyPort = ubaProxyPort,
-				};
-				worker.BackgroundTask = RunWorkerAsync(worker, lease, locator, exeName, workerLogger, hordeConfig, agentName, desc, _cancellationTokenSource.Token);
-				lock (_workers)
-				{
-					_workers.Add(worker);
-				}
-				UpdateHordeStatus(null);
-				lease = null; // Will be disposed by RunWorkerAsync
-
-				return true;
 			}
 			finally
 			{
-				if (lease != null)
-				{
-					await lease.DisposeAsync();
-				}
+				_initDoneEvent.Set();
 			}
-		}
-
-		int _targetCoreCount;
-
-		async Task UpdateCpuCoreNeedAsync(int targetCoreCount, CancellationToken cancellationToken = default)
-		{
-			if (_client != null && _pool != null && _clusterId != null && targetCoreCount != _targetCoreCount)
+			lock (_startLock)
 			{
-				_targetCoreCount = targetCoreCount;
-
-				_logger.LogDebug("Setting CPU core need to {TargetCoreCount}", targetCoreCount);
-				Dictionary<string, int> resourceNeeds = new() { { ResourceLogicalCores, targetCoreCount } };
-				try
+				_initDone = true;
+				if (_scheduler != null)
 				{
-					await _client.DeclareResourceNeedsAsync(_clusterId.Value, _pool, resourceNeeds, cancellationToken);
-				}
-				catch (Exception e)
-				{
-					_logger.Log(_strict ? LogLevel.Error : LogLevel.Debug, KnownLogEvents.Systemic_Horde_Compute, e, "Failed updating resource need to {TargetCoreCount} cores", targetCoreCount);
+					StartTimer();
 				}
 			}
 		}
-		
-		private async Task ResolveClusterIdAsync(UnrealBuildAcceleratorHordeConfig config, Requirements requirements, string requestId, ConnectionMetadataRequest cmr, ILogger logger, CancellationToken cancellationToken = default)
+
+		public UnrealBuildAcceleratorCacheConfig? RequestCacheServer(CancellationToken cancellationToken)
 		{
-			if (_client == null)
+			for (uint i = 0; i != _hordeConfigs.Length; ++i)
 			{
-				throw new InvalidOperationException($"Session not initialized. Call {nameof(InitAsync)} first");
-			}
-			
-			if (_clusterId == null)
-			{
-				if (config.HordeCluster == UnrealBuildAcceleratorHordeConfig.ClusterAuto)
-				{
-					_clusterId = await _client.GetClusterAsync(requirements, requestId, cmr, logger, cancellationToken);
-				}
-				else if (!String.IsNullOrEmpty(config.HordeCluster))
-				{
-					_clusterId = new ClusterId(config.HordeCluster);
-				}
-				else
+				UnrealBuildAcceleratorHordeConfig hordeConfig = _hordeConfigs[i];
+				if (!hordeConfig.bCacheEnabled)
 				{
-					_clusterId = new ClusterId(UnrealBuildAcceleratorHordeConfig.ClusterDefault);
+					continue;
 				}
-				_logger.LogInformation("Horde cluster resolved as '{ClusterId}'", _clusterId.ToString());
-			}
-		}
-
-		public static async Task<(UBAHordeSession?, string)> TryCreateHordeSessionAsync(UnrealBuildAcceleratorHordeConfig hordeConfig, UnrealBuildAcceleratorConfig ubaConfig, string crypto, bool bStrictErrors, ILogger logger, CancellationToken cancellationToken = default)
-		{
-			using ITimelineEvent hordeCreateTimlineEvent = Timeline.ScopeEvent($"HordeCreate");
 
-			if (hordeConfig.bDisableHorde)
-			{
-				logger.LogInformation("Horde disabled via command line option.");
-				return (null, "Disabled via command line");
-			}
-
-			if (String.IsNullOrEmpty(hordeConfig.HordeServer) && HordeOptions.GetServerUrlFromEnvironment() == null && HordeOptions.GetDefaultServerUrl() == null)
-			{
-				logger.LogInformation("Horde disabled. Url not set.");
-				return (null, "No server url set");
-			}
-
-			Uri? server = (hordeConfig.HordeServer == null) ? null : new Uri(hordeConfig.HordeServer);
-			string? token = hordeConfig.HordeToken;
-
-			ConnectionMode? connectionMode = Enum.TryParse(hordeConfig.HordeConnectionMode, true, out ConnectionMode cm) ? cm : null;
-			Encryption? encryption = Enum.TryParse(hordeConfig.HordeEncryption, true, out Encryption enc) ? enc : null;
-
-			// Default to SSL encryption for relay mode if unset
-			if (connectionMode == ConnectionMode.Relay && encryption == null)
-			{
-				encryption = Encryption.Ssl;
-			}
-
-			logger.LogInformation("Horde URL: {Server}, Pool: {Pool}, Cluster: {Cluster}, Condition: {Condition}, Connection: {Connection}, Encryption: {Encryption}, MaxCores: {MaxCores}, MaxWorkers: {MaxWorkers}, MaxIdle: {MaxIdle}s",
-				server, hordeConfig.HordePool ?? "(none)", hordeConfig.HordeCluster ?? "(none)", hordeConfig.HordeCondition ?? "(none)", connectionMode?.ToString() ?? "(none)", encryption?.ToString() ?? "(none)",
-				hordeConfig.HordeMaxCores, hordeConfig.HordeMaxWorkers, hordeConfig.HordeMaxIdle);
-			try
-			{
-				bool allowWine = hordeConfig.bHordeAllowWine && OperatingSystem.IsWindows();
-
-				UBAHordeSession session = new(ubaConfig, crypto, server, token, hordeConfig.HordePool, allowWine, hordeConfig.HordeMaxCores, bStrictErrors, connectionMode, encryption, logger);
-				hordeCreateTimlineEvent.Dispose();
-				await session.InitAsync(useSentry: !String.IsNullOrEmpty(hordeConfig.UBASentryUrl), hordeConfig.HordeCondition, cancellationToken);
-				return (session, "Ready");
-			}
-			catch (TaskCanceledException)
-			{
-				return (null, "Cancelled");
-			}
-			catch (Exception ex)
-			{
-				logger.Log(bStrictErrors ? LogLevel.Error : LogLevel.Information, ex, "Unable to create Horde session: {Message}", ex.Message);
-				return (null, ex.Message);
-			}
-		}
-
-		async Task RunWorkerAsync(Worker self, IComputeLease lease, BlobLocator tool, string executable, ILogger logger, UnrealBuildAcceleratorHordeConfig hordeConfig, string agentName, string desc, CancellationToken cancellationToken)
-		{
-			logger.LogDebug("Running worker task..");
-			try
-			{
-				await using (_ = lease)
+				_initDoneEvent.Wait(cancellationToken);
+				if (!_initDoneEvent.IsSet)
 				{
-					// Create a message channel on channel id 0. The Horde Agent always listens on this channel for requests.
-					const int PrimaryChannelId = 0;
-					using (AgentMessageChannel channel = lease.Socket.CreateAgentMessageChannel(PrimaryChannelId, 4 * 1024 * 1024))
-					{
-						logger.LogDebug("Waiting for attach...");
-
-						TimeSpan attachTimeout = TimeSpan.FromSeconds(20.0);
-						try
-						{
-							Task attachTask = channel.WaitForAttachAsync(cancellationToken).AsTask();
-							await attachTask.WaitAsync(attachTimeout, cancellationToken);
-						}
-						catch (TimeoutException)
-						{
-							logger.Log(_strict ? LogLevel.Error : LogLevel.Debug, KnownLogEvents.Systemic_Horde_Compute, "Waited {Time}s on attach message. Giving up", (int)attachTimeout.TotalSeconds);
-							throw;
-						}
-
-						logger.LogDebug("Uploading files...");
-						await channel.UploadFilesAsync("", tool, _storage.Backend, cancellationToken);
-
-						string hordeHost = _config.Host;
-						if (!String.IsNullOrEmpty(hordeConfig.HordeHost))
-						{
-							hordeHost = hordeConfig.HordeHost;
-						}
-
-						bool useListen = !String.IsNullOrEmpty(hordeConfig.HordeHost);
-						List<string> arguments = new();
-
-						if (!String.IsNullOrEmpty(agentName))
-						{
-							arguments.Add($"-name={agentName}");
-						}
-
-						if (useListen)
-						{
-							arguments.Add($"-Host={hordeHost}:{_config.Port}");
-						}
-						else
-						{
-							arguments.Add($"-Listen={self.Port!.AgentPort}");
-							arguments.Add("-ListenTimeout=10");
-						}
-
-						if (!String.IsNullOrEmpty(_crypto))
-						{
-							arguments.Add($"-crypto={_crypto}");
-						}
-
-						arguments.Add("-NoPoll");
-						arguments.Add("-Quiet");
-						if (!String.IsNullOrEmpty(hordeConfig.UBASentryUrl))
-						{
-							arguments.Add($"-Sentry=\"{hordeConfig.UBASentryUrl}\"");
-						}
-						arguments.Add("-ProxyPort=" + self.ProxyPort!.AgentPort);
-						if (_config.bUseQuic)
-						{
-							arguments.Add("-quic");
-						}
-						//arguments.Add("-NoStore");
-						//arguments.Add("-KillRandom"); // For debugging
-
-						if (OperatingSystem.IsMacOS() && _targetWorkerPlatform == OSPlatform.OSX)
-						{
-							// we need to populate the cas with all known xcodes so we can serve all the ones we have installed
-							string xcodeVersion = Utils.RunLocalProcessAndReturnStdOut("/bin/sh", "-c '/usr/bin/defaults read $(xcode-select -p)/../version.plist ProductBuildVersion");
-							arguments.Add($"-populateCasFromXcodeVersion={xcodeVersion}");
-							arguments.Add("-killtcphogs");
-						}
-
-						arguments.Add("-Dir=%UE_HORDE_SHARED_DIR%\\Uba");
-						arguments.Add("-Eventfile=%UE_HORDE_TERMINATION_SIGNAL_FILE%");
-						arguments.Add("-MaxCpu=%UE_HORDE_CPU_COUNT%");
-						arguments.Add("-MulCpu=%UE_HORDE_CPU_MULTIPLIER%");
-						arguments.Add("-MaxCon=8");
-						arguments.Add($"-MaxWorkers={hordeConfig.HordeMaxWorkers}");
-						arguments.Add($"-MaxIdle={hordeConfig.HordeMaxIdle}");
-						arguments.Add("-UseCrawler");
-						// arguments.Add("-ProxyUseLocalStorage"); // If set then agent's cas storage will be used by proxy. [honk] This is disabled because we suspect a hang in this feature
-						//arguments.Add("-UseIocp=4");
-						//arguments.Add("-NoStore"); // This means that no cas files will be written to disk.. so no state is stored between runs
-
-						if (!String.IsNullOrEmpty(desc))
-						{
-							arguments.Add($"-Description=\"{desc}\"");
-						}
-
-						if (_config.bLogEnabled)
-						{
-							arguments.Add("-Log");
-						}
-
-						LogLevel logLevel = _config.bDetailedLog ? LogLevel.Information : LogLevel.Debug;
-
-						logger.Log(logLevel, "Executing child process: {Executable} {Arguments}", executable, CommandLineArguments.Join(arguments));
-
-						ExecuteProcessFlags execFlags = _allowWine ? ExecuteProcessFlags.UseWine : ExecuteProcessFlags.None;
-						await using AgentManagedProcess process = await channel.ExecuteAsync(executable, arguments, null, null, execFlags, cancellationToken);
-						bool shouldConnect = !useListen;
-						bool isFirstRead = true;
-						string? line;
-
-						while ((line = await process.ReadLineAsync(cancellationToken)) != null)
-						{
-							logger.Log(logLevel, "{Line}", line);
-
-							if (shouldConnect && line.Contains("Listening on", StringComparison.OrdinalIgnoreCase)) // This log entry means that the agent is ready for connections.
-							{
-								long totalMs = self.StartTime!.ElapsedMilliseconds;
-								logger.LogInformation("Connecting to UbaAgent on {Ip}:{Port} (local agent port {AgentPort}) {Seconds}.{Milliseconds} seconds after assigned",
-									self.Ip, self.Port!.Port, self.Port.AgentPort, totalMs / 1000, totalMs % 1000);
-
-								if (!_scheduler!.AddClient(self.Ip!, self.Port.Port, _crypto))
-								{
-									break;
-								}
-								shouldConnect = false;
-							}
-
-							if (isFirstRead)
-							{
-								isFirstRead = false;
-								self.Active = true;
-								UpdateHordeStatus(null);
-							}
-						}
-						self.NumLogicalCores = 0;
-						UpdateHordeStatus(null);
-						logger.LogDebug("Shutting down process");
-						int ExitCode = await process.WaitForExitAsync(cancellationToken);
-						if (ExitCode != 0)
-						{
-							logger.LogInformation("UbaAgent exited with exit code: {ExitCode}", ExitCode);
-						}
-					}
-
-					logger.LogDebug("Closing channel");
-					await lease.CloseAsync(cancellationToken);
+					return null;
 				}
-			}
-			catch (TimeoutException)
-			{
-				self.NumLogicalCores = 0;
-				UpdateHordeStatus(null);
-			}
-			catch (ComputeExecutionCancelledException ex) 
-			{
-				// Cancellations are expected to happen due to spot instance interruptions or unscheduled maintenance of agents
-				self.NumLogicalCores = 0;
-				UpdateHordeStatus(null);
-				
-				if (!cancellationToken.IsCancellationRequested)
+				if (_hordeSessionTasks.Length < i)
 				{
-					logger.Log(_strict ? LogLevel.Information : LogLevel.Debug, KnownLogEvents.Systemic_Horde_Compute, ex, "Compute lease cancelled");
+					return null;
 				}
-			}
-			catch (Exception ex)
-			{
-				self.NumLogicalCores = 0;
-				UpdateHordeStatus(null);
 
-				if (!cancellationToken.IsCancellationRequested)
+				(UBAHordeSession? hordeSession, string reason) = _hordeSessionTasks[i].Result;
+				_updateStatus!(0, 6, reason, EpicGames.UBA.LogEntryType.Info, null);
+				if (hordeSession == null)
 				{
-					logger.Log(_strict ? LogLevel.Error : LogLevel.Debug, KnownLogEvents.Systemic_Horde_Compute, ex, "Exception in worker task: {Ex}", ex.ToString());
-
-					// Add additional properties to aid debugging
-					logger.Log(_strict ? LogLevel.Information : LogLevel.Debug, KnownLogEvents.Systemic_Horde_Compute, ex, "UBA agent locator {UBAAgentLocator}", _ubaAgentLocator.ToString());
+					continue;
 				}
-			}
-		}
-
-		string _additionalStatus = "";
-		string _lastStatus = "";
 
-		internal void UpdateHordeStatus(string? additionalStatus)
-		{
-			if (additionalStatus != null)
-			{
-				_additionalStatus = additionalStatus;
-			}
-
-			if (_updateStatus == null)
-			{
-				return;
-			}
-
-			int activeCores;
-			int queuedCores;
-			int activeAgents;
-			int queuedAgents;
-			GetCoreCount(out activeCores, out queuedCores, out activeAgents, out queuedAgents);
-			string status = $"Running. {activeAgents} agent{(activeAgents != 1 ? "s" : "")} ({activeCores} cores){_additionalStatus} {(queuedAgents != 0 ? $"(Preparing {queuedAgents} agent{(queuedAgents != 1 ? "s" : "")})" : "")}";
-			if (_lastStatus != status)
-			{
-				_lastStatus = status;
-				_updateStatus(0, 6, status, EpicGames.UBA.LogEntryType.Info, null);
+				return hordeSession.RequestCacheServer(hordeConfig.CacheDesiredConnectionCount, cancellationToken);
 			}
+			return null;
 		}
 
-		static OSPlatform GetTargetWorkerPlatform(string? HordeCondition)
+		public void Start(IUbaAgentCoordinatorScheduler scheduler, Func<LinkedAction, bool> canRunRemotely)
 		{
-			if (!String.IsNullOrEmpty(HordeCondition))
+			lock (_startLock)
 			{
-				if (HordeCondition.Contains("Linux", StringComparison.OrdinalIgnoreCase))
-				{
-					return OSPlatform.Linux;
-				}
-				if (HordeCondition.Contains("Windows", StringComparison.OrdinalIgnoreCase))
+				_scheduler = scheduler;
+				if (_initDone)
 				{
-					return OSPlatform.Windows;
+					StartTimer();
 				}
-				if (HordeCondition.Contains("MacOS", StringComparison.OrdinalIgnoreCase))
-				{
-					return OSPlatform.OSX;
-				}
-			}
-
-			return GetHostPlatform();
-		}
-
-		static OSPlatform GetHostPlatform()
-		{
-			if (OperatingSystem.IsLinux())
-			{
-				return OSPlatform.Linux;
-			}
-			if (OperatingSystem.IsWindows())
-			{
-				return OSPlatform.Windows;
-			}
-			return OSPlatform.OSX;
-		}
-	}
-
-	class UBAAgentCoordinatorHorde : IUBAAgentCoordinator, IDisposable
-	{
-		public string ConnectionModeString => (Enum.TryParse(_hordeConfig.HordeConnectionMode, true, out ConnectionMode cm) ? cm : ConnectionMode.Direct).ToString();
-
-		public static string ProviderPrefix = "Uba.Provider.Horde";
-
-		public static IEnumerable<IUBAAgentCoordinator> Init(ILogger logger, IEnumerable<TargetDescriptor> targetDescriptors, UnrealBuildAcceleratorConfig ubaConfig, ref string remoteConnectionMode)
-		{
-			IEnumerable<string> providers = Unreal.IsBuildMachine()
-				? [.. ubaConfig.BuildMachineProviders, .. ubaConfig.IniBuildMachineProviders]
-				: [.. ubaConfig.Providers, .. ubaConfig.IniProviders];
-			providers = providers.Distinct().Where(x => !String.IsNullOrEmpty(x) && x.StartsWith(ProviderPrefix, StringComparison.Ordinal));
-
-			if (!providers.Any())
-			{
-				providers = [ProviderPrefix];
 			}
-
-			IEnumerable<UBAAgentCoordinatorHorde> coordinators = providers
-				.SelectMany(provider => targetDescriptors.Select(targetDescriptor => new UBAAgentCoordinatorHorde(logger, ubaConfig, provider, targetDescriptor.AdditionalArguments, targetDescriptor.ProjectFile?.Directory)))
-				.Where(x => !x._hordeConfig.bDisableHorde && !String.IsNullOrWhiteSpace(x._hordeConfig.HordeServer))
-				.DistinctBy(x => HashCode.Combine(x._hordeConfig.HordeServer, x._hordeConfig.HordeCondition, x._hordeConfig.HordeCluster, x._hordeConfig.HordePool, x._hordeConfig.HordeConnectionMode));
-
-			remoteConnectionMode = coordinators.FirstOrDefault()?.ConnectionModeString ?? "Local";
-			return coordinators;
 		}
-
-		public UBAAgentCoordinatorHorde(ILogger logger, UnrealBuildAcceleratorConfig ubaConfig, string provider, CommandLineArguments? additionalArguments = null, DirectoryReference? projectDir = null)
+		private void StartTimer()
 		{
-			_logger = logger;
-			_title = "Horde";
-			_ubaConfig = ubaConfig;
-
-			_hordeConfig = new();
-
-			ConfigHierarchy engineIni = ConfigCache.ReadHierarchy(ConfigHierarchyType.Engine, projectDir, BuildHostPlatform.Current.Platform);
-			_hordeConfig.LoadConfigProvider(engineIni, provider);
-
-			if (provider == ProviderPrefix)
+			for (uint i2 = 0; i2 != _hordeConfigs.Length; ++i2)
 			{
-				XmlConfig.ApplyTo(_hordeConfig);
-			}
+				uint i = i2;
+				UnrealBuildAcceleratorHordeConfig hordeConfig = _hordeConfigs[i];
 
-			if (Unreal.IsBuildMachine())
-			{
-				string? hordeUrl = System.Environment.GetEnvironmentVariable("UE_HORDE_URL");
-				if (!String.IsNullOrEmpty(hordeUrl))
+				if (!hordeConfig.bAgentsEnabled)
 				{
-					_hordeConfig.HordeServer = hordeUrl;
+					continue;
 				}
-			}
 
-			if (_hordeConfig.HordeEnabled?.Equals("False", StringComparison.OrdinalIgnoreCase) == true || (_hordeConfig.HordeEnabled?.Equals("BuildMachineOnly", StringComparison.OrdinalIgnoreCase) == true && !Unreal.IsBuildMachine()))
-			{
-				_hordeConfig.bDisableHorde = true;
-			}
+				int timerPeriod = 5000;
+				bool shownNoAgentsFoundMessage = false;
 
-			additionalArguments?.ApplyTo(_hordeConfig);
+				(UBAHordeSession? hordeSession, string reason) = _hordeSessionTasks[i].Result;
 
-			// Sentry is currently unsupported for non-Windows and non-x64
-			if (!OperatingSystem.IsWindows() || RuntimeInformation.ProcessArchitecture != Architecture.X64)
-			{
-				_hordeConfig.UBASentryUrl = null;
-			}
-
-			// Normalize URLs
-#pragma warning disable CA1308 // Normalize strings to uppercase - URLs are conventionally lowercased
-			_hordeConfig.HordeServer = _hordeConfig.HordeServer?.TrimEnd('/').ToLowerInvariant();
-			_hordeConfig.UBASentryUrl = _hordeConfig.UBASentryUrl?.TrimEnd('/').ToLowerInvariant();
-#pragma warning restore CA1308
-		}
-
-		public DirectoryReference? GetUBARootDir()
-		{
-			DirectoryReference? hordeSharedDir = DirectoryReference.FromString(Environment.GetEnvironmentVariable("UE_HORDE_SHARED_DIR"));
-			if (hordeSharedDir != null)
-			{
-				return DirectoryReference.Combine(hordeSharedDir, "UbaHost");
-			}
-			return null;
-		}
-
-		public async Task InitAsync(UBAExecutor executor, StatusUpdateAction updateStatus, CancellationToken cancellationToken)
-		{
-			_updateStatus = updateStatus;
-			_cancellationToken = cancellationToken;
-
-			if (_ubaConfig.bDisableRemote)
-			{
-				return;
-			}
-
-			string? link = null;
-			if (_hordeConfig.HordeServer != null)
-			{
-				link = $"{_hordeConfig.HordeServer!.TrimEnd('/')}/agents";
-				if (!String.IsNullOrEmpty(_hordeConfig.HordePool))
+				if (hordeSession == null)
 				{
-					link += $"?agent={_hordeConfig.HordePool}";
+					if (_updateStatus != null)
+					{
+						_updateStatus(0, 6, reason, EpicGames.UBA.LogEntryType.Info, null);
+					}
+					continue;
 				}
-			}
-			_updateStatus(0, 1, _title, EpicGames.UBA.LogEntryType.Info, link);
-			_updateStatus!(0, 6, "Connecting", EpicGames.UBA.LogEntryType.Info, null);
-
-			_hordeSessionTask = UBAHordeSession.TryCreateHordeSessionAsync(_hordeConfig, _ubaConfig, executor.Crypto, _ubaConfig.bStrict, _logger, cancellationToken);
-			(UBAHordeSession? session, string status) = await _hordeSessionTask;
-			_updateStatus!(0, 6, status, EpicGames.UBA.LogEntryType.Info, null);
-			executor.AgentCoordinatorInitialized(_hordeSessionTask.IsCompletedSuccessfully && session is not null);
-		}
-
-		public UnrealBuildAcceleratorCacheConfig? RequestCacheServer(CancellationToken cancellationToken)
-		{
-			if (!_hordeConfig.bCacheEnabled)
-			{
-				return null;
-			}
-
-			UBAHordeSession? hordeSession = null;
-			if (_hordeSessionTask != null)
-			{
-				(hordeSession, string reason) = _hordeSessionTask.Result;
-				_updateStatus!(0, 6, reason, EpicGames.UBA.LogEntryType.Info, null);
-			}
-			if (hordeSession == null)
-			{
-				return null;
-			}
-
-			return hordeSession.RequestCacheServer(_hordeConfig.CacheDesiredConnectionCount, cancellationToken);
-		}
 
-		public void Start(IUbaAgentCoordinatorScheduler scheduler, Func<LinkedAction, bool> canRunRemotely)
-		{
-			if (!_hordeConfig.bAgentsEnabled)
-			{
-				return;
-			}
+				hordeSession._scheduler = _scheduler;
+				hordeSession._updateStatus = (sr, sc, st, t, sl) => _updateStatus!(sr + i, sc, st, t, sl);
+				hordeSession.UpdateHordeStatus(null);
 
-			int timerPeriod = 5000;
-			bool shownNoAgentsFoundMessage = false;
-
-			UBAHordeSession? hordeSession = null;
-			string reason = "Unknown";
-			if (_hordeSessionTask != null)
-			{
-				(hordeSession, reason) = _hordeSessionTask.Result;
-			}
-
-			if (hordeSession == null)
-			{
-				if (_updateStatus != null)
+				int requestCounter = 0;
+				Timer timer = new(async (_) =>
 				{
-					_updateStatus(0, 6, reason, EpicGames.UBA.LogEntryType.Info, null);
-				}
-				return;
-			}
+					_timers[i]?.Change(Timeout.Infinite, Timeout.Infinite);
 
-			hordeSession._scheduler = scheduler;
-			hordeSession._updateStatus = _updateStatus;
-			hordeSession.UpdateHordeStatus(null);
+					if (_hordeSessionTasks[i] == null)
+					{
+						return;
+					}
 
-			int requestCounter = 0;
+					(UBAHordeSession? hordeSession, string? reason) = await _hordeSessionTasks[i];
 
-			_timer = new(async (_) =>
-			{
-				_timer?.Change(Timeout.Infinite, Timeout.Infinite);
+					if (hordeSession == null)
+					{
+						return;
+					}
 
-				if (_hordeSessionTask == null)
-				{
-					return;
-				}
+					if (_scheduler!.IsEmpty || _cancellationToken.IsCancellationRequested)
+					{
+						hordeSession.UpdateHordeStatus(" - Requests stopped");
+						return;
+					}
 
-				(UBAHordeSession? hordeSession, string? reason) = await _hordeSessionTask;
+					// We are assuming all active logical cores are already being used.. so queueWeight is essentially work that could be executed but can't because of bandwidth
+					double queueThreshold = _ubaConfig.bForceBuildAllRemote ? 0 : 5;
+					LogLevel logLevel = _ubaConfig.bStrict ? LogLevel.Warning : LogLevel.Information;
 
-				if (hordeSession == null)
-				{
-					return;
-				}
+					try
+					{
+						while (true)
+						{
+							if (_cancellationToken.IsCancellationRequested)
+							{
+								hordeSession.UpdateHordeStatus(" - Requests stopped");
+								return;
+							}
 
-				if (scheduler.IsEmpty || _cancellationToken.IsCancellationRequested)
-				{
-					hordeSession.UpdateHordeStatus(" - Requests stopped");
-					return;
-				}
+							hordeSession.RemoveCompleteWorkersAsync();
 
-				// We are assuming all active logical cores are already being used.. so queueWeight is essentially work that could be executed but can't because of bandwidth
-				double queueThreshold = _ubaConfig.bForceBuildAllRemote ? 0 : 5;
-				LogLevel logLevel = _ubaConfig.bStrict ? LogLevel.Warning : LogLevel.Information;
+							double totalWeight;
+							double crossArchitecture;
+							double crossPlatform;
+							_scheduler.GetProcessWeightThatCanRunRemotelyNow(out totalWeight, out crossArchitecture, out crossPlatform);
 
-				try
-				{
-					while (true)
-					{
-						if (_cancellationToken.IsCancellationRequested)
-						{
-							hordeSession.UpdateHordeStatus(" - Requests stopped");
-							return;
-						}
+							double queueWeight = totalWeight;
 
-						hordeSession.RemoveCompleteWorkersAsync();
+							// Hacky, needs to be fixed up
+							if (UBAExecutor.GetHostPlatform() == OSPlatform.OSX && hordeSession.WorkerPlatform == OSPlatform.Linux)
+							{
+								queueWeight = crossPlatform;
+							}
 
-						double totalWeight;
-						double crossArchitecture;
-						double crossPlatform;
-						scheduler.GetProcessWeightThatCanRunRemotelyNow(out totalWeight, out crossArchitecture, out crossPlatform);
-						double queueWeight = totalWeight;
+							int activeCores;
+							int queuedCores;
+							int activeAgents;
+							int queuedAgents;
+							hordeSession.GetCoreCount(out activeCores, out queuedCores, out activeAgents, out queuedAgents);
 
-						int activeCores;
-						int queuedCores;
-						int activeAgents;
-						int queuedAgents;
-						hordeSession.GetCoreCount(out activeCores, out queuedCores, out activeAgents, out queuedAgents);
+							queueWeight -= queuedCores;
 
-						queueWeight -= queuedCores;
+							bool isSatisfied = queueWeight <= queueThreshold || (activeCores + queuedCores) >= hordeConfig.HordeMaxCores || _cancellationToken.IsCancellationRequested;
 
-						bool isSatisfied = queueWeight <= queueThreshold || (activeCores + queuedCores) >= _hordeConfig.HordeMaxCores || _cancellationToken.IsCancellationRequested;
+							hordeSession.UpdateHordeStatus(isSatisfied ? " - Requests paused" : $" - Requesting agent{("...."[..(requestCounter++ % 4)])}");
 
-						hordeSession.UpdateHordeStatus(isSatisfied ? " - Requests paused" : $" - Requesting agent{("...."[..(requestCounter++ % 4)])}");
+							if (isSatisfied)
+							{
+								break;
+							}
 
-						if (isSatisfied)
-						{
-							break;
+							if (!await hordeSession.AddWorkerAsync(activeCores, _cancellationToken))
+							{
+								await Task.Delay(500); // Sleep a little bit to make sure previous horde status is visible
+								string status = (queuedAgents + activeAgents) == 0 ? " - No agents available" : " - No additional agents available";
+								hordeSession.UpdateHordeStatus(status);
+								_logger.LogDebug("{HordeStatus}", status);
+								break;
+							}
 						}
-
-						Requirements requirements = GetRequirements(_hordeConfig);
-						if (!await hordeSession.AddWorkerAsync(requirements, _hordeConfig, activeCores, _cancellationToken))
+					}
+					catch (NoComputeAgentsFoundException ex)
+					{
+						if (!shownNoAgentsFoundMessage)
 						{
-							await Task.Delay(500); // Sleep a little bit to make sure previous horde status is visible
-							string status = (queuedAgents + activeAgents) == 0 ? " - No agents available" : " - No additional agents available";
-							hordeSession.UpdateHordeStatus(status);
-							_logger.LogDebug("{HordeStatus}", status);
-							break;
+							_logger.Log(logLevel, KnownLogEvents.Systemic_Horde_Compute, ex, "No agents found matching requirements (cluster: {ClusterId}, requirements: {Requirements})", ex.ClusterId, ex.Requirements);
+							shownNoAgentsFoundMessage = true;
 						}
 					}
-				}
-				catch (NoComputeAgentsFoundException ex)
-				{
-					if (!shownNoAgentsFoundMessage)
+					catch (ComputeClientException cce)
 					{
-						_logger.Log(logLevel, KnownLogEvents.Systemic_Horde_Compute, ex, "No agents found matching requirements (cluster: {ClusterId}, requirements: {Requirements})", ex.ClusterId, ex.Requirements);
-						shownNoAgentsFoundMessage = true;
+						_logger.Log(logLevel, KnownLogEvents.Systemic_Horde_Compute, cce, "{ErrorMessage}", cce.Message);
 					}
-				}
-				catch (ComputeClientException cce)
-				{
-					_logger.Log(logLevel, KnownLogEvents.Systemic_Horde_Compute, cce, "{ErrorMessage}", cce.Message); 
-				}
-				catch (Exception ex)
-				{
-					if (!_cancellationToken.IsCancellationRequested)
+					catch (Exception ex)
 					{
-						hordeSession.UpdateHordeStatus(" - " + ex.Message);
-						_logger.Log(logLevel, KnownLogEvents.Systemic_Horde_Compute, ex, "Unable to get worker: {Ex}", ex.Message);
+						if (!_cancellationToken.IsCancellationRequested)
+						{
+							hordeSession.UpdateHordeStatus(" - " + ex.Message);
+							_logger.Log(logLevel, KnownLogEvents.Systemic_Horde_Compute, ex, "Unable to get worker: {Ex}", ex.Message);
+						}
 					}
-				}
 
-				_timer?.Change(timerPeriod, Timeout.Infinite);
-			}, null, _hordeConfig.HordeDelay * 1000, timerPeriod);
+					_timers[i]?.Change(timerPeriod, Timeout.Infinite);
+				}, null, Timeout.Infinite, Timeout.Infinite);
+
+				_timers[i] = timer;
+				timer.Change(hordeConfig.HordeDelay * 1000, timerPeriod);
+			}
 		}
 
 		public void Stop()
@@ -1054,16 +396,24 @@ namespace UnrealBuildTool
 				_logger.LogWarning("Cancellation should have been set");
 			}
 
-			if (_hordeSessionTask == null)
+			for (uint i = 0; i != _hordeSessionTasks.Length; ++i)
 			{
-				return;
+				if (_hordeSessionTasks[i] == null)
+				{
+					return;
+				}
+
+				(UBAHordeSession? hordeSession, _) = await _hordeSessionTasks[i];
+				_hordeSessionTasks[i] = Task.FromResult(((UBAHordeSession?)null, ""));
+				if (hordeSession != null)
+				{
+					await hordeSession.DisposeAsync();
+				}
 			}
 
-			(UBAHordeSession? hordeSession, _) = await _hordeSessionTask;
-			_hordeSessionTask = null;
-			if (hordeSession != null)
+			if (_serviceProvider != null)
 			{
-				await hordeSession.DisposeAsync();
+				await _serviceProvider.DisposeAsync();
 			}
 		}
 
@@ -1071,7 +421,10 @@ namespace UnrealBuildTool
 		{
 			if (_updateStatus != null)
 			{
-				_updateStatus(0, 6, "Done", EpicGames.UBA.LogEntryType.Info, null);
+				for (uint i = 0; i != _hordeSessionTasks?.Length; ++i)
+				{
+					_updateStatus(0, 6 + i, "Done", EpicGames.UBA.LogEntryType.Info, null);
+				}
 			}
 		}
 
@@ -1087,66 +440,30 @@ namespace UnrealBuildTool
 		{
 			if (disposing)
 			{
-				_timer?.Dispose();
-				_timer = null;
-			}
-		}
-		
-		private static Requirements GetRequirements(UnrealBuildAcceleratorHordeConfig config)
-		{
-			Requirements requirements = new() { Exclusive = true };
-			
-			if (!String.IsNullOrEmpty(config.HordeCluster))
-			{
-				// Apply filtering options when a cluster, either explicit or auto is specified, ensuring the cluster options take precidence.
-				string condition;
-				if (config.HordeCondition != null)
-				{
-					// Use explicitly configured condition, allowing cross-platform distribution (e.g., Mac host to Linux agents)
-					condition = config.HordeCondition;
-				}
-				else if (OperatingSystem.IsWindows())
-				{
-					condition = $"({KnownPropertyNames.OsFamily} == 'Windows' || {KnownPropertyNames.WineEnabled} == 'true')";
-				}
-				else if (OperatingSystem.IsMacOS())
-				{
-					condition = $"{KnownPropertyNames.OsFamily} == 'MacOS'";
-				}
-				else if (OperatingSystem.IsLinux())
-				{
-					condition = $"{KnownPropertyNames.OsFamily} == 'Linux'";
-				}
-				else
-				{
-					condition = "";
-				}
-				requirements.Condition = Condition.Parse(condition);
-			}
-			else
-			{
-				if (!String.IsNullOrEmpty(config.HordePool))
-				{
-					requirements.Pool = config.HordePool;
-				}
-				
-				if (config.HordeCondition != null)
+				_initDoneEvent.Dispose();
+
+				for (uint i = 0; i != _timers.Length; ++i)
 				{
-					requirements.Condition = Condition.Parse(config.HordeCondition);
+					_timers[i]?.Dispose();
+					_timers[i] = null;
 				}
 			}
-			
-			return requirements;
 		}
-
+		
 		readonly ILogger _logger;
 		readonly string _title;
 		readonly UnrealBuildAcceleratorConfig _ubaConfig;
-		readonly UnrealBuildAcceleratorHordeConfig _hordeConfig;
+		readonly UnrealBuildAcceleratorHordeConfig[] _hordeConfigs;
+		Task<(UBAHordeSession?, string)>[] _hordeSessionTasks = [];
+		readonly Timer?[] _timers = [];
+		readonly ManualResetEventSlim _initDoneEvent = new();
+
+		bool _initDone;
+		readonly object _startLock = new();
+		IUbaAgentCoordinatorScheduler? _scheduler;
 
+		ServiceProvider? _serviceProvider;
 		CancellationToken _cancellationToken;
-		Task<(UBAHordeSession?, string)>? _hordeSessionTask;
 		StatusUpdateAction? _updateStatus;
-		Timer? _timer;
 	}
 }
diff --git a/Engine/Source/Programs/UnrealBuildTool/Executors/UnrealBuildAccelerator/UBAExecutor.cs b/Engine/Source/Programs/UnrealBuildTool/Executors/UnrealBuildAccelerator/UBAExecutor.cs
--- a/Engine/Source/Programs/UnrealBuildTool/Executors/UnrealBuildAccelerator/UBAExecutor.cs
+++ b/Engine/Source/Programs/UnrealBuildTool/Executors/UnrealBuildAccelerator/UBAExecutor.cs
@@ -115,33 +115,60 @@ namespace UnrealBuildTool
 		// Event invoked when ISessionServer is created. Can be used to for example register virtual files etc
 		public event SessionCreatedDelegate? OnSessionCreated;
 
-		public static DirectoryReference UbaBinariesDir
+		public static OSPlatform GetHostPlatform()
 		{
-			get
+			if (OperatingSystem.IsLinux())
 			{
-				if (OperatingSystem.IsWindows())
-				{
-					#pragma warning disable CA1308 // Normalize strings to uppercase
-					return DirectoryReference.Combine(Unreal.EngineDirectory, "Binaries", "Win64", "UnrealBuildAccelerator", RuntimeInformation.ProcessArchitecture.ToString().ToLowerInvariant());
-					#pragma warning restore CA1308 // Normalize strings to uppercase
-				}
-				else if (OperatingSystem.IsLinux())
-				{
-					if (RuntimeInformation.ProcessArchitecture == Architecture.X64)
-					{
-						return DirectoryReference.Combine(Unreal.EngineDirectory, "Binaries", "Linux", "UnrealBuildAccelerator");
-					}
-					else if (RuntimeInformation.ProcessArchitecture == Architecture.Arm64)
-					{
-						return DirectoryReference.Combine(Unreal.EngineDirectory, "Binaries", "LinuxArm64", "UnrealBuildAccelerator");
-					}
-				}
-				else if (OperatingSystem.IsMacOS())
+				return OSPlatform.Linux;
+			}
+			if (OperatingSystem.IsWindows())
+			{
+				return OSPlatform.Windows;
+			}
+			return OSPlatform.OSX;
+		}
+
+		public static string GetPlatformName(OSPlatform platform)
+		{
+			if (platform == OSPlatform.Windows)
+			{
+				return "Win64";
+			}
+			if (platform == OSPlatform.Linux)
+			{
+				return "Linux";
+			}
+			if (platform == OSPlatform.OSX)
+			{
+				return "Mac";
+			}
+			throw new PlatformNotSupportedException();
+		}
+
+		public static DirectoryReference GetUbaBinariesDir(OSPlatform platform)
+		{
+			if (platform == OSPlatform.Windows)
+			{
+				#pragma warning disable CA1308 // Normalize strings to uppercase
+				return DirectoryReference.Combine(Unreal.EngineDirectory, "Binaries", "Win64", "UnrealBuildAccelerator", RuntimeInformation.ProcessArchitecture.ToString().ToLowerInvariant());
+				#pragma warning restore CA1308 // Normalize strings to uppercase
+			}
+			else if (platform == OSPlatform.Linux)
+			{
+				//if (RuntimeInformation.ProcessArchitecture == Architecture.X64)
 				{
-					return DirectoryReference.Combine(Unreal.EngineDirectory, "Binaries", "Mac", "UnrealBuildAccelerator");
+					return DirectoryReference.Combine(Unreal.EngineDirectory, "Binaries", "Linux", "UnrealBuildAccelerator");
 				}
-				throw new PlatformNotSupportedException();
+				//else if (RuntimeInformation.ProcessArchitecture == Architecture.Arm64)
+				//{
+				//	return DirectoryReference.Combine(Unreal.EngineDirectory, "Binaries", "LinuxArm64", "UnrealBuildAccelerator");
+				//}
+			}
+			else if (platform == OSPlatform.OSX)
+			{
+				return DirectoryReference.Combine(Unreal.EngineDirectory, "Binaries", "Mac", "UnrealBuildAccelerator");
 			}
+			throw new PlatformNotSupportedException();
 		}
 
 		internal static UnrealBuildAcceleratorConfig GetAppliedConfig(IEnumerable<TargetDescriptor> targetDescriptors)
@@ -190,6 +217,7 @@ namespace UnrealBuildTool
 
 			if (!UBAConfig.bForceNoCache)
 			{
+				// These are non-coordinator cache servers. Coordinators can also provide cache servers
 				const string ProviderPrefix = "Uba.CacheProvider";
 
 				IEnumerable<string> providers = Unreal.IsBuildMachine()
@@ -302,7 +330,9 @@ namespace UnrealBuildTool
 			DirectoryReference.CreateDirectory(ubaTraceFile.Directory);
 			Log.BackupLogFile(ubaTraceFile);
 
-			FileReference configFile = FileReference.Combine(UbaBinariesDir, "UbaHost.toml");
+			DirectoryReference ubaBinariesDir = GetUbaBinariesDir(GetHostPlatform());
+
+			FileReference configFile = FileReference.Combine(ubaBinariesDir, "UbaHost.toml");
 			_config = EpicGames.UBA.IConfig.LoadConfig(configFile.FullName);
 			if (!String.IsNullOrEmpty(UBAConfig.CompressionLevel))
 			{
@@ -414,7 +444,7 @@ namespace UnrealBuildTool
 
 			if (!UBAConfig.bDisableRemote)
 			{
-				_server.SetClientsConfig(FileReference.Combine(UbaBinariesDir, "UbaAgent.toml").FullName);
+				_server.SetClientsConfig(FileReference.Combine(ubaBinariesDir, "UbaAgent.toml").FullName);
 
 				_server.StartServer(UBAConfig.Host, UBAConfig.Port, Crypto);
 			}
@@ -655,7 +685,7 @@ namespace UnrealBuildTool
 		[SupportedOSPlatform("windows")]
 		static void LaunchVisualizerWin()
 		{
-			FileReference visualizerPath = FileReference.Combine(UbaBinariesDir, "UbaVisualizer.exe");
+			FileReference visualizerPath = FileReference.Combine(GetUbaBinariesDir(OSPlatform.Windows), "UbaVisualizer.exe");
 			FileReference tempPath = FileReference.Combine(new DirectoryReference(System.IO.Path.GetTempPath()), visualizerPath.GetFileName());
 			if (!FileReference.Exists(visualizerPath))
 			{
@@ -698,7 +728,7 @@ namespace UnrealBuildTool
 		[SupportedOSPlatform("macos")]
 		static void LaunchVisualizerMac()
 		{
-			DirectoryReference visualizerPath = DirectoryReference.Combine(UbaBinariesDir, "UbaVisualizer.app");
+			DirectoryReference visualizerPath = DirectoryReference.Combine(GetUbaBinariesDir(OSPlatform.OSX), "UbaVisualizer.app");
 			if (!DirectoryReference.Exists(visualizerPath))
 			{
 				return;
diff --git a/Engine/Source/Programs/UnrealBuildTool/Executors/UnrealBuildAccelerator/UBAHordeSession.cs b/Engine/Source/Programs/UnrealBuildTool/Executors/UnrealBuildAccelerator/UBAHordeSession.cs
new file mode 100644
--- /dev/null
+++ b/Engine/Source/Programs/UnrealBuildTool/Executors/UnrealBuildAccelerator/UBAHordeSession.cs
@@ -0,0 +1,767 @@
+// Copyright Epic Games, Inc. All Rights Reserved.
+
+using System;
+using System.Collections.Generic;
+using System.Diagnostics;
+using System.Linq;
+using System.Runtime.InteropServices;
+using System.Threading;
+using System.Threading.Tasks;
+using EpicGames.Core;
+using EpicGames.Horde;
+using EpicGames.Horde.Agents;
+using EpicGames.Horde.Common;
+using EpicGames.Horde.Compute;
+using EpicGames.Horde.Storage;
+using EpicGames.Horde.Storage.Bundles;
+using EpicGames.Horde.Storage.Nodes;
+using Microsoft.Extensions.Logging;
+using Microsoft.Extensions.Logging.Abstractions;
+using UnrealBuildBase;
+
+namespace UnrealBuildTool
+{
+	// statusRow, statusColumn, statusText, statusType, statusLink
+	using StatusUpdateAction = Action<uint, uint, string, EpicGames.UBA.LogEntryType, string?>;
+
+	internal class UBAHordeSession : IAsyncDisposable
+	{
+		/// <summary>
+		/// ID for this Horde session
+		/// </summary>
+		readonly Guid _id = Guid.NewGuid();
+		readonly IComputeClient _client;
+
+		const string ResourceLogicalCores = "LogicalCores";
+		ClusterId? _clusterId;
+
+		readonly bool _strict;
+		readonly ConnectionMode? _connectionMode;
+		readonly Encryption? _encryption;
+		readonly CancellationTokenSource _cancellationTokenSource = new();
+		readonly ILogger _logger;
+
+		readonly UnrealBuildAcceleratorConfig _config;
+		readonly UnrealBuildAcceleratorHordeConfig _hordeConfig;
+		internal StatusUpdateAction? _updateStatus;
+		internal IUbaAgentCoordinatorScheduler? _scheduler;
+
+		readonly BundleStorageNamespace _storage = BundleStorageNamespace.CreateInMemory(NullLogger.Instance);
+		BlobLocator _ubaAgentLocator;
+
+		readonly string _crypto;
+
+		OSPlatform _targetWorkerPlatform;
+
+		public class Worker
+		{
+			public Stopwatch? StartTime { get; init; }
+			public string? Ip { get; init; }
+			public ConnectionMetadataPort? Port { get; init; }
+			public ConnectionMetadataPort? ProxyPort { get; init; }
+
+			public int NumLogicalCores { get; set; }
+			public bool Active { get; set; }
+			public Task? BackgroundTask { get; set; }
+		}
+
+		readonly List<Worker> _workers = new();
+
+		public UBAHordeSession(UnrealBuildAcceleratorConfig config, UnrealBuildAcceleratorHordeConfig hordeConfig, IComputeClient client, string crypto, bool strict, ConnectionMode? connectionMode, Encryption? encryption, ILogger logger)
+		{
+			_config = config;
+			_hordeConfig = hordeConfig;
+			_client = client;
+			_strict = strict;
+			_connectionMode = connectionMode;
+			_encryption = encryption;
+			_logger = logger;
+			_crypto = crypto;
+
+			if (connectionMode == ConnectionMode.Relay && String.IsNullOrEmpty(_crypto))
+			{
+				_crypto = UBAExecutor.CreateCrypto();
+				_encryption = Encryption.Ssl;
+			}
+		}
+
+		public OSPlatform WorkerPlatform => _targetWorkerPlatform;
+
+		public async ValueTask DisposeAsync()
+		{
+			if (_workers.Count > 0) // Should handle double-dispose, prevent cancelling twice
+			{
+				await _cancellationTokenSource.CancelAsync();
+
+				for (int idx = _workers.Count - 1; idx >= 0; idx--)
+				{
+					await _workers[idx].BackgroundTask!;
+					lock (_workers)
+					{
+						_workers.RemoveAt(idx);
+					}
+				}
+			}
+
+			// Set CPU resource need to zero (will also expire on the server if not updated)
+			await UpdateCpuCoreNeedAsync(0);
+
+			_cancellationTokenSource.Dispose();
+		}
+
+		public async Task InitAsync(bool useSentry, string? hordeCondition, CancellationToken cancellationToken)
+		{
+			if (!String.IsNullOrEmpty(_hordeConfig.AgentPlatform))
+			{
+				_targetWorkerPlatform = GetTargetWorkerPlatform(_hordeConfig.AgentPlatform);
+			}
+			else
+			{
+				_targetWorkerPlatform = GetTargetWorkerPlatform(hordeCondition);
+			}
+
+			DirectoryReference ubaDir = UBAExecutor.GetUbaBinariesDir(_targetWorkerPlatform);
+
+			List<string> agentFiles = [];
+
+			if (_targetWorkerPlatform == OSPlatform.Windows)
+			{
+				agentFiles.Add("UbaAgent.exe");
+				if (FileReference.Exists(FileReference.Combine(ubaDir, "UbaWine.dll.so")))
+				{
+					agentFiles.Add("UbaWine.dll.so"); // Not needed but used by linux/wine helpers for more detailed logging or optimizations
+				}
+				if (useSentry)
+				{
+					agentFiles.Add("crashpad_handler.exe");
+					agentFiles.Add("sentry.dll");
+				}
+			}
+			else if (_targetWorkerPlatform == OSPlatform.Linux)
+			{
+				agentFiles.Add("UbaAgent");
+				bool bIsDebug = false;
+				if (bIsDebug)
+				{
+					agentFiles.Add("UbaAgent.debug");
+					agentFiles.Add("UbaAgent.sym");
+					agentFiles.Add("libclang_rt.tsan.so"); // Needs to be copied from autosdk
+					agentFiles.Add("llvm-symbolizer"); // Needs to be copied from autosdk
+				}
+			}
+			else if (_targetWorkerPlatform == OSPlatform.OSX)
+			{
+				agentFiles.Add("UbaAgent");
+			}
+			else
+			{
+				throw new PlatformNotSupportedException();
+			}
+
+			_ubaAgentLocator = await CreateToolAsync(ubaDir, agentFiles.Select(x => FileReference.Combine(ubaDir, x)), cancellationToken);
+			_logger.LogInformation("Created tool bundle with locator {UbaAgentLocator}", _ubaAgentLocator.ToString());
+		}
+
+		public UnrealBuildAcceleratorCacheConfig? RequestCacheServer(int desiredConnectionCount, CancellationToken cancellationToken)
+		{
+			ClusterId clusterId = new(UnrealBuildAcceleratorHordeConfig.ClusterDefault);
+			using ITimelineEvent timelineEvent = Timeline.ScopeEvent($"HordeGetCache");
+			UbaConfig ubaConfig = _client.AllocateUbaCacheServerAsync(clusterId, cancellationToken).Result;
+			timelineEvent.Finish();
+			if (String.IsNullOrEmpty(ubaConfig.CacheEndpoint))
+			{
+				return null;
+			}
+			bool writeAccess = ubaConfig.WriteAccess;
+			return new UnrealBuildAcceleratorCacheConfig() { CacheServer = ubaConfig.CacheEndpoint, bRequireVfs = true, WriteCache = writeAccess.ToString(), Crypto = ubaConfig.CacheSessionKey, DesiredConnectionCount = desiredConnectionCount };
+		}
+
+		async Task<BlobLocator> CreateToolAsync(DirectoryReference baseDir, IEnumerable<FileReference> files, CancellationToken cancellationToken)
+		{
+			using ITimelineEvent __ = Timeline.ScopeEvent($"HordeCreateTool");
+
+			// TODO: should drive this off API version reported by server
+			BlobSerializerOptions serializerOptions = BlobSerializerOptions.Create(HordeApiVersion.Initial);
+
+			await using IBlobWriter writer = _storage.CreateBlobWriter(serializerOptions: serializerOptions, cancellationToken: cancellationToken);
+			DirectoryNode sandbox = new();
+			await sandbox.AddFilesAsync(baseDir, files, writer, cancellationToken: cancellationToken);
+
+			// If host is windows and we want to use non-windows helper, then UbaAgent does not have executable flag so we have to add it
+			if (_targetWorkerPlatform != OSPlatform.Windows && UBAExecutor.GetHostPlatform() == OSPlatform.Windows)
+			{
+				sandbox.GetFileEntry("UbaAgent").Flags |= FileEntryFlags.Executable;
+			}
+
+			IHashedBlobRef <DirectoryNode> handle = await writer.WriteBlobAsync(sandbox, cancellationToken);
+			await writer.FlushAsync(cancellationToken);
+			return handle.GetLocator();
+		}
+
+		public async void RemoveCompleteWorkersAsync()
+		{
+			for (int idx = 0; idx < _workers.Count; idx++)
+			{
+				Worker worker = _workers[idx];
+				if (worker.BackgroundTask!.IsCompleted)
+				{
+					await worker.BackgroundTask;
+					lock (_workers)
+					{
+						_workers.RemoveAt(idx--);
+					}
+				}
+			}
+		}
+
+		public void GetCoreCount(out int active, out int queued, out int activeAgents, out int queuedAgents)
+		{
+			active = 0;
+			queued = 0;
+			activeAgents = 0;
+			queuedAgents = 0;
+			lock (_workers)
+			{
+				foreach (Worker worker in _workers)
+				{
+					if (worker.NumLogicalCores == 0)
+					{
+						continue;
+					}
+					if (worker.Active)
+					{
+						++activeAgents;
+						active += worker.NumLogicalCores;
+					}
+					else
+					{
+						++queuedAgents;
+						queued += worker.NumLogicalCores;
+					}
+				}
+			}
+		}
+
+		int _workerId = 0;
+
+		static readonly HashSet<StringView> s_logProperties = new()
+			{
+				"ComputeIp",
+				"CPU",
+				"RAM",
+				"DiskFreeSpace",
+				"PhysicalCores",
+				"LogicalCores",
+				"EC2",
+				"LeaseId",
+				"aws-instance-type",
+			};
+
+		private static Requirements GetRequirements(UnrealBuildAcceleratorHordeConfig config)
+		{
+			Requirements requirements = new() { Exclusive = true };
+
+			if (!String.IsNullOrEmpty(config.HordeCluster))
+			{
+				OSPlatform platform = GetTargetWorkerPlatform(config.AgentPlatform);
+
+				// Apply filtering options when a cluster, either explicit or auto is specified, ensuring the cluster options take precidence.
+				string condition;
+				if (config.HordeCondition != null)
+				{
+					// Use explicitly configured condition, allowing cross-platform distribution (e.g., Mac host to Linux agents)
+					condition = config.HordeCondition;
+				}
+				else if (platform == OSPlatform.Windows)
+				{
+					condition = $"({KnownPropertyNames.OsFamily} == 'Windows' || {KnownPropertyNames.WineEnabled} == 'true')";
+				}
+				else if (platform == OSPlatform.OSX)
+				{
+					condition = $"{KnownPropertyNames.OsFamily} == 'MacOS'";
+				}
+				else if (platform == OSPlatform.Linux)
+				{
+					condition = $"{KnownPropertyNames.OsFamily} == 'Linux'";
+				}
+				else
+				{
+					condition = "";
+				}
+				requirements.Condition = Condition.Parse(condition);
+			}
+			else
+			{
+				if (!String.IsNullOrEmpty(config.HordePool))
+				{
+					requirements.Pool = config.HordePool;
+				}
+
+				if (config.HordeCondition != null)
+				{
+					requirements.Condition = Condition.Parse(config.HordeCondition);
+				}
+			}
+
+			return requirements;
+		}
+
+		public async Task<bool> AddWorkerAsync(int activeCores, CancellationToken cancellationToken)
+		{
+			if (_client == null)
+			{
+				throw new InvalidOperationException($"Session not initialized. Call {nameof(InitAsync)} first");
+			}
+
+			int workerId = _workerId;
+
+			string agentNamePrefix = _targetWorkerPlatform != UBAExecutor.GetHostPlatform() ? UBAExecutor.GetPlatformName(_targetWorkerPlatform) : "";
+			PrefixLogger workerLogger = new($"[{agentNamePrefix}Worker{workerId}]", _logger);
+
+			const string UbaPortName = "UbaPort";
+			const string UbaProxyPortName = "UbaProxyPort";
+			const int UbaPort = 7001;
+			const int UbaProxyPort = 7002;
+
+			// Request ID that is unique per attempt to acquire the same compute lease/worker
+			// Primarily for tracking worker demand on Horde server as UBAExecutor will repeatedly try adding a new worker
+			string requestId = $"{_id}-worker-{workerId}";
+			IComputeLease? lease = null;
+			try
+			{
+				Stopwatch stopwatch = Stopwatch.StartNew();
+				ConnectionMetadataRequest cmr = new()
+				{
+					ModePreference = _connectionMode,
+					Encryption = _encryption,
+					Ports = { { UbaPortName, UbaPort } },
+					InactivityTimeoutMs = 90000
+				};
+
+				Requirements requirements = GetRequirements(_hordeConfig);
+
+				await ResolveClusterIdAsync(requirements, requestId, cmr, workerLogger, cancellationToken);
+				lease = await _client.TryAssignWorkerAsync(_clusterId, requirements, requestId, cmr, false, workerLogger, cancellationToken);
+				if (lease == null)
+				{
+					_logger.LogDebug("Unable to assign a remote worker");
+
+					int missingNumCores = Math.Max(0, _hordeConfig.HordeMaxCores - activeCores);
+					await UpdateCpuCoreNeedAsync(missingNumCores, cancellationToken);
+					return false;
+				}
+
+				++_workerId;
+
+				workerLogger.LogDebug("Agent properties:");
+
+				string agentName = String.Empty;
+				int numLogicalCores = 24; // Assume 24 if something goes wrong here and property is not found
+				string computeIp = String.Empty;
+				string leaseId = String.Empty;
+				string instanceType = String.Empty;
+				foreach (string property in lease.Properties)
+				{
+					int equalsIdx = property.IndexOf('=', StringComparison.OrdinalIgnoreCase);
+					StringView propertyName = new(property, 0, equalsIdx);
+					if (s_logProperties.Contains(propertyName))
+					{
+						_logger.LogDebug("  {Property}", property);
+
+						if (propertyName == ResourceLogicalCores && Int32.TryParse(property.AsSpan(equalsIdx + 1), out int value))
+						{
+							numLogicalCores = value;
+						}
+						else if (propertyName == "ComputeIp")
+						{
+							computeIp = property[(equalsIdx + 1)..];
+						}
+						else if (propertyName == "EC2")
+						{
+							agentName = $"{agentNamePrefix}Worker{workerId}";
+						}
+						else if (propertyName == "LeaseId")
+						{
+							leaseId = property[(equalsIdx + 1)..];
+						}
+						else if (propertyName == "aws-instance-type")
+						{
+							instanceType = property[(equalsIdx + 1)..];
+						}
+					}
+				}
+
+				string desc = instanceType;
+				if (!String.IsNullOrEmpty(leaseId))
+				{
+					string link = $" {_hordeConfig.HordeServer!.TrimEnd('/')}/lease/{leaseId}";
+					desc += link;
+				}
+				// When using relay connection mode, the IP will be relay server's IP
+				string ip = String.IsNullOrEmpty(lease.Ip) ? computeIp : lease.Ip;
+
+				if (!lease.Ports.TryGetValue(UbaPortName, out ConnectionMetadataPort? ubaPort))
+				{
+					ubaPort = new ConnectionMetadataPort(UbaPort, UbaPort);
+				}
+
+				if (!lease.Ports.TryGetValue(UbaProxyPortName, out ConnectionMetadataPort? ubaProxyPort))
+				{
+					ubaProxyPort = new ConnectionMetadataPort(UbaProxyPort, UbaProxyPort);
+				}
+
+				Worker worker = new()
+				{
+					StartTime = stopwatch,
+					NumLogicalCores = numLogicalCores,
+					Ip = ip,
+					Port = ubaPort,
+					ProxyPort = ubaProxyPort,
+				};
+				worker.BackgroundTask = RunWorkerAsync(worker, lease, workerLogger, agentName, desc, _cancellationTokenSource.Token);
+				lock (_workers)
+				{
+					_workers.Add(worker);
+				}
+				UpdateHordeStatus(null);
+				lease = null; // Will be disposed by RunWorkerAsync
+
+				return true;
+			}
+			finally
+			{
+				if (lease != null)
+				{
+					await lease.DisposeAsync();
+				}
+			}
+		}
+
+		int _targetCoreCount;
+
+		async Task UpdateCpuCoreNeedAsync(int targetCoreCount, CancellationToken cancellationToken = default)
+		{
+			string? pool = _hordeConfig.HordePool;
+			if (_client != null && pool != null && _clusterId != null && targetCoreCount != _targetCoreCount)
+			{
+				_targetCoreCount = targetCoreCount;
+
+				_logger.LogDebug("Setting CPU core need to {TargetCoreCount}", targetCoreCount);
+				Dictionary<string, int> resourceNeeds = new() { { ResourceLogicalCores, targetCoreCount } };
+				try
+				{
+					await _client.DeclareResourceNeedsAsync(_clusterId.Value, pool, resourceNeeds, cancellationToken);
+				}
+				catch (Exception e)
+				{
+					_logger.Log(_strict ? LogLevel.Error : LogLevel.Debug, KnownLogEvents.Systemic_Horde_Compute, e, "Failed updating resource need to {TargetCoreCount} cores", targetCoreCount);
+				}
+			}
+		}
+
+		private async Task ResolveClusterIdAsync(Requirements requirements, string requestId, ConnectionMetadataRequest cmr, ILogger logger, CancellationToken cancellationToken = default)
+		{
+			if (_client == null)
+			{
+				throw new InvalidOperationException($"Session not initialized. Call {nameof(InitAsync)} first");
+			}
+
+			if (_clusterId == null)
+			{
+				if (_hordeConfig.HordeCluster == UnrealBuildAcceleratorHordeConfig.ClusterAuto)
+				{
+					_clusterId = await _client.GetClusterAsync(requirements, requestId, cmr, logger, cancellationToken);
+				}
+				else if (!String.IsNullOrEmpty(_hordeConfig.HordeCluster))
+				{
+					_clusterId = new ClusterId(_hordeConfig.HordeCluster);
+				}
+				else
+				{
+					_clusterId = new ClusterId(UnrealBuildAcceleratorHordeConfig.ClusterDefault);
+				}
+				_logger.LogInformation("Horde cluster resolved as '{ClusterId}'", _clusterId.ToString());
+			}
+		}
+
+		public static async Task<(UBAHordeSession?, string)> TryCreateHordeSessionAsync(UnrealBuildAcceleratorConfig ubaConfig, UnrealBuildAcceleratorHordeConfig hordeConfig, IComputeClient client, string crypto, bool bStrictErrors, ILogger logger, CancellationToken cancellationToken = default)
+		{
+			using ITimelineEvent hordeCreateTimlineEvent = Timeline.ScopeEvent($"HordeCreate");
+
+			if (hordeConfig.bDisableHorde)
+			{
+				logger.LogInformation("Horde disabled via command line option.");
+				return (null, "Disabled via command line");
+			}
+
+			if (String.IsNullOrEmpty(hordeConfig.HordeServer) && HordeOptions.GetServerUrlFromEnvironment() == null && HordeOptions.GetDefaultServerUrl() == null)
+			{
+				logger.LogInformation("Horde disabled. Url not set.");
+				return (null, "No server url set");
+			}
+
+			Uri? server = (hordeConfig.HordeServer == null) ? null : new Uri(hordeConfig.HordeServer);
+			string? token = hordeConfig.HordeToken;
+
+			ConnectionMode? connectionMode = Enum.TryParse(hordeConfig.HordeConnectionMode, true, out ConnectionMode cm) ? cm : null;
+			Encryption? encryption = Enum.TryParse(hordeConfig.HordeEncryption, true, out Encryption enc) ? enc : null;
+
+			// Default to SSL encryption for relay mode if unset
+			if (connectionMode == ConnectionMode.Relay && encryption == null)
+			{
+				encryption = Encryption.Ssl;
+			}
+
+			logger.LogInformation("Horde URL: {Server}, Pool: {Pool}, Cluster: {Cluster}, Condition: {Condition}, Connection: {Connection}, Encryption: {Encryption}, MaxCores: {MaxCores}, MaxWorkers: {MaxWorkers}, MaxIdle: {MaxIdle}s",
+				server, hordeConfig.HordePool ?? "(none)", hordeConfig.HordeCluster ?? "(none)", hordeConfig.HordeCondition ?? "(none)", connectionMode?.ToString() ?? "(none)", encryption?.ToString() ?? "(none)",
+				hordeConfig.HordeMaxCores, hordeConfig.HordeMaxWorkers, hordeConfig.HordeMaxIdle);
+			try
+			{
+				UBAHordeSession session = new(ubaConfig, hordeConfig, client, crypto, bStrictErrors, connectionMode, encryption, logger);
+				hordeCreateTimlineEvent.Dispose();
+				await session.InitAsync(useSentry: !String.IsNullOrEmpty(hordeConfig.UBASentryUrl), hordeConfig.HordeCondition, cancellationToken);
+				return (session, "Ready");
+			}
+			catch (TaskCanceledException)
+			{
+				return (null, "Cancelled");
+			}
+			catch (Exception ex)
+			{
+				logger.Log(bStrictErrors ? LogLevel.Error : LogLevel.Information, ex, "Unable to create Horde session: {Message}", ex.Message);
+				return (null, ex.Message);
+			}
+		}
+
+		async Task RunWorkerAsync(Worker self, IComputeLease lease, ILogger logger, string agentName, string desc, CancellationToken cancellationToken)
+		{
+			logger.LogDebug("Running worker task..");
+			try
+			{
+				await using (_ = lease)
+				{
+					// Create a message channel on channel id 0. The Horde Agent always listens on this channel for requests.
+					const int PrimaryChannelId = 0;
+					using (AgentMessageChannel channel = lease.Socket.CreateAgentMessageChannel(PrimaryChannelId, 4 * 1024 * 1024))
+					{
+						logger.LogDebug("Waiting for attach...");
+
+						TimeSpan attachTimeout = TimeSpan.FromSeconds(20.0);
+						try
+						{
+							Task attachTask = channel.WaitForAttachAsync(cancellationToken).AsTask();
+							await attachTask.WaitAsync(attachTimeout, cancellationToken);
+						}
+						catch (TimeoutException)
+						{
+							logger.Log(_strict ? LogLevel.Error : LogLevel.Debug, KnownLogEvents.Systemic_Horde_Compute, "Waited {Time}s on attach message. Giving up", (int)attachTimeout.TotalSeconds);
+							throw;
+						}
+
+						logger.LogDebug("Uploading files...");
+						await channel.UploadFilesAsync("", _ubaAgentLocator, _storage.Backend, cancellationToken);
+
+						string hordeHost = _config.Host;
+						if (!String.IsNullOrEmpty(_hordeConfig.HordeHost))
+						{
+							hordeHost = _hordeConfig.HordeHost;
+						}
+
+						bool useListen = !String.IsNullOrEmpty(_hordeConfig.HordeHost);
+						List<string> arguments = new();
+
+						if (!String.IsNullOrEmpty(agentName))
+						{
+							arguments.Add($"-name={agentName}");
+						}
+
+						if (useListen)
+						{
+							arguments.Add($"-Host={hordeHost}:{_config.Port}");
+						}
+						else
+						{
+							arguments.Add($"-Listen={self.Port!.AgentPort}");
+							arguments.Add("-ListenTimeout=10");
+						}
+
+						if (!String.IsNullOrEmpty(_crypto))
+						{
+							arguments.Add($"-crypto={_crypto}");
+						}
+
+						arguments.Add("-NoPoll");
+						arguments.Add("-Quiet");
+						if (!String.IsNullOrEmpty(_hordeConfig.UBASentryUrl))
+						{
+							arguments.Add($"-Sentry=\"{_hordeConfig.UBASentryUrl}\"");
+						}
+						arguments.Add("-ProxyPort=" + self.ProxyPort!.AgentPort);
+						if (_config.bUseQuic)
+						{
+							arguments.Add("-quic");
+						}
+						//arguments.Add("-NoStore");
+						//arguments.Add("-KillRandom"); // For debugging
+
+						if (_targetWorkerPlatform == OSPlatform.OSX)
+						{
+							// we need to populate the cas with all known xcodes so we can serve all the ones we have installed
+							string xcodeVersion = Utils.RunLocalProcessAndReturnStdOut("/bin/sh", "-c '/usr/bin/defaults read $(xcode-select -p)/../version.plist ProductBuildVersion");
+							arguments.Add($"-populateCasFromXcodeVersion={xcodeVersion}");
+							arguments.Add("-killtcphogs");
+						}
+
+						arguments.Add("-Dir=%UE_HORDE_SHARED_DIR%/Uba");
+						arguments.Add("-Eventfile=%UE_HORDE_TERMINATION_SIGNAL_FILE%");
+						arguments.Add("-MaxCpu=%UE_HORDE_CPU_COUNT%");
+						arguments.Add("-MulCpu=%UE_HORDE_CPU_MULTIPLIER%");
+						arguments.Add("-MaxCon=8");
+						arguments.Add($"-MaxWorkers={_hordeConfig.HordeMaxWorkers}");
+						arguments.Add($"-MaxIdle={_hordeConfig.HordeMaxIdle}");
+						arguments.Add("-UseCrawler");
+						// arguments.Add("-ProxyUseLocalStorage"); // If set then agent's cas storage will be used by proxy. [honk] This is disabled because we suspect a hang in this feature
+						//arguments.Add("-UseIocp=4");
+						//arguments.Add("-NoStore"); // This means that no cas files will be written to disk.. so no state is stored between runs
+
+						if (!String.IsNullOrEmpty(desc))
+						{
+							arguments.Add($"-Description=\"{desc}\"");
+						}
+
+						if (_config.bLogEnabled)
+						{
+							arguments.Add("-Log");
+						}
+
+						LogLevel logLevel = _config.bDetailedLog ? LogLevel.Information : LogLevel.Debug;
+
+						string executable = _targetWorkerPlatform == OSPlatform.Windows ? "UbaAgent.exe" : "UbaAgent";
+						logger.Log(logLevel, "Executing child process: {Executable} {Arguments}", executable, CommandLineArguments.Join(arguments));
+
+						bool allowWine = _hordeConfig.bHordeAllowWine && GetTargetWorkerPlatform(_hordeConfig.AgentPlatform) == OSPlatform.Windows;
+						ExecuteProcessFlags execFlags = allowWine ? ExecuteProcessFlags.UseWine : ExecuteProcessFlags.None;
+						await using AgentManagedProcess process = await channel.ExecuteAsync(executable, arguments, null, null, execFlags, cancellationToken);
+						bool shouldConnect = !useListen;
+						bool isFirstRead = true;
+						string? line;
+
+						while ((line = await process.ReadLineAsync(cancellationToken)) != null)
+						{
+							logger.Log(logLevel, "{Line}", line);
+
+							if (shouldConnect && line.Contains("Listening on", StringComparison.OrdinalIgnoreCase)) // This log entry means that the agent is ready for connections.
+							{
+								long totalMs = self.StartTime!.ElapsedMilliseconds;
+								logger.LogInformation("Connecting to UbaAgent on {Ip}:{Port} (local agent port {AgentPort}) {Seconds}.{Milliseconds} seconds after assigned",
+									self.Ip, self.Port!.Port, self.Port.AgentPort, totalMs / 1000, totalMs % 1000);
+
+								if (!_scheduler!.AddClient(self.Ip!, self.Port.Port, _crypto))
+								{
+									break;
+								}
+								shouldConnect = false;
+							}
+
+							if (isFirstRead)
+							{
+								isFirstRead = false;
+								self.Active = true;
+								UpdateHordeStatus(null);
+							}
+						}
+						self.NumLogicalCores = 0;
+						UpdateHordeStatus(null);
+						logger.LogDebug("Shutting down process");
+						int ExitCode = await process.WaitForExitAsync(cancellationToken);
+						if (ExitCode != 0)
+						{
+							logger.LogInformation("UbaAgent exited with exit code: {ExitCode}", ExitCode);
+						}
+					}
+
+					logger.LogDebug("Closing channel");
+					await lease.CloseAsync(cancellationToken);
+				}
+			}
+			catch (TimeoutException)
+			{
+				self.NumLogicalCores = 0;
+				UpdateHordeStatus(null);
+			}
+			catch (ComputeExecutionCancelledException ex)
+			{
+				// Cancellations are expected to happen due to spot instance interruptions or unscheduled maintenance of agents
+				self.NumLogicalCores = 0;
+				UpdateHordeStatus(null);
+
+				if (!cancellationToken.IsCancellationRequested)
+				{
+					logger.Log(_strict ? LogLevel.Information : LogLevel.Debug, KnownLogEvents.Systemic_Horde_Compute, ex, "Compute lease cancelled");
+				}
+			}
+			catch (Exception ex)
+			{
+				self.NumLogicalCores = 0;
+				UpdateHordeStatus(null);
+
+				if (!cancellationToken.IsCancellationRequested)
+				{
+					logger.Log(_strict ? LogLevel.Error : LogLevel.Debug, KnownLogEvents.Systemic_Horde_Compute, ex, "Exception in worker task: {Ex}", ex.ToString());
+
+					// Add additional properties to aid debugging
+					logger.Log(_strict ? LogLevel.Information : LogLevel.Debug, KnownLogEvents.Systemic_Horde_Compute, ex, "UBA agent locator {UBAAgentLocator}", _ubaAgentLocator.ToString());
+				}
+			}
+		}
+
+		string _additionalStatus = "";
+		string _lastStatus = "";
+
+		internal void UpdateHordeStatus(string? additionalStatus)
+		{
+			if (additionalStatus != null)
+			{
+				_additionalStatus = additionalStatus;
+			}
+
+			if (_updateStatus == null)
+			{
+				return;
+			}
+
+			int activeCores;
+			int queuedCores;
+			int activeAgents;
+			int queuedAgents;
+			GetCoreCount(out activeCores, out queuedCores, out activeAgents, out queuedAgents);
+			string status = $"Running. {activeAgents} agent{(activeAgents != 1 ? "s" : "")} ({activeCores} cores){_additionalStatus} {(queuedAgents != 0 ? $"(Preparing {queuedAgents} agent{(queuedAgents != 1 ? "s" : "")})" : "")}";
+			if (_lastStatus != status)
+			{
+				_lastStatus = status;
+				_updateStatus(0, 6, status, EpicGames.UBA.LogEntryType.Info, null);
+			}
+		}
+
+		static OSPlatform GetTargetWorkerPlatform(string? searchString)
+		{
+			if (!String.IsNullOrEmpty(searchString))
+			{
+				if (searchString.Contains("Linux", StringComparison.OrdinalIgnoreCase))
+				{
+					return OSPlatform.Linux;
+				}
+				if (searchString.Contains("Windows", StringComparison.OrdinalIgnoreCase))
+				{
+					return OSPlatform.Windows;
+				}
+				if (searchString.Contains("MacOS", StringComparison.OrdinalIgnoreCase))
+				{
+					return OSPlatform.OSX;
+				}
+			}
+
+			return UBAExecutor.GetHostPlatform();
+		}
+	}
+}
diff --git a/Engine/Source/Programs/UnrealBuildTool/Executors/UnrealBuildAccelerator/UnrealBuildAcceleratorHordeConfig.cs b/Engine/Source/Programs/UnrealBuildTool/Executors/UnrealBuildAccelerator/UnrealBuildAcceleratorHordeConfig.cs
--- a/Engine/Source/Programs/UnrealBuildTool/Executors/UnrealBuildAccelerator/UnrealBuildAcceleratorHordeConfig.cs
+++ b/Engine/Source/Programs/UnrealBuildTool/Executors/UnrealBuildAccelerator/UnrealBuildAcceleratorHordeConfig.cs
@@ -236,6 +236,10 @@ namespace UnrealBuildTool
 		[ConfigFile(ConfigHierarchyType.Engine, "Uba.Provider.Horde", "UbaCacheDesiredConnectionCount")]
 		[CommandLine("-UBAHordeCacheDesiredConnectionCount=")]
 		public int CacheDesiredConnectionCount { get; set; } = 1;
+
+		[ConfigFile(ConfigHierarchyType.Engine, "Uba.Provider.Horde", "UbaAgentPlatform")]
+		[XmlConfigFile(Category = "Horde", Name = "AgentPlatform")]
+		public string? AgentPlatform { get; set; }
 	}
 
 	#region Extension methods
