GameApplication initializations only for non-PSO processes

#android

#Synced-CL -1

diff --git a/Engine/Build/Android/Java/src/com/epicgames/unreal/GameApplication.java.template b/Engine/Build/Android/Java/src/com/epicgames/unreal/GameApplication.java.template
--- a/Engine/Build/Android/Java/src/com/epicgames/unreal/GameApplication.java.template
+++ b/Engine/Build/Android/Java/src/com/epicgames/unreal/GameApplication.java.template
@@ -39,7 +39,7 @@ public class GameApplication extends Application implements LifecycleObserver //
 	public void onCreate() {
 		super.onCreate();
 		GameApplication.context = getApplicationContext();
-//$${gameApplicationOnCreateAdditions}$$
+//$${gameApplicationOnCreateAdditionsAlways}$$
 
 		ProcessLifecycleOwner.get().getLifecycle().addObserver(this);
 
@@ -64,6 +64,7 @@ public class GameApplication extends Application implements LifecycleObserver //
 		
 		if (!processName.contains("psoprogramservice"))
 		{
+//$${gameApplicationOnCreateAdditions}$$
 			NetworkChangedManager.getInstance().initNetworkCallback(this);
 		}
 	}
diff --git a/Engine/Source/Programs/UnrealBuildTool/Platform/Android/UEDeployAndroid.cs b/Engine/Source/Programs/UnrealBuildTool/Platform/Android/UEDeployAndroid.cs
--- a/Engine/Source/Programs/UnrealBuildTool/Platform/Android/UEDeployAndroid.cs
+++ b/Engine/Source/Programs/UnrealBuildTool/Platform/Android/UEDeployAndroid.cs
@@ -6246,6 +6246,7 @@ import java.util.Collection;
 				{ "//$${gameApplicationImportAdditions}$$", UPL!.ProcessPluginNode(NDKArch, "gameApplicationImportAdditions", GameApplicationImportAdditionsDefault)},
 				{ "//$${gameApplicationImplementsAdditions}$$", UPL!.ProcessPluginNode(NDKArch, "gameApplicationImplementsAdditions", "")},
 				{ "//$${gameApplicationOnCreateAdditions}$$", UPL.ProcessPluginNode(NDKArch, "gameApplicationOnCreateAdditions", "")},
+				{ "//$${gameApplicationOnCreateAdditionsAlways}$$", UPL.ProcessPluginNode(NDKArch, "gameApplicationOnCreateAdditionsAlways", "")},
 				{ "//$${gameApplicationAttachBaseContextAdditions}$$", UPL.ProcessPluginNode(NDKArch, "gameApplicationAttachBaseContextAdditions", "")},
 				{ "//$${gameApplicationOnLowMemoryAdditions}$$", UPL.ProcessPluginNode(NDKArch, "gameApplicationOnLowMemoryAdditions", "")},
 				{ "//$${gameApplicationOnTrimMemoryAdditions}$$", UPL.ProcessPluginNode(NDKArch, "gameApplicationOnTrimMemoryAdditions", "")},
diff --git a/Engine/Source/Programs/UnrealBuildTool/System/UnrealPluginLanguage.cs b/Engine/Source/Programs/UnrealBuildTool/System/UnrealPluginLanguage.cs
--- a/Engine/Source/Programs/UnrealBuildTool/System/UnrealPluginLanguage.cs
+++ b/Engine/Source/Programs/UnrealBuildTool/System/UnrealPluginLanguage.cs
@@ -437,6 +437,12 @@ namespace UnrealBuildTool
 	 * 	<!-- optional additional interfaces that GameApplication can implement in GameApplication.java (must start with a comma and list interfaces using commas) -->
 	 * 	<gameApplicationImplementsAdditions> </gameApplicationImplementsAdditions>
 	 * 	
+	 *	<!-- optional additions to GameApplication onCreate in GameApplication.java (all processes) -->
+	 *	<gameApplicationOnCreateAdditionsAlways> </gameApplicationOnCreateAdditionsAlways>
+	 *  
+	 *	<!-- optional additions to GameApplication onCreate in GameApplication.java (non-PSO processes) -->
+	 *	<gameApplicationOnCreateAdditions> </gameApplicationOnCreateAdditions>
+	 *  
 	 * 	<!-- optional libraries to load in GameActivity.java before libUnreal.so -->
 	 * 	<soLoadLibrary>	</soLoadLibrary>
 	 * 	
