UnrealBuildTool: Add ini config for bUsePortableToolchain

#rnx

#Synced-CL 50271719

diff --git a/Engine/Source/Programs/UnrealBuildTool/Platform/Apple/UEBuildApple.cs b/Engine/Source/Programs/UnrealBuildTool/Platform/Apple/UEBuildApple.cs
--- a/Engine/Source/Programs/UnrealBuildTool/Platform/Apple/UEBuildApple.cs
+++ b/Engine/Source/Programs/UnrealBuildTool/Platform/Apple/UEBuildApple.cs
@@ -70,6 +70,7 @@ namespace UnrealBuildTool
 		/// </summary>
 		[CommandLine("-UsePortableToolchain")]
 		[XmlConfigFile(Category = "BuildConfiguration")]
+		[ConfigFile(ConfigHierarchyType.Game, "/Script/AppleTargetPlatform.AppleTargetSettings", "bUsePortableToolchain")]
 		public bool bUsePortableToolchain = false;
 	}
 
