[UBA]
* Fixed handling of ubaexecutor init event cancelled

#Synced-CL 50276738

diff --git a/Engine/Source/Programs/UnrealBuildTool/Executors/UnrealBuildAccelerator/UBAAgentCoordinatorHorde.cs b/Engine/Source/Programs/UnrealBuildTool/Executors/UnrealBuildAccelerator/UBAAgentCoordinatorHorde.cs
--- a/Engine/Source/Programs/UnrealBuildTool/Executors/UnrealBuildAccelerator/UBAAgentCoordinatorHorde.cs
+++ b/Engine/Source/Programs/UnrealBuildTool/Executors/UnrealBuildAccelerator/UBAAgentCoordinatorHorde.cs
@@ -208,8 +208,11 @@ namespace UnrealBuildTool
 					continue;
 				}
 
-				_initDoneEvent.Wait(cancellationToken);
-				if (!_initDoneEvent.IsSet)
+				try
+				{
+					_initDoneEvent.Wait(cancellationToken);
+				}
+				catch (OperationCanceledException)
 				{
 					return null;
 				}
