UnrealBuildTool: Don't assume that there's any valid Horde configs for running with UBA.

#Synced-CL 50295728

diff --git a/Engine/Source/Programs/UnrealBuildTool/Executors/UnrealBuildAccelerator/UBAAgentCoordinatorHorde.cs b/Engine/Source/Programs/UnrealBuildTool/Executors/UnrealBuildAccelerator/UBAAgentCoordinatorHorde.cs
--- a/Engine/Source/Programs/UnrealBuildTool/Executors/UnrealBuildAccelerator/UBAAgentCoordinatorHorde.cs
+++ b/Engine/Source/Programs/UnrealBuildTool/Executors/UnrealBuildAccelerator/UBAAgentCoordinatorHorde.cs
@@ -81,19 +81,23 @@ namespace UnrealBuildTool
 
 			// We group configs based on hordeserver/token to be able to have only one connection for configs having matching server/token.
 
-			IEnumerable<IGrouping<string, UnrealBuildAcceleratorHordeConfig>> configGroups = providers
+			List<IGrouping<string, UnrealBuildAcceleratorHordeConfig>> configGroups = [.. providers
 				.SelectMany(provider => targetDescriptors.Select(targetDescriptor => CreateHordeConfig(provider, targetDescriptor.AdditionalArguments, targetDescriptor.ProjectFile?.Directory)))
 				.Where(x => !x.bDisableHorde && !String.IsNullOrWhiteSpace(x.HordeServer))
 				.DistinctBy(x => HashCode.Combine(x.HordeServer, x.HordeCondition, x.HordeCluster, x.HordePool, x.HordeConnectionMode, x.AgentPlatform))
-				.GroupBy(x => String.Join(";", x.HordeServer, x.HordeToken));
+				.GroupBy(x => String.Join(";", x.HordeServer, x.HordeToken))];
 
-			List<UBAAgentCoordinatorHorde> coordinators = new();
+			List<UBAAgentCoordinatorHorde> coordinators = [];
 			foreach (IGrouping<string, UnrealBuildAcceleratorHordeConfig> configGroup in configGroups)
 			{
 				coordinators.Add(new UBAAgentCoordinatorHorde(logger, ubaConfig, configGroup));
 			}
 
-			remoteConnectionMode = (Enum.TryParse(configGroups.First().First().HordeConnectionMode, true, out ConnectionMode cm) ? cm : ConnectionMode.Direct).ToString();
+			remoteConnectionMode = ConnectionMode.Direct.ToString();
+			if (configGroups.Count > 0)
+			{
+				remoteConnectionMode = (Enum.TryParse(configGroups[0].First().HordeConnectionMode, true, out ConnectionMode cm) ? cm : ConnectionMode.Direct).ToString();
+			}
 			return coordinators;
 		}
 
