Removed unnecessary dependencies from LauncherServices, RenderCore, ShaderCompilerCommon, and TargetPlatform

Reduces the size of programs like ShaderCompileWorker that do not otherwise depend on some of the removed dependencies.

#Synced-CL 50253651

diff --git a/Engine/Source/Developer/LauncherServices/LauncherServices.Build.cs b/Engine/Source/Developer/LauncherServices/LauncherServices.Build.cs
--- a/Engine/Source/Developer/LauncherServices/LauncherServices.Build.cs
+++ b/Engine/Source/Developer/LauncherServices/LauncherServices.Build.cs
@@ -21,7 +21,6 @@ namespace UnrealBuildTool.Rules
 					"CoreUObject",
 					"DesktopPlatform",
 					"SessionMessages",
-					"SourceCodeAccess",
 					"TargetPlatform",
 					"UnrealEdMessages",
 					"Json",
diff --git a/Engine/Source/Developer/LauncherServices/Private/Launcher/LauncherWorker.cpp b/Engine/Source/Developer/LauncherServices/Private/Launcher/LauncherWorker.cpp
--- a/Engine/Source/Developer/LauncherServices/Private/Launcher/LauncherWorker.cpp
+++ b/Engine/Source/Developer/LauncherServices/Private/Launcher/LauncherWorker.cpp
@@ -4,8 +4,6 @@
 #include "HAL/PlatformTime.h"
 #include "HAL/FileManager.h"
 #include "ILauncherServicesModule.h"
-#include "ISourceCodeAccessor.h"
-#include "ISourceCodeAccessModule.h"
 #include "ITargetDeviceProxy.h"
 #include "ITargetDeviceProxyManager.h"
 #include "ITurnkeyIOModule.h"
diff --git a/Engine/Source/Developer/ShaderCompilerCommon/ShaderCompilerCommon.Build.cs b/Engine/Source/Developer/ShaderCompilerCommon/ShaderCompilerCommon.Build.cs
--- a/Engine/Source/Developer/ShaderCompilerCommon/ShaderCompilerCommon.Build.cs
+++ b/Engine/Source/Developer/ShaderCompilerCommon/ShaderCompilerCommon.Build.cs
@@ -14,7 +14,6 @@ public class ShaderCompilerCommon : ModuleRules
 				"Core",
 				"RenderCore",
 				"ShaderPreprocessor",
-				"TargetPlatform",
 			}
 			);
 		PublicDependencyModuleNames.AddRange(
@@ -22,6 +21,11 @@ public class ShaderCompilerCommon : ModuleRules
 				"HlslParser",
 			}
 			);
+		PublicIncludePathModuleNames.AddRange(
+			new string[] {
+				"TargetPlatform",
+			}
+			);
 		if (Target.bCompileAgainstEngine)
 		{
 			PublicDependencyModuleNames.AddRange(
diff --git a/Engine/Source/Programs/UnrealBuildTool/Platform/Android/UEBuildAndroid.cs b/Engine/Source/Programs/UnrealBuildTool/Platform/Android/UEBuildAndroid.cs
--- a/Engine/Source/Programs/UnrealBuildTool/Platform/Android/UEBuildAndroid.cs
+++ b/Engine/Source/Programs/UnrealBuildTool/Platform/Android/UEBuildAndroid.cs
@@ -343,28 +343,21 @@ namespace UnrealBuildTool
 			if ((Target.Platform == UnrealTargetPlatform.Win64) || (Target.Platform == UnrealTargetPlatform.Mac) || (Target.Platform == UnrealTargetPlatform.Linux))
 			{
 				bool bBuildShaderFormats = Target.bForceBuildShaderFormats;
-				if (!Target.bBuildRequiresCookedData)
+				if (!Target.bBuildRequiresCookedData && Target.bBuildDeveloperTools)
 				{
 					if (ModuleName == "Engine")
 					{
-						if (Target.bBuildDeveloperTools)
-						{
-							Rules.DynamicallyLoadedModuleNames.Add("AndroidTargetPlatform");
-							Rules.DynamicallyLoadedModuleNames.Add("AndroidTargetPlatformSettings");
-							Rules.DynamicallyLoadedModuleNames.Add("AndroidTargetPlatformControls");
-						}
+						Rules.DynamicallyLoadedModuleNames.Add("AndroidTargetPlatform");
+						Rules.DynamicallyLoadedModuleNames.Add("AndroidTargetPlatformSettings");
+						Rules.DynamicallyLoadedModuleNames.Add("AndroidTargetPlatformControls");
 					}
 					else if (ModuleName == "TargetPlatform")
 					{
 						bBuildShaderFormats = true;
-
-						if (Target.bBuildDeveloperTools)
-						{
-							Rules.DynamicallyLoadedModuleNames.Add("TextureFormatOodle");
-							Rules.DynamicallyLoadedModuleNames.Add("TextureFormatASTC");
-							Rules.DynamicallyLoadedModuleNames.Add("TextureFormatETC2");  // ETC2 
-							//Rules.DynamicallyLoadedModuleNames.Add("AudioFormatADPCM");	//@todo android: android audio
-						}
+						Rules.DynamicallyLoadedModuleNames.Add("TextureFormatOodle");
+						Rules.DynamicallyLoadedModuleNames.Add("TextureFormatASTC");
+						Rules.DynamicallyLoadedModuleNames.Add("TextureFormatETC2");  // ETC2 
+						//Rules.DynamicallyLoadedModuleNames.Add("AudioFormatADPCM");	//@todo android: android audio
 					}
 				}
 
diff --git a/Engine/Source/Programs/UnrealBuildTool/Platform/IOS/UEBuildIOS.cs b/Engine/Source/Programs/UnrealBuildTool/Platform/IOS/UEBuildIOS.cs
--- a/Engine/Source/Programs/UnrealBuildTool/Platform/IOS/UEBuildIOS.cs
+++ b/Engine/Source/Programs/UnrealBuildTool/Platform/IOS/UEBuildIOS.cs
@@ -954,22 +954,19 @@ namespace UnrealBuildTool
 			{
 				// For Windows we use the MetalDeveloperTools and we will check if the toolchain is present at runtime
 				bool bBuildShaderFormats = Target.bForceBuildShaderFormats || Target.Platform == UnrealTargetPlatform.Win64;
-				if (!Target.bBuildRequiresCookedData)
+				if (!Target.bBuildRequiresCookedData && Target.bBuildDeveloperTools)
 				{
 					if (ModuleName == "Engine")
 					{
-						if (Target.bBuildDeveloperTools)
-						{
-							Rules.DynamicallyLoadedModuleNames.Add("IOSTargetPlatformSettings");
-							Rules.DynamicallyLoadedModuleNames.Add("TVOSTargetPlatformSettings");
+						Rules.DynamicallyLoadedModuleNames.Add("IOSTargetPlatformSettings");
+						Rules.DynamicallyLoadedModuleNames.Add("TVOSTargetPlatformSettings");
 
-							if (bIsPlatformAvailableForTargetWithSDK)
-							{
-								Rules.DynamicallyLoadedModuleNames.Add("IOSTargetPlatform");
-								Rules.DynamicallyLoadedModuleNames.Add("IOSTargetPlatformControls");
-								Rules.DynamicallyLoadedModuleNames.Add("TVOSTargetPlatform");
-								Rules.DynamicallyLoadedModuleNames.Add("TVOSTargetPlatformControls");
-							}
+						if (bIsPlatformAvailableForTargetWithSDK)
+						{
+							Rules.DynamicallyLoadedModuleNames.Add("IOSTargetPlatform");
+							Rules.DynamicallyLoadedModuleNames.Add("IOSTargetPlatformControls");
+							Rules.DynamicallyLoadedModuleNames.Add("TVOSTargetPlatform");
+							Rules.DynamicallyLoadedModuleNames.Add("TVOSTargetPlatformControls");
 						}
 					}
 					else if (ModuleName == "TargetPlatform")
@@ -979,7 +976,7 @@ namespace UnrealBuildTool
 							bBuildShaderFormats = true;
 							Rules.DynamicallyLoadedModuleNames.Add("TextureFormatASTC");
 							Rules.DynamicallyLoadedModuleNames.Add("TextureFormatETC2");
-							if (Target.bBuildDeveloperTools && Target.bCompileAgainstEngine)
+							if (Target.bCompileAgainstEngine)
 							{
 								Rules.DynamicallyLoadedModuleNames.Add("AudioFormatADPCM");
 							}
diff --git a/Engine/Source/Runtime/RenderCore/Public/ShaderCodeArchive.h b/Engine/Source/Runtime/RenderCore/Public/ShaderCodeArchive.h
--- a/Engine/Source/Runtime/RenderCore/Public/ShaderCodeArchive.h
+++ b/Engine/Source/Runtime/RenderCore/Public/ShaderCodeArchive.h
@@ -42,7 +42,7 @@ class FCbWriter;
 #endif
 
 // enable visualization in the desktop Development builds only as it has a memory hit and writes files
-#define UE_SCA_VISUALIZE_SHADER_USAGE			(!WITH_EDITOR && UE_BUILD_DEVELOPMENT && PLATFORM_DESKTOP)
+#define UE_SCA_VISUALIZE_SHADER_USAGE			(!WITH_EDITOR && !IS_PROGRAM && UE_BUILD_DEVELOPMENT && PLATFORM_DESKTOP)
 
 // enable deep, manual debugging of leaked preload groups. This level of information slows the engine down and is only needed when chasing tricky bugs
 #define UE_SCA_DEBUG_PRELOADING					(0)
diff --git a/Engine/Source/Runtime/RenderCore/RenderCore.Build.cs b/Engine/Source/Runtime/RenderCore/RenderCore.Build.cs
--- a/Engine/Source/Runtime/RenderCore/RenderCore.Build.cs
+++ b/Engine/Source/Runtime/RenderCore/RenderCore.Build.cs
@@ -19,13 +19,13 @@ public class RenderCore : ModuleRules
 		PrivateDependencyModuleNames.Add("BuildSettings");
 
 		if (Target.bBuildEditor == true)
-        {
+		{
 			DynamicallyLoadedModuleNames.Add("TargetPlatform");
 			PrivateIncludePathModuleNames.Add("IoStoreUtilities");
 		}
 		// shader runtime usage visualization requires ability to create images - it is only used in non-editor desktop development targets
 		// UE_BUILD_DEVELOPMENT is also defined for DebugGame
-		else if ((Target.Configuration == UnrealTargetConfiguration.Development || Target.Configuration == UnrealTargetConfiguration.DebugGame) && Array.IndexOf(Utils.GetPlatformsInClass(UnrealPlatformClass.Desktop), Target.Platform) >= 0)
+		else if (Target.Type != TargetType.Program && (Target.Configuration == UnrealTargetConfiguration.Development || Target.Configuration == UnrealTargetConfiguration.DebugGame) && Array.IndexOf(Utils.GetPlatformsInClass(UnrealPlatformClass.Desktop), Target.Platform) >= 0)
 		{
 			PrivateDependencyModuleNames.Add("ImageWrapper");
 		}
@@ -36,7 +36,14 @@ public class RenderCore : ModuleRules
 			RuntimeDependencies.Add(Path.Combine(Unreal.EngineDirectory.ToString(), "Extras/GPUDumpViewer/..."), StagedFileType.DebugNonUFS);
 		}
 
-		PrivateDependencyModuleNames.AddRange(new string[] { "Core", "Projects", "ApplicationCore", "TraceLog", "CookOnTheFly", "StateStream" });
+		PrivateDependencyModuleNames.AddRange(new string[]
+		{
+			"Core",
+			"Projects",
+			"ApplicationCore",
+			"TraceLog",
+			"StateStream",
+		});
 		
 		PublicIncludePathModuleNames.AddRange(new string[] { "RHI" });
 
