UnrealBuildTool: Partial backout of 50157090 for WriteMetadataMode, as this code is required to execute synchronously in order to run within the global mutex.

#Synced-CL 50348122

diff --git a/Engine/Source/Programs/UnrealBuildTool/Modes/WriteMetadataMode.cs b/Engine/Source/Programs/UnrealBuildTool/Modes/WriteMetadataMode.cs
--- a/Engine/Source/Programs/UnrealBuildTool/Modes/WriteMetadataMode.cs
+++ b/Engine/Source/Programs/UnrealBuildTool/Modes/WriteMetadataMode.cs
@@ -10,6 +10,8 @@ using EpicGames.Core;
 using Microsoft.Extensions.Logging;
 using UnrealBuildBase;
 
+#pragma warning disable CA1849
+
 namespace UnrealBuildTool
 {
 	/// <summary>
@@ -75,14 +77,14 @@ namespace UnrealBuildTool
 		/// <param name="Arguments">Command line arguments</param>
 		/// <param name="Logger">Logger for output</param>
 		/// <returns>Exit code</returns>
-		private static async Task<int> ExecuteInternal(CommandLineArguments Arguments, ILogger Logger)
+		private static Task<int> ExecuteInternal(CommandLineArguments Arguments, ILogger Logger)
 		{
 			// Read the target info
 			FileReference FilePath = Arguments.GetFileReference("-Input=");
 			WriteMetadataTargetInfo TargetInfo;
 			using (FileStream FileStream = FileReference.Open(FilePath, FileMode.Open, FileAccess.Read, FileShare.Read))
 			{
-				WriteMetadataTargetInfo? MaybeTargetInfo = await JsonSerializer.DeserializeAsync<WriteMetadataTargetInfo>(FileStream, SerializerOptions);
+				WriteMetadataTargetInfo? MaybeTargetInfo = JsonSerializer.Deserialize<WriteMetadataTargetInfo>(FileStream, SerializerOptions);
 				TargetInfo = MaybeTargetInfo ?? throw new BuildLogEventException("Target info file '{FilePath}' was invalid", FilePath);
 			}
 
@@ -155,7 +157,7 @@ namespace UnrealBuildTool
 						// Check if the manifest has changed. Note that if a manifest is out of date, we should have generated a new build id causing the contents to differ.
 						if (bNoManifestChanges)
 						{
-							string CurrentText = await FileReference.ReadAllTextAsync(ManifestFile);
+							string CurrentText = FileReference.ReadAllText(ManifestFile);
 							if (CurrentText != OutputText)
 							{
 								Logger.LogError(
@@ -167,7 +169,7 @@ namespace UnrealBuildTool
 						}
 
 						// Write it to disk
-						await FileReference.WriteAllTextAsync(ManifestFile, OutputText);
+						FileReference.WriteAllText(ManifestFile, OutputText);
 					}
 				}
 			}
@@ -199,7 +201,7 @@ namespace UnrealBuildTool
 						// Check if the manifest has changed. Note that if a manifest is out of date, we should have generated a new build id causing the contents to differ.
 						if (bNoManifestChanges)
 						{
-							string CurrentText = await FileReference.ReadAllTextAsync(ManifestFile);
+							string CurrentText = FileReference.ReadAllText(ManifestFile);
 							if (CurrentText != OutputText)
 							{
 								Logger.LogError(
@@ -211,7 +213,7 @@ namespace UnrealBuildTool
 						}
 
 						// Write it to disk
-						await FileReference.WriteAllTextAsync(ManifestFile, OutputText);
+						FileReference.WriteAllText(ManifestFile, OutputText);
 					}
 				}
 			}
@@ -232,7 +234,7 @@ namespace UnrealBuildTool
 				TargetInfo.Receipt.Write(TargetInfo.ReceiptFile);
 			}
 
-			return 0;
+			return Task.FromResult(0);
 		}
 
 		/// <summary>
