UnrealBuildTool: Optimizations to cs project evalulation
* Replace regex usage with RecursiveDir as VS Mac is no longer supported
* Don't process non-compiled files only added to be shown in the IDE when building

#rnx

#Synced-CL 50271719

diff --git a/Engine/Source/Programs/UnrealBuildTool/UnrealBuildTool.csproj b/Engine/Source/Programs/UnrealBuildTool/UnrealBuildTool.csproj
--- a/Engine/Source/Programs/UnrealBuildTool/UnrealBuildTool.csproj
+++ b/Engine/Source/Programs/UnrealBuildTool/UnrealBuildTool.csproj
@@ -4,12 +4,10 @@
   <PropertyGroup>
     <TargetFramework>net8.0</TargetFramework>
     <OutputType>Exe</OutputType>
-    <AppendTargetFrameworkToOutputPath>false</AppendTargetFrameworkToOutputPath>
     <GenerateAssemblyInfo>false</GenerateAssemblyInfo>
     <GenerateTargetFrameworkAttribute>false</GenerateTargetFrameworkAttribute>
     <ApplicationManifest>app.manifest</ApplicationManifest>
-    <SatelliteResourceLanguages>en</SatelliteResourceLanguages> <!-- remove non english resource languages -->
-    <AppendRuntimeIdentifierToOutputPath>false</AppendRuntimeIdentifierToOutputPath>
+    <SatelliteResourceLanguages>en</SatelliteResourceLanguages>
     <Nullable>enable</Nullable>
     <GenerateDocumentationFile>True</GenerateDocumentationFile>
   </PropertyGroup>
@@ -20,51 +18,25 @@
     </Reference>
   </ItemGroup>
 
-  <Choose>
-    <When Condition="Exists('..\..\..\Restricted\NotForLicensees\Source\Programs\UnrealBuildTool\ProjectFiles\VProject\VProjectFileGenerator.cs')">
-      <PropertyGroup>
-        <DefineConstants>$(DefineConstants);__VPROJECT_AVAILABLE__</DefineConstants>
-      </PropertyGroup>
-    </When>
-  </Choose>
+  <PropertyGroup Condition="Exists('..\..\..\Restricted\NotForLicensees\Source\Programs\UnrealBuildTool\ProjectFiles\VProject\VProjectFileGenerator.cs')">
+    <DefineConstants>$(DefineConstants);__VPROJECT_AVAILABLE__</DefineConstants>
+  </PropertyGroup>
 
   <ItemGroup>
-    <Compile Include="..\Shared\MetaData.cs">
-      <Link>Properties\MetaData.cs</Link>
-    </Compile>
+    <Compile Include="..\Shared\MetaData.cs" Link="Properties\MetaData.cs" />
   </ItemGroup>
 
-  <!-- The following will glob all platform extensions' UnrealBuildTool source
-  files and include them in the build. They will also appear in VisualStudio.
-  RecursiveDir does not work in VS Mac (and maybe Linux?), so only use it in Windows. 
-  See: https://developercommunity.visualstudio.com/content/problem/527478/recursivedir-causes-infinitely-nested-subdirectori.html -->
+  <!-- The following will glob all platform extensions' UnrealBuildTool source files and include them in the build. They will also appear in VisualStudio. -->
   <ItemGroup>
-    <Compile Include="../../../Platforms/*/Source/Programs/UnrealBuildTool/**/*.cs">
-      <!-- Since RecursiveDir is bugged on VS Mac, we use the FullPath, and use Regex to pull out portions of the path to get Platform name -->
-      <Link>Platform/$([System.Text.RegularExpressions.Regex]::Replace(%(FullPath), '^(.*[\\/]Platforms[\\/])(.+?)[\\/](.*)$', '$2'))/%(Filename)%(Extension)</Link>
-    </Compile>
-    <Compile Include="../../../Restricted/*/Platforms/*/Source/Programs/UnrealBuildTool/**/*.cs">
-      <!-- Restricted may or may not be in a Platforms dir, so just use the subdir directly, even tho it makes for a
-       messier directory structure for these few files -->
-      <Link>Platform/$([System.Text.RegularExpressions.Regex]::Replace(%(FullPath), '^(.*[\\/]Platforms[\\/])(.+?)[\\/](.*)$', '$2'))/%(Filename)%(Extension)</Link>
-    </Compile>
-    <Compile Include="../../../Restricted/*/Source/Programs/UnrealBuildTool/**/*.cs">
-      <!-- Restricted may or may not be in a Platforms dir, so just use the subdir directly, even tho it makes for a
-       messier directory structure for these few files -->
-      <Link>$([System.Text.RegularExpressions.Regex]::Replace(%(FullPath), '^.+?[\\/]UnrealBuildTool[\\/]', ''))</Link>
-    </Compile>
-
-    <!-- SDK config files -->
-    <None Include="../../../Platforms/*/Config/*_SDK.json">
-      <!-- Since RecursiveDir is bugged on VS Mac, we use the FullPath, and use Regex to pull out portions of the path to get Platform name -->
-      <Link>Platform/$([System.Text.RegularExpressions.Regex]::Replace(%(FullPath), '^(.*[\\/]Platforms[\\/])(.+?)[\\/](.*)$', '$2'))/%(Filename)%(Extension)</Link>
-    </None>
-    <None Include="../../../Config/*/*_SDK.json">
-      <!-- Since RecursiveDir is bugged on VS Mac, we use the FullPath, and use Regex to pull out portions of the path to get Platform name -->
-      <Link>Platform/$([System.Text.RegularExpressions.Regex]::Replace(%(FullPath), '^(.*[\\/]Config[\\/])(.+?)[\\/](.*)$', '$2'))/%(Filename)%(Extension)</Link>
-    </None>
+    <Compile Include="../../../Platforms/*/Source/Programs/UnrealBuildTool/**/*.cs"
+             Link="Platform/$([System.String]::Copy('%(RecursiveDir)').Replace('\', '/').Replace('Source/Programs/UnrealBuildTool/', ''))%(Filename)%(Extension)" />
+    <Compile Include="../../../Restricted/*/Platforms/*/Source/Programs/UnrealBuildTool/**/*.cs"
+             Link="Platform/$([System.String]::Copy('%(RecursiveDir)').Replace('\', '/').Replace('Source/Programs/UnrealBuildTool/', ''))%(Filename)%(Extension)" />
+    <Compile Include="../../../Restricted/*/Source/Programs/UnrealBuildTool/**/*.cs"
+             Link="Restricted/$([System.String]::Copy('%(RecursiveDir)').Replace('\', '/').Replace('Source/Programs/UnrealBuildTool/', ''))%(Filename)%(Extension)" />
   </ItemGroup>
 
+
   <ItemGroup>
     <!-- If you add a project reference, BuildUAT.bat/.sh and BuildUBT.bat/.sh need to be updated to also test for those projects.-->
     <ProjectReference Include="..\Shared\EpicGames.Build\EpicGames.Build.csproj" />
@@ -94,21 +66,23 @@
     <PackageReference Include="System.Text.Encoding.CodePages" Version="8.0.0" />
   </ItemGroup>
 
-  <!-- Various batch files and shell scripts that build or run UnrealBuildTool -->
-  <ItemGroup>
-    <None Include="../../../Build/BatchFiles/**/*">
-      <Link>$([System.Text.RegularExpressions.Regex]::Replace(%(FullPath), '^.+?[\\/]BatchFiles[\\/]', 'BatchFiles/'))</Link>
-    </None>
-  </ItemGroup>
+  <!-- Non-compiled files for IDE display only. -->
+  <ItemGroup Condition="'$(DesignTimeBuild)' == 'true'">
+    <!-- SDK config files -->
+    <None Include="../../../Platforms/*/Config/*_SDK.json" Link="Platform/%(RecursiveDir)%(Filename)%(Extension)" />
+    <None Include="../../../Config/*/*_SDK.json" Link="Platform/%(RecursiveDir)%(Filename)%(Extension)" />
 
-  <!-- Various BuildConfiguration files -->
-  <ItemGroup>
-      <None Condition="Exists('..\..\..\Restricted\NotForLicensees\Programs\UnrealBuildTool\BuildConfiguration.xml')" Include="..\..\..\Restricted\NotForLicensees\Programs\UnrealBuildTool\BuildConfiguration.xml"><Link>BuildConfiguration\BuildConfiguration.xml (Engine NotForLicensees)</Link></None>
-      <None Condition="Exists('..\..\..\Saved\UnrealBuildTool\BuildConfiguration.xml')" Include="..\..\..\Saved\UnrealBuildTool\BuildConfiguration.xml"><Link>BuildConfiguration\BuildConfiguration.xml (Engine Saved)</Link></None>
-      <None Condition="Exists('$(ProgramData)\Unreal Engine\UnrealBuildTool\BuildConfiguration.xml')" Include="$(ProgramData)\Unreal Engine\UnrealBuildTool\BuildConfiguration.xml"><Link>BuildConfiguration\BuildConfiguration.xml (Global ProgramData)</Link></None>
-      <None Condition="Exists('$(AppData)\Unreal Engine\UnrealBuildTool\BuildConfiguration.xml')" Include="$(AppData)\Unreal Engine\UnrealBuildTool\BuildConfiguration.xml"><Link>BuildConfiguration\BuildConfiguration.xml (Global AppData)</Link></None>
-      <None Condition="Exists('$(LocalAppData)\Unreal Engine\UnrealBuildTool\BuildConfiguration.xml')" Include="$(LocalAppData)\Unreal Engine\UnrealBuildTool\BuildConfiguration.xml"><Link>BuildConfiguration\BuildConfiguration.xml (Global LocalAppData)</Link></None>
-      <None Condition="Exists('$(UserProfile)\Documents\Unreal Engine\UnrealBuildTool\BuildConfiguration.xml')" Include="$(UserProfile)\Documents\Unreal Engine\UnrealBuildTool\BuildConfiguration.xml"><Link>BuildConfiguration\BuildConfiguration.xml (Global Documents)</Link></None>
-      <None Condition="Exists('..\..\..\Intermediate\Build\UnrealBuildTool.Env.BuildConfiguration.xml')" Include="..\..\..\Intermediate\Build\UnrealBuildTool.Env.BuildConfiguration.xml"><Link>BuildConfiguration\BuildConfiguration.xml (Environment)</Link></None>
+    <!-- Various batch files and shell scripts that build or run UnrealBuildTool -->
+    <None Include="../../../Build/BatchFiles/**/*" Link="BatchFiles/%(RecursiveDir)%(Filename)%(Extension)" />
+
+    <!-- Various BuildConfiguration files -->
+    <None Condition="Exists('..\..\..\Restricted\NotForLicensees\Programs\UnrealBuildTool\BuildConfiguration.xml')" Include="..\..\..\Restricted\NotForLicensees\Programs\UnrealBuildTool\BuildConfiguration.xml" Link="BuildConfiguration\BuildConfiguration.xml (Engine NotForLicensees)" />
+    <None Condition="Exists('..\..\..\Saved\UnrealBuildTool\BuildConfiguration.xml')" Include="..\..\..\Saved\UnrealBuildTool\BuildConfiguration.xml" Link="BuildConfiguration\BuildConfiguration.xml (Engine Saved)" />
+    <None Condition="Exists('$(ProgramData)\Unreal Engine\UnrealBuildTool\BuildConfiguration.xml')" Include="$(ProgramData)\Unreal Engine\UnrealBuildTool\BuildConfiguration.xml" Link="BuildConfiguration\BuildConfiguration.xml (Global ProgramData)" />
+    <None Condition="Exists('$(AppData)\Unreal Engine\UnrealBuildTool\BuildConfiguration.xml')" Include="$(AppData)\Unreal Engine\UnrealBuildTool\BuildConfiguration.xml" Link="BuildConfiguration\BuildConfiguration.xml (Global AppData)" />
+    <None Condition="Exists('$(LocalAppData)\Unreal Engine\UnrealBuildTool\BuildConfiguration.xml')" Include="$(LocalAppData)\Unreal Engine\UnrealBuildTool\BuildConfiguration.xml" Link="BuildConfiguration\BuildConfiguration.xml (Global LocalAppData)" />
+    <None Condition="Exists('$(UserProfile)\Documents\Unreal Engine\UnrealBuildTool\BuildConfiguration.xml')" Include="$(UserProfile)\Documents\Unreal Engine\UnrealBuildTool\BuildConfiguration.xml" Link="BuildConfiguration\BuildConfiguration.xml (Global Documents)" />
+    <None Condition="Exists('..\..\..\Intermediate\Build\UnrealBuildTool.Env.BuildConfiguration.xml')" Include="..\..\..\Intermediate\Build\UnrealBuildTool.Env.BuildConfiguration.xml" Link="BuildConfiguration\BuildConfiguration.xml (Environment)" />
   </ItemGroup>
+
 </Project>
