Adding an Apple Certificate checker to UBT so we can check that designated certificates are being used to sign Apple platform builds

#Synced-CL 50166977

diff --git a/Engine/Source/Developer/Mac/MacTargetPlatformSettings/Classes/XcodeProjectSettings.h b/Engine/Source/Developer/Mac/MacTargetPlatformSettings/Classes/XcodeProjectSettings.h
--- a/Engine/Source/Developer/Mac/MacTargetPlatformSettings/Classes/XcodeProjectSettings.h
+++ b/Engine/Source/Developer/Mac/MacTargetPlatformSettings/Classes/XcodeProjectSettings.h
@@ -168,6 +168,24 @@ public:
 	UPROPERTY(EditAnywhere, config, Category="Code Signing", meta = (EditCondition="bUseAutomaticCodeSigning && bUseAppStoreConnect", DisplayName = "App Store Connect: Key File"))
 	FFilePath AppStoreConnectKeyPath;
 	
+	/**
+	 * SHA1 hashes of certificates to verify. If these certs are not present, the build will either try to install them or fail.
+	 */
+	UPROPERTY(EditAnywhere, config, Category="Code Signing")
+	TArray<FString> CertificateHashes;
+	
+	/**
+	 * Perforce file paths to certificate files to verify. If these certs are not present, the build will either try to install them or fail.
+	 */
+	UPROPERTY(EditAnywhere, config, Category="Code Signing")
+	TArray<FString> P4CertificateFilePaths;
+	
+	/**
+	 * Perforce stream names containing the certificate files to verify. If these certs are not present, the build will either try to install them or fail.
+	 */
+	UPROPERTY(EditAnywhere, config, Category="Code Signing")
+	TArray<FString> CertificateStreamNames;
+	
 	/**
 	 * The path to the optional PrivacyInfo.xcprivacy file for your project
 	 */
diff --git a/Engine/Source/Programs/UnrealBuildTool/Platform/Apple/AppleCertificateVerification.cs b/Engine/Source/Programs/UnrealBuildTool/Platform/Apple/AppleCertificateVerification.cs
new file mode 100644
--- /dev/null
+++ b/Engine/Source/Programs/UnrealBuildTool/Platform/Apple/AppleCertificateVerification.cs
@@ -0,0 +1,193 @@
+// Copyright Epic Games, Inc. All Rights Reserved.
+
+using System;
+using System.Collections.Generic;
+using System.Diagnostics;
+using System.IO;
+using System.Security.Cryptography.X509Certificates;
+using System.Text;
+using EpicGames.Core;
+using Microsoft.Extensions.Logging;
+using UnrealBuildBase;
+
+namespace UnrealBuildTool
+{
+	/// <summary>
+	/// Utility class for verifying that you have certain Apple Certificates installed from Perforce
+	/// </summary>
+	public partial class AppleCertificateVerification
+	{
+		private static bool bCertificatesVerified = false;
+
+		[ConfigFile(ConfigHierarchyType.Engine, "/Script/MacTargetPlatform.XcodeProjectSettings", "CertificateHashes")]
+		private List<string> HardCodedCertOfTruthHashes = new List<string>();
+
+		[ConfigFile(ConfigHierarchyType.Engine, "/Script/MacTargetPlatform.XcodeProjectSettings", "P4CertificateFilePaths")]
+		private List<string> HardCodedCertOfTruthP4FilePaths = new List<string>();
+
+		[ConfigFile(ConfigHierarchyType.Engine, "/Script/MacTargetPlatform.XcodeProjectSettings", "CertificateStreamNames")]
+		private List<string> HardCodedStreamNames = new List<string>();
+
+		partial void TryInstallingCertificates(ILogger Logger);
+
+		/// <summary>
+		/// Constructor
+		/// </summary>
+	    public AppleCertificateVerification()
+	    {
+	    	// the config entries shouldn't be in a platform-specific ini file, so we are just
+	    	// using the mac platform because we install the certs on mac and the certs are used across
+	    	// all apple platforms.
+	        ConfigCache.ReadSettings(null, UnrealTargetPlatform.Mac, this);
+	    }
+
+		private List<string> GetCertOfTruthHashes(bool bPullDownFromP4, ILogger Logger)
+		{
+			List<string> CertOfTruthHashes = new List<string>();
+
+			if(HardCodedCertOfTruthHashes.Count == 0 || HardCodedCertOfTruthP4FilePaths.Count == 0 || HardCodedStreamNames.Count == 0)
+			{
+				return CertOfTruthHashes;
+			}
+			else if(HardCodedCertOfTruthHashes.Count != HardCodedCertOfTruthP4FilePaths.Count || 
+					HardCodedStreamNames.Count != HardCodedCertOfTruthHashes.Count)
+			{
+				Logger.LogError("Invalid Apple certificate verification config setup!");
+				return CertOfTruthHashes;
+			}
+
+			for(int i = 0; HardCodedCertOfTruthHashes.Count > i; ++i)
+			{
+				string CertOfTruthFileName = Path.GetFileName(HardCodedCertOfTruthP4FilePaths[i]);
+				FileReference LocalProfile =
+					FileReference.Combine(new DirectoryReference(Environment.GetFolderPath(Environment.SpecialFolder.UserProfile)), 
+										  "Library", "Application Support", "Epic", CertOfTruthFileName);
+
+				bool bPullDownSucceeded = false;
+
+				if(bPullDownFromP4)
+				{
+					DirectoryReference.CreateDirectory(LocalProfile.Directory);
+
+					// first strategy is pulling it down from p4 to make sure we have the latest revision of the cert
+					string Output = Utils.RunLocalProcessAndReturnStdOut(
+		              "/bin/bash",
+		              $"p4 print -o \"{LocalProfile.FullName}\" {HardCodedStreamNames[i]}/{HardCodedCertOfTruthP4FilePaths[i]}",
+		              out int ExitCode
+		          	);
+
+		        	bPullDownSucceeded = !Output.Contains("no such file(s)") && ExitCode == 0;
+				}
+
+				FileReference P4SyncCertOfTruthFilePath = 
+					FileReference.Combine(Unreal.RootDirectory, HardCodedCertOfTruthP4FilePaths[i]);
+
+				if(bPullDownSucceeded || FileReference.Exists(P4SyncCertOfTruthFilePath) || FileReference.Exists(LocalProfile))
+				{
+					// if the pull down succeeded, the LocalProfile needs to exist otherwise p4 print is broken
+					Debug.Assert(!bPullDownSucceeded || FileReference.Exists(LocalProfile));
+
+					// second strategy is using a previously synced version of the cert of truth if the file exists
+					string CertOfTruthFilePath = FileReference.Exists(LocalProfile) ? LocalProfile.FullName : P4SyncCertOfTruthFilePath.FullName;
+
+					if(!bPullDownSucceeded)
+					{
+						int DaysSinceCertUpdate = (int)(DateTime.Now - File.GetLastWriteTime(CertOfTruthFilePath)).TotalDays;
+
+						if(DaysSinceCertUpdate > 90)
+						{
+							Logger.LogError($"Unable to sync Apple certificate of truth {CertOfTruthFileName}. Locally stored copies were updated over 90 days ago ({DaysSinceCertUpdate} since last update). Please check your p4 connection or update your certs!");
+							throw new BuildException("Stopping build due to broken Apple certificate setup...");
+						}
+						else
+						{
+							Logger.LogWarning($"Unable to sync Apple certificate of truth {CertOfTruthFileName} but found locally stored copy updated {DaysSinceCertUpdate} days ago. If you see signing errors, this certificate might be out of date!");
+						}
+					}
+
+					X509Certificate2 Cert = new X509Certificate2(CertOfTruthFilePath);
+					
+					CertOfTruthHashes.Add(Cert.GetCertHashString());
+
+					if(bPullDownSucceeded)
+					{
+						Logger.LogDebug($"using P4 cert of truth for {CertOfTruthFileName}");
+					} 
+					else 
+					{
+						Logger.LogDebug($"using FS cert of truth {CertOfTruthFilePath} for {CertOfTruthFileName}");
+					}
+				}
+				else
+				{
+					// last resort strategy: use a hard coded hash for the cert of truth
+					Logger.LogDebug($"using hardcoded cert of truth for {CertOfTruthFileName}");
+					CertOfTruthHashes.Add(HardCodedCertOfTruthHashes[i]);
+				}
+			}
+
+			return CertOfTruthHashes;
+		}
+
+		private bool HasCertsOfTruthInSystemStore(ILogger Logger)
+		{
+			List<string> CertOfTruthHashes = GetCertOfTruthHashes(true, Logger);
+
+			if(CertOfTruthHashes.Count == 0)
+			{
+				Logger.LogInformation("No certs of truth reported, skipping certificate verification...");
+				return CertOfTruthHashes.Count == 0;
+			}
+
+			X509Certificate2Collection FoundCertCollection = AppleCodeSign.FindCertificateCollection();
+
+			foreach (X509Certificate2 FoundCert in FoundCertCollection)
+			{
+				for(int i = 0; CertOfTruthHashes.Count > i; ++i)
+				{
+					if(FoundCert.GetCertHashString() == CertOfTruthHashes[i])
+					{
+						Logger.LogDebug($"Found a new locally installed Apple certificate matching known Apple certificate of truth.");
+						CertOfTruthHashes.RemoveAt(i);
+						break;
+					}
+				}
+			}
+
+			return CertOfTruthHashes.Count == 0;
+		}
+
+		/// <summary>
+		/// Verifies you have the valid certificates installed and if not,
+		/// attempts to install them.
+		/// </summary> 
+		public void VerifyCertificateInstallation(ILogger Logger)
+		{
+			if(!bCertificatesVerified)
+			{
+				if(HasCertsOfTruthInSystemStore(Logger))
+				{
+					Logger.LogInformation("All known Apple certificates are installed!");
+				}
+				else
+				{
+					Logger.LogError("Not all known Apple certificates are locally installed! Please install them before continuing");
+
+					TryInstallingCertificates(Logger);
+
+					if(HasCertsOfTruthInSystemStore(Logger))
+					{
+						Logger.LogInformation("Succesfully installed known Apple certificates!");
+					}
+					else
+					{
+						Logger.LogInformation("Installation of known Apple certificates failed!");
+						throw new BuildException("Stopping build due to broken Apple certificate setup...");
+					}
+				}
+
+				bCertificatesVerified = true;
+			}
+		}
+	}
+}
diff --git a/Engine/Source/Programs/UnrealBuildTool/Platform/Apple/UEBuildApple.cs b/Engine/Source/Programs/UnrealBuildTool/Platform/Apple/UEBuildApple.cs
--- a/Engine/Source/Programs/UnrealBuildTool/Platform/Apple/UEBuildApple.cs
+++ b/Engine/Source/Programs/UnrealBuildTool/Platform/Apple/UEBuildApple.cs
@@ -127,6 +127,10 @@ namespace UnrealBuildTool
 		public AppleBuildPlatform(UnrealTargetPlatform Platform, UEBuildPlatformSDK SDK, UnrealArchitectureConfig ArchitectureConfig, ILogger Logger)
 			: base(Platform, SDK, ArchitectureConfig, Logger)
 		{
+			if(BuildHostPlatform.Current.Platform == UnrealTargetPlatform.Mac)
+			{
+				(new AppleCertificateVerification()).VerifyCertificateInstallation(Logger);
+			}
 		}
 
 		public abstract AppleTargetRules GetAppleTargetRules(TargetRules Target);
diff --git a/Engine/Source/Programs/UnrealBuildTool/Platform/IOS/IOSCodeSigning.cs b/Engine/Source/Programs/UnrealBuildTool/Platform/IOS/IOSCodeSigning.cs
--- a/Engine/Source/Programs/UnrealBuildTool/Platform/IOS/IOSCodeSigning.cs
+++ b/Engine/Source/Programs/UnrealBuildTool/Platform/IOS/IOSCodeSigning.cs
@@ -85,9 +85,9 @@ namespace UnrealBuildTool
 		}
 
 		/// <summary>
-		/// Finds all valid installed certificates
+		/// Finds all installed certificates as a X509Certificate2Collection
 		/// </summary>
-		public static List<string> FindCertificates()
+		public static X509Certificate2Collection FindCertificateCollection()
 		{
 			string[] ValidCertificatePrefixes = { "iPhone Developer", "iPhone Distribution", "Apple Development", "Apple Distribution" };
 
@@ -141,6 +141,16 @@ namespace UnrealBuildTool
 				Store.Close();
 			}
 
+			return FoundCerts;
+		}
+
+		/// <summary>
+		/// Finds all valid installed certificates
+		/// </summary>
+		public static List<string> FindCertificates()
+		{
+			X509Certificate2Collection FoundCerts = FindCertificateCollection();
+			
 			List<string> Certs = new();
 			foreach (X509Certificate2 TestCert in FoundCerts)
 			{
