UnrealBuildTool: Only link globbed source files for IDE to further reduce project evaluation time when building

#rnx
#Synced-CL 50360693

diff --git a/Engine/Source/Programs/UnrealBuildTool/UnrealBuildTool.csproj b/Engine/Source/Programs/UnrealBuildTool/UnrealBuildTool.csproj
--- a/Engine/Source/Programs/UnrealBuildTool/UnrealBuildTool.csproj
+++ b/Engine/Source/Programs/UnrealBuildTool/UnrealBuildTool.csproj
@@ -26,16 +26,40 @@
     <Compile Include="..\Shared\MetaData.cs" Link="Properties\MetaData.cs" />
   </ItemGroup>
 
-  <!-- The following will glob all platform extensions' UnrealBuildTool source files and include them in the build. They will also appear in VisualStudio. -->
-  <ItemGroup>
+  <!-- For regular builds, include only the source files without Link attribute. -->
+  <ItemGroup Condition="'$(DesignTimeBuild)' != 'true'">
+    <Compile Include="../../../Platforms/*/Source/Programs/UnrealBuildTool/**/*.cs" />
+    <Compile Include="../../../Restricted/*/Platforms/*/Source/Programs/UnrealBuildTool/**/*.cs" />
+    <Compile Include="../../../Restricted/*/Source/Programs/UnrealBuildTool/**/*.cs" />
+  </ItemGroup>
+
+  <!-- For design-time builds (IDE), use Link attribute to organize files in Solution Explorer. -->
+  <ItemGroup Condition="'$(DesignTimeBuild)' == 'true'">
+    <!-- Platform extension and Restricted source files -->
     <Compile Include="../../../Platforms/*/Source/Programs/UnrealBuildTool/**/*.cs"
              Link="Platform/$([System.String]::Copy('%(RecursiveDir)').Replace('\', '/').Replace('Source/Programs/UnrealBuildTool/', ''))%(Filename)%(Extension)" />
     <Compile Include="../../../Restricted/*/Platforms/*/Source/Programs/UnrealBuildTool/**/*.cs"
              Link="Platform/$([System.String]::Copy('%(RecursiveDir)').Replace('\', '/').Replace('Source/Programs/UnrealBuildTool/', ''))%(Filename)%(Extension)" />
     <Compile Include="../../../Restricted/*/Source/Programs/UnrealBuildTool/**/*.cs"
              Link="Restricted/$([System.String]::Copy('%(RecursiveDir)').Replace('\', '/').Replace('Source/Programs/UnrealBuildTool/', ''))%(Filename)%(Extension)" />
-  </ItemGroup>
 
+    <!-- SDK config files -->
+    <None Include="../../../Platforms/*/Config/*_SDK.json" Link="Platform/%(RecursiveDir)%(Filename)%(Extension)" />
+    <None Include="../../../Config/*/*_SDK.json" Link="Platform/%(RecursiveDir)%(Filename)%(Extension)" />
+
+    <!-- Non-compiled files for IDE display only. -->
+    <!-- Various batch files and shell scripts that build or run UnrealBuildTool -->
+    <None Include="../../../Build/BatchFiles/**/*" Link="BatchFiles/%(RecursiveDir)%(Filename)%(Extension)" />
+
+    <!-- Various BuildConfiguration files -->
+    <None Condition="Exists('..\..\..\Restricted\NotForLicensees\Programs\UnrealBuildTool\BuildConfiguration.xml')" Include="..\..\..\Restricted\NotForLicensees\Programs\UnrealBuildTool\BuildConfiguration.xml" Link="BuildConfiguration\BuildConfiguration.xml (Engine NotForLicensees)" />
+    <None Condition="Exists('..\..\..\Saved\UnrealBuildTool\BuildConfiguration.xml')" Include="..\..\..\Saved\UnrealBuildTool\BuildConfiguration.xml" Link="BuildConfiguration\BuildConfiguration.xml (Engine Saved)" />
+    <None Condition="Exists('$(ProgramData)\Unreal Engine\UnrealBuildTool\BuildConfiguration.xml')" Include="$(ProgramData)\Unreal Engine\UnrealBuildTool\BuildConfiguration.xml" Link="BuildConfiguration\BuildConfiguration.xml (Global ProgramData)" />
+    <None Condition="Exists('$(AppData)\Unreal Engine\UnrealBuildTool\BuildConfiguration.xml')" Include="$(AppData)\Unreal Engine\UnrealBuildTool\BuildConfiguration.xml" Link="BuildConfiguration\BuildConfiguration.xml (Global AppData)" />
+    <None Condition="Exists('$(LocalAppData)\Unreal Engine\UnrealBuildTool\BuildConfiguration.xml')" Include="$(LocalAppData)\Unreal Engine\UnrealBuildTool\BuildConfiguration.xml" Link="BuildConfiguration\BuildConfiguration.xml (Global LocalAppData)" />
+    <None Condition="Exists('$(UserProfile)\Documents\Unreal Engine\UnrealBuildTool\BuildConfiguration.xml')" Include="$(UserProfile)\Documents\Unreal Engine\UnrealBuildTool\BuildConfiguration.xml" Link="BuildConfiguration\BuildConfiguration.xml (Global Documents)" />
+    <None Condition="Exists('..\..\..\Intermediate\Build\UnrealBuildTool.Env.BuildConfiguration.xml')" Include="..\..\..\Intermediate\Build\UnrealBuildTool.Env.BuildConfiguration.xml" Link="BuildConfiguration\BuildConfiguration.xml (Environment)" />
+  </ItemGroup>
 
   <ItemGroup>
     <!-- If you add a project reference, BuildUAT.bat/.sh and BuildUBT.bat/.sh need to be updated to also test for those projects.-->
@@ -66,23 +90,4 @@
     <PackageReference Include="System.Text.Encoding.CodePages" Version="8.0.0" />
   </ItemGroup>
 
-  <!-- Non-compiled files for IDE display only. -->
-  <ItemGroup Condition="'$(DesignTimeBuild)' == 'true'">
-    <!-- SDK config files -->
-    <None Include="../../../Platforms/*/Config/*_SDK.json" Link="Platform/%(RecursiveDir)%(Filename)%(Extension)" />
-    <None Include="../../../Config/*/*_SDK.json" Link="Platform/%(RecursiveDir)%(Filename)%(Extension)" />
-
-    <!-- Various batch files and shell scripts that build or run UnrealBuildTool -->
-    <None Include="../../../Build/BatchFiles/**/*" Link="BatchFiles/%(RecursiveDir)%(Filename)%(Extension)" />
-
-    <!-- Various BuildConfiguration files -->
-    <None Condition="Exists('..\..\..\Restricted\NotForLicensees\Programs\UnrealBuildTool\BuildConfiguration.xml')" Include="..\..\..\Restricted\NotForLicensees\Programs\UnrealBuildTool\BuildConfiguration.xml" Link="BuildConfiguration\BuildConfiguration.xml (Engine NotForLicensees)" />
-    <None Condition="Exists('..\..\..\Saved\UnrealBuildTool\BuildConfiguration.xml')" Include="..\..\..\Saved\UnrealBuildTool\BuildConfiguration.xml" Link="BuildConfiguration\BuildConfiguration.xml (Engine Saved)" />
-    <None Condition="Exists('$(ProgramData)\Unreal Engine\UnrealBuildTool\BuildConfiguration.xml')" Include="$(ProgramData)\Unreal Engine\UnrealBuildTool\BuildConfiguration.xml" Link="BuildConfiguration\BuildConfiguration.xml (Global ProgramData)" />
-    <None Condition="Exists('$(AppData)\Unreal Engine\UnrealBuildTool\BuildConfiguration.xml')" Include="$(AppData)\Unreal Engine\UnrealBuildTool\BuildConfiguration.xml" Link="BuildConfiguration\BuildConfiguration.xml (Global AppData)" />
-    <None Condition="Exists('$(LocalAppData)\Unreal Engine\UnrealBuildTool\BuildConfiguration.xml')" Include="$(LocalAppData)\Unreal Engine\UnrealBuildTool\BuildConfiguration.xml" Link="BuildConfiguration\BuildConfiguration.xml (Global LocalAppData)" />
-    <None Condition="Exists('$(UserProfile)\Documents\Unreal Engine\UnrealBuildTool\BuildConfiguration.xml')" Include="$(UserProfile)\Documents\Unreal Engine\UnrealBuildTool\BuildConfiguration.xml" Link="BuildConfiguration\BuildConfiguration.xml (Global Documents)" />
-    <None Condition="Exists('..\..\..\Intermediate\Build\UnrealBuildTool.Env.BuildConfiguration.xml')" Include="..\..\..\Intermediate\Build\UnrealBuildTool.Env.BuildConfiguration.xml" Link="BuildConfiguration\BuildConfiguration.xml (Environment)" />
-  </ItemGroup>
-
 </Project>
