[UBT]
* Bumped CppDependencyCache to make sure to flush out mac .d files that does not have proper casing

#Synced-CL 50356797

diff --git a/Engine/Source/Programs/UnrealBuildTool/System/CppDependencyCache.cs b/Engine/Source/Programs/UnrealBuildTool/System/CppDependencyCache.cs
--- a/Engine/Source/Programs/UnrealBuildTool/System/CppDependencyCache.cs
+++ b/Engine/Source/Programs/UnrealBuildTool/System/CppDependencyCache.cs
@@ -67,7 +67,7 @@ namespace UnrealBuildTool
 			/// <summary>
 			/// The current file version
 			/// </summary>
-			public const int CurrentVersion = 3;
+			public const int CurrentVersion = 4;
 
 			/// <summary>
 			/// Location of this dependency cache
