UnrealBuildTool: When defaulting to a MSVC toolchain, allow for any available rather than checking what Visual Studio is installed

#rnx

#Synced-CL 50376887

diff --git a/Engine/Source/Programs/UnrealBuildTool/Platform/Windows/UEBuildWindows.cs b/Engine/Source/Programs/UnrealBuildTool/Platform/Windows/UEBuildWindows.cs
--- a/Engine/Source/Programs/UnrealBuildTool/Platform/Windows/UEBuildWindows.cs
+++ b/Engine/Source/Programs/UnrealBuildTool/Platform/Windows/UEBuildWindows.cs
@@ -10,6 +10,7 @@ using System.Runtime.Versioning;
 using System.Text;
 using EpicGames.Core;
 using Microsoft.Extensions.Logging;
+using Microsoft.Extensions.Logging.Abstractions;
 using UnrealBuildBase;
 
 namespace UnrealBuildTool
@@ -896,20 +897,9 @@ namespace UnrealBuildTool
 		/// <summary>
 		/// Directory containing the IDE package (Professional, Community, etc...)
 		/// </summary>
-		public string? IDEDir
-		{
-			get
-			{
-				try
-				{
-					return MicrosoftPlatformSDK.FindVisualStudioInstallations(Environment!.ToolChain, Target.Logger).Select(x => x.BaseDir.FullName).FirstOrDefault();
-				}
-				catch (Exception) // Find function will throw if there is no visual studio installed! This can happen w/ clang builds
-				{
-					return null;
-				}
-			}
-		}
+		public string? IDEDir => MicrosoftPlatformSDK.FindVisualStudioInstallations(Environment!.ToolChain, Target.Logger)
+			.Select(x => x.BaseDir.FullName)
+			.FirstOrDefault();
 
 		/// <summary>
 		/// Directory containing ThirdParty DirectX
@@ -1676,7 +1666,11 @@ namespace UnrealBuildTool
 			ProjectFileFormat PreferredAccessor;
 			if (ProjectFileGenerator.GetPreferredSourceCodeAccessor(TargetRules?.ProjectFile, out PreferredAccessor))
 			{
-				if (PreferredAccessor == ProjectFileFormat.VisualStudio2022)
+				if (PreferredAccessor == ProjectFileFormat.VisualStudio)
+				{
+					return WindowsCompiler.VisualStudio;
+				}
+				else if (PreferredAccessor == ProjectFileFormat.VisualStudio2022)
 				{
 					return WindowsCompiler.VisualStudio2022;
 				}
@@ -1686,28 +1680,9 @@ namespace UnrealBuildTool
 				}
 			}
 
-			// Next, default based on what Visual Studio is installed, test for 2022 first
-			IReadOnlyList<VisualStudioInstallation> vsInstalls = MicrosoftPlatformSDK.FindVisualStudioInstallations(Logger);
-			if (vsInstalls.Any(x => x.Compiler == WindowsCompiler.VisualStudio2022))
-			{
-				return WindowsCompiler.VisualStudio2022;
-			}
-			else if (vsInstalls.Any(x => x.Compiler == WindowsCompiler.VisualStudio2026))
-			{
-				return WindowsCompiler.VisualStudio2026;
-			}
-
-			// Finally, default based on what compiler is available, test for 2022 first
-			if (MicrosoftPlatformSDK.HasValidCompiler(WindowsCompiler.VisualStudio2022, Architecture, Logger))
-			{
-				return WindowsCompiler.VisualStudio2022;
-			}
-			else if (MicrosoftPlatformSDK.HasValidCompiler(WindowsCompiler.VisualStudio2026, Architecture, Logger))
-			{
-				return WindowsCompiler.VisualStudio2026;
-			}
-
-			if (!bSkipWarning && !MicrosoftPlatformSDK.IgnoreToolchainErrors)
+			// Check if any valid compilers are installed so a warning can be shown if none are available
+			bool compilerAvailable = MicrosoftPlatformSDK.HasValidCompiler(WindowsCompiler.VisualStudio, Architecture, NullLogger.Instance);
+			if (!compilerAvailable && !bSkipWarning && !MicrosoftPlatformSDK.IgnoreToolchainErrors)
 			{
 				UEBuildPlatformSDK? SDK = GetSDK(UnrealTargetPlatform.Win64);
 				VersionNumber? MinimumMsvcVersion = SDK?.GetVersionNumberFromConfig("MinimumVisualCppVersion");
@@ -1716,24 +1691,15 @@ namespace UnrealBuildTool
 				string MinimumMsvcVersionStr = MinimumMsvcVersion?.ToString() ?? "Latest";
 				string MinimumVsVersionStr = MinimumVsVersion?.ToString() ?? "Latest";
 				string PreferredMsvcVersionStr = PreferredMsvcVersion?.ToString() ?? "Latest";
-				string PreferredComponentVersionStr = PreferredMsvcVersion != null
-					? $"v{PreferredMsvcVersion.Components[0]}.{PreferredMsvcVersion.Components[1]}-{PreferredMsvcVersion.Components[0] + 3}.{PreferredMsvcVersion.Components[1] - 30}"
-					: "Latest";
-				string PreferredComponentStr = Architecture == UnrealArch.X64 ?
-					$"MSVC v143 - VS 2022 C++ x64/x86 build tools ({PreferredComponentVersionStr})" :
-					$"MSVC v143 - VS 2022 C++ ARM64 build tools ({PreferredComponentVersionStr})";
+				string ArchStr = Architecture == UnrealArch.X64 ? "x64/x86" : "ARM64/ARM64EC";
+				string PreferredComponentStr = PreferredMsvcVersion != null ?
+					$"MSVC Build Tools v{PreferredMsvcVersion.Components[0]}.{PreferredMsvcVersion.Components[1]} for {ArchStr}" :
+					$"MSVC Build Tools for {ArchStr} (Latest)";
+
 				// If we do have a Visual Studio installation, but we're missing just the C++ parts, warn about that.
-				if (TryGetVSInstallDirs(WindowsCompiler.VisualStudio2022, Logger) != null)
-				{
-					Logger.LogWarning("Visual Studio 2022 is installed, but is out of date or missing a valid C++ toolchain (minimum version {MinVersion}, preferred version {PreferredVersion}). Please update Visual Studio 2022 to {MinimumVsVersionStr} or later and verify that the \"{Component}\" component is selected in the Visual Studio 2022 installation options.",
-						MinimumMsvcVersionStr,
-						PreferredMsvcVersionStr,
-						MinimumVsVersionStr,
-						PreferredComponentStr);
-				}
-				else if (TryGetVSInstallDirs(WindowsCompiler.VisualStudio2026, Logger) != null)
+				if (MicrosoftPlatformSDK.FindVisualStudioInstallations(NullLogger.Instance).Any())
 				{
-					Logger.LogWarning("Visual Studio 2026 is installed, but is out of date or missing a valid C++ toolchain (minimum version {MinVersion}, preferred version {PreferredVersion}). Please update Visual Studio 2026 to {MinimumVsVersionStr} or later and verify that the \"{Component}\" component is selected in the Visual Studio 2026 installation options.",
+					Logger.LogWarning("Visual Studio is installed, but is out of date or missing a valid C++ toolchain (minimum version {MinVersion}, preferred version {PreferredVersion}). Please update Visual Studio to {MinimumVsVersionStr} or later and verify that the \"{Component}\" component is selected in the Visual Studio installation options.",
 						MinimumMsvcVersionStr,
 						PreferredMsvcVersionStr,
 						MinimumVsVersionStr,
@@ -1741,7 +1707,7 @@ namespace UnrealBuildTool
 				}
 				else
 				{
-					Logger.LogWarning("No valid Visual C++ toolchain was found (minimum version {MinVersion}, preferred version {PreferredVersion}). Please download and install Visual Studio 2022 {MinimumVsVersionStr} or later and verify that the \"{Component}\" component is selected in the Visual Studio 2022 installation options.",
+					Logger.LogWarning("No valid Visual C++ toolchain was found (minimum version {MinVersion}, preferred version {PreferredVersion}). Please download and install Visual Studio {MinimumVsVersionStr} or later and verify that the \"{Component}\" component is selected in the Visual Studio installation options.",
 						MinimumMsvcVersionStr,
 						PreferredMsvcVersionStr,
 						MinimumVsVersionStr,
@@ -1749,8 +1715,8 @@ namespace UnrealBuildTool
 				}
 			}
 
-			// Finally, default to VS2022 anyway
-			return WindowsCompiler.VisualStudio2022;
+			// Default to any available visual studio compiler
+			return WindowsCompiler.VisualStudio;
 		}
 
 		/// <summary>
