UnrealBuildTool: Catch and handle cancellation of result store tasks in the UBA executor.

#Synced-CL 50395407

diff --git a/Engine/Source/Programs/UnrealBuildTool/Executors/UnrealBuildAccelerator/UBAExecutor.cs b/Engine/Source/Programs/UnrealBuildTool/Executors/UnrealBuildAccelerator/UBAExecutor.cs
--- a/Engine/Source/Programs/UnrealBuildTool/Executors/UnrealBuildAccelerator/UBAExecutor.cs
+++ b/Engine/Source/Programs/UnrealBuildTool/Executors/UnrealBuildAccelerator/UBAExecutor.cs
@@ -1095,7 +1095,10 @@ namespace UnrealBuildTool
 
 				ExecuteResults result = new(logLines, exitCode, executionTime, processorTime, peakMemory, action, additionalDescription);
 				actionLogger.AddActionToLog(result);
-				resultStoreTasks.Add(Task.Run(() => resultStore.StoreResultAsync(action, ActionResult.From(result), _cancellationSource.Token)));
+				Task storeResultTask = Task.Run(
+					() => resultStore.StoreResultAsync(action, ActionResult.From(result), _cancellationSource.Token),
+					_cancellationSource.Token);
+				resultStoreTasks.Add(storeResultTask);
 
 				return ProcessFinishedResponse.None;
 			});
@@ -1218,6 +1221,10 @@ namespace UnrealBuildTool
 				{
 					logger.LogWarning(whenAllTask.Exception, "Some action result store tasks did not succeed");
 				}
+				catch (OperationCanceledException)
+				{
+					// We've already logged the cancellation, we don't have to do anything more here, just don't throw the cancellation up the callstack.
+				}
 
 				return success;
 			}
