UnrealBuildTool: In WriteMetadataMode, await the inner execute method so we still hold the mutex whilst we carry out this logic. Add a single-thread task scheduler to Core to achieve this.

#Synced-CL 50198585

diff --git a/Engine/Source/Programs/Shared/EpicGames.Build.Tests/GlobalSingleInstanceMutexTests.cs b/Engine/Source/Programs/Shared/EpicGames.Build.Tests/GlobalSingleInstanceMutexTests.cs
new file mode 100644
--- /dev/null
+++ b/Engine/Source/Programs/Shared/EpicGames.Build.Tests/GlobalSingleInstanceMutexTests.cs
@@ -0,0 +1,49 @@
+// Copyright Epic Games, Inc. All Rights Reserved.
+
+using System.Threading.Tasks;
+using Microsoft.VisualStudio.TestTools.UnitTesting;
+using UnrealBuildBase;
+
+namespace EpicGames.Build.Tests;
+
+[TestClass]
+public class GlobalSingleInstanceMutexTests
+{
+	[TestMethod]
+	public async Task AcquiringGlobalMutexBlocksCorrectlyAsync()
+	{
+		const string MutexName = "EpicGames.Build.Tests.UnitTestMutex";
+		GlobalSingleInstanceMutex? mutex1 = new(MutexName, true);
+
+		try
+		{
+			Task mutex2Task = Task.Run(() =>
+			{
+				using GlobalSingleInstanceMutex mutex2 = new(MutexName, true);
+			}, TestContext.CancellationToken);
+
+			// We should get the delay task back because the mutex 2 task should be blocked acquiring the mutex.
+			Task delayTask = Task.Delay(200, TestContext.CancellationToken);
+			Task finishedTask = await Task.WhenAny(mutex2Task, delayTask);
+
+			Assert.AreSame(delayTask, finishedTask);
+
+			// Dispose of the first mutex so the second one can acquire it.
+			mutex1.Dispose();
+			mutex1 = null;
+
+			delayTask = Task.Delay(200, TestContext.CancellationToken);
+			finishedTask = await Task.WhenAny(mutex2Task, delayTask);
+			await finishedTask;
+
+			// We should get the mutex 2 task as it should have been able to acquire the mutex and free it again.
+			Assert.AreSame(mutex2Task, finishedTask);
+		}
+		finally
+		{
+			mutex1?.Dispose();
+		}
+	}
+
+	public TestContext TestContext { get; set; }
+}
diff --git a/Engine/Source/Programs/Shared/EpicGames.Build/System/GlobalSingleInstanceMutex.cs b/Engine/Source/Programs/Shared/EpicGames.Build/System/GlobalSingleInstanceMutex.cs
--- a/Engine/Source/Programs/Shared/EpicGames.Build/System/GlobalSingleInstanceMutex.cs
+++ b/Engine/Source/Programs/Shared/EpicGames.Build/System/GlobalSingleInstanceMutex.cs
@@ -2,6 +2,7 @@
 
 using System;
 using System.Threading;
+using System.Threading.Tasks;
 using EpicGames.Core;
 
 namespace UnrealBuildBase
@@ -16,6 +17,11 @@ namespace UnrealBuildBase
 		/// </summary>
 		Mutex? _globalMutex;
 
+		/// <summary>
+		/// A single-threaded scheduler to run the mutex acquire / release on, to make this class safe for use within an async context.
+		/// </summary>
+		SingleThreadTaskScheduler _scheduler = new();
+
 		/// <summary>
 		/// Constructor. Attempts to acquire the global mutex
 		/// </summary>
@@ -23,27 +29,16 @@ namespace UnrealBuildBase
 		/// <param name="bWaitMutex">Allow waiting for the mutex to be acquired</param>
 		public GlobalSingleInstanceMutex(string mutexName, bool bWaitMutex)
 		{
-			// Try to create the mutex, with it initially locked
-			bool bCreatedMutex;
-			_globalMutex = new Mutex(true, mutexName, out bCreatedMutex);
-
-			// If we didn't create the mutex, we can wait for it or fail immediately
-			if (!bCreatedMutex)
+			try
 			{
-				if (bWaitMutex)
-				{
-					try
-					{
-						_globalMutex.WaitOne();
-					}
-					catch (AbandonedMutexException)
-					{
-					}
-				}
-				else
-				{
-					throw new BuildLogEventException(new CompilationResultException(CompilationResult.ConflictingInstance), "A conflicting instance of {Mutex} is already running.", mutexName);
-				}
+				RunActionOnScheduler(() => GetMutex(mutexName, bWaitMutex));
+			}
+			catch
+			{
+				// If we can't construct the object fully then we should dispose our resources otherwise they would leak.
+				_globalMutex?.Dispose();
+				_scheduler.Dispose();
+				throw;
 			}
 		}
 
@@ -82,9 +77,41 @@ namespace UnrealBuildBase
 		{
 			if (disposing)
 			{
-				_globalMutex?.ReleaseMutex();
+				RunActionOnScheduler(() => _globalMutex?.ReleaseMutex());
 				_globalMutex?.Dispose();
 				_globalMutex = null;
+				_scheduler.Dispose();
+			}
+		}
+
+		private void RunActionOnScheduler(Action action)
+		{
+			Task getMutexTask = Task.Factory.StartNew(action, CancellationToken.None, TaskCreationOptions.None, _scheduler);
+			getMutexTask.GetAwaiter().GetResult();
+		}
+
+		private void GetMutex(string mutexName, bool bWaitMutex)
+		{
+			// Try to create the mutex, with it initially locked
+			_globalMutex = new Mutex(true, mutexName, out bool bCreatedMutex);
+
+			// If we didn't create the mutex, we can wait for it or fail immediately
+			if (!bCreatedMutex)
+			{
+				if (bWaitMutex)
+				{
+					try
+					{
+						_globalMutex.WaitOne();
+					}
+					catch (AbandonedMutexException)
+					{
+					}
+				}
+				else
+				{
+					throw new BuildLogEventException(new CompilationResultException(CompilationResult.ConflictingInstance), "A conflicting instance of {Mutex} is already running.", mutexName);
+				}
 			}
 		}
 	}
diff --git a/Engine/Source/Programs/Shared/EpicGames.Core.Tests/SingleThreadTaskSchedulerTests.cs b/Engine/Source/Programs/Shared/EpicGames.Core.Tests/SingleThreadTaskSchedulerTests.cs
new file mode 100644
--- /dev/null
+++ b/Engine/Source/Programs/Shared/EpicGames.Core.Tests/SingleThreadTaskSchedulerTests.cs
@@ -0,0 +1,53 @@
+// Copyright Epic Games, Inc. All Rights Reserved.
+
+using System;
+using System.Collections.Concurrent;
+using System.Diagnostics;
+using System.Linq;
+using System.Threading;
+using System.Threading.Tasks;
+using Microsoft.VisualStudio.TestTools.UnitTesting;
+
+namespace EpicGames.Core.Tests;
+
+[TestClass]
+public class SingleThreadTaskSchedulerTests
+{
+	[TestMethod]
+	public async Task TasksAreExecutedOnTheSameThreadInOrderAsync()
+	{
+		using SingleThreadTaskScheduler scheduler = new();
+
+		const int NumTasksToRun = 10;
+		ConcurrentQueue<int> threadIDs = [];
+		ConcurrentQueue<int> taskIDs = [];
+		ConcurrentQueue<long> timestamps = [];
+		Task[] tasks = new Task[NumTasksToRun];
+
+		for (int i = 0; i < NumTasksToRun; ++i)
+		{
+			int taskID = i;
+			tasks[i] = Task.Factory.StartNew(() =>
+				{
+					threadIDs.Enqueue(Environment.CurrentManagedThreadId);
+					taskIDs.Enqueue(taskID);
+					timestamps.Enqueue(Stopwatch.GetTimestamp());
+					// Sleep for just long enough to not be instantaneous.
+					Thread.Sleep(1);
+					timestamps.Enqueue(Stopwatch.GetTimestamp());
+				},
+				CancellationToken.None,
+				TaskCreationOptions.None,
+				scheduler);
+		}
+
+		await Task.WhenAll(tasks);
+
+		// Everything should have ran on the same thread.
+		Assert.AreEqual(1, threadIDs.Distinct().Count());
+		// The tasks should have executed in the order they were started.
+		CollectionAssert.AreEqual(Enumerable.Range(0, NumTasksToRun).ToList(), taskIDs);
+		// The timestamps for each task shouldn't overlap at all.
+		CollectionAssert.AreEqual(timestamps.OrderBy(x => x).ToList(), timestamps);
+	}
+}
diff --git a/Engine/Source/Programs/Shared/EpicGames.Core/SingleThreadTaskScheduler.cs b/Engine/Source/Programs/Shared/EpicGames.Core/SingleThreadTaskScheduler.cs
new file mode 100644
--- /dev/null
+++ b/Engine/Source/Programs/Shared/EpicGames.Core/SingleThreadTaskScheduler.cs
@@ -0,0 +1,86 @@
+// Copyright Epic Games, Inc. All Rights Reserved.
+
+using System;
+using System.Collections.Concurrent;
+using System.Collections.Generic;
+using System.Threading;
+using System.Threading.Tasks;
+
+namespace EpicGames.Core;
+
+/// <summary>
+/// A task scheduler that runs with a single thread.
+/// Allows for use of the TPL whilst running Tasks on a consistent thread, useful for anything that requires thread affinity,
+/// such as global mutexes.
+/// </summary>
+public sealed class SingleThreadTaskScheduler : TaskScheduler, IDisposable
+{
+	/// <summary>
+	/// Creates a single-thread task scheduler.
+	/// </summary>
+	public SingleThreadTaskScheduler()
+	{
+		_cancellationTokenSource = new CancellationTokenSource();
+		_thread = new Thread(ThreadEntryPoint)
+		{
+			// The thread should not stop the application from exiting, so set it as background.
+			IsBackground = true
+		};
+		_thread.Start();
+	}
+
+	/// <inheritdoc />
+	public override int MaximumConcurrencyLevel => 1;
+
+	/// <inheritdoc />
+	public void Dispose()
+	{
+		_cancellationTokenSource.Cancel();
+		_cancellationTokenSource.Dispose();
+		_tasks.Dispose();
+		GC.SuppressFinalize(this);
+	}
+
+	/// <inheritdoc />
+	protected override IEnumerable<Task> GetScheduledTasks()
+	{
+		return [.. _tasks];
+	}
+
+	/// <inheritdoc />
+	protected override void QueueTask(Task task)
+	{
+		_tasks.Add(task);
+	}
+
+	/// <inheritdoc />
+	protected override bool TryExecuteTaskInline(Task task, bool taskWasPreviouslyQueued)
+	{
+		if (Thread.CurrentThread == _thread)
+		{
+			return TryExecuteTask(task);
+		}
+
+		return false;
+	}
+
+	private void ThreadEntryPoint()
+	{
+		while (!_cancellationTokenSource.IsCancellationRequested)
+		{
+			try
+			{
+				Task task = _tasks.Take(_cancellationTokenSource.Token);
+				TryExecuteTask(task);
+			}
+			catch (OperationCanceledException)
+			{
+				// Don't do anything, we'll exit the loop immediately.
+			}
+		}
+	}
+
+	private readonly CancellationTokenSource _cancellationTokenSource;
+	private readonly Thread _thread;
+	private readonly BlockingCollection<Task> _tasks = [];
+}
\ No newline at end of file
diff --git a/Engine/Source/Programs/UnrealBuildTool/Modes/WriteMetadataMode.cs b/Engine/Source/Programs/UnrealBuildTool/Modes/WriteMetadataMode.cs
--- a/Engine/Source/Programs/UnrealBuildTool/Modes/WriteMetadataMode.cs
+++ b/Engine/Source/Programs/UnrealBuildTool/Modes/WriteMetadataMode.cs
@@ -10,8 +10,6 @@ using EpicGames.Core;
 using Microsoft.Extensions.Logging;
 using UnrealBuildBase;
 
-#pragma warning disable CA1849
-
 namespace UnrealBuildTool
 {
 	/// <summary>
@@ -60,14 +58,14 @@ namespace UnrealBuildTool
 		/// <param name="Arguments">Command line arguments</param>
 		/// <returns>Exit code</returns>
 		/// <param name="Logger"></param>
-		public override Task<int> ExecuteAsync(CommandLineArguments Arguments, ILogger Logger)
+		public override async Task<int> ExecuteAsync(CommandLineArguments Arguments, ILogger Logger)
 		{
 			// Acquire a different mutex to the regular UBT instance, since this mode will be called as part of a build. We need the mutex to ensure that building two modular configurations 
 			// in parallel don't clash over writing shared *.modules files (eg. DebugGame and Development editors).
 			string MutexName = GlobalSingleInstanceMutex.GetUniqueMutexForPath("UnrealBuildTool_WriteMetadata", Unreal.RootDirectory);
 			using (new GlobalSingleInstanceMutex(MutexName, true))
 			{
-				return ExecuteInternal(Arguments, Logger);
+				return await ExecuteInternal(Arguments, Logger);
 			}
 		}
 
@@ -77,14 +75,14 @@ namespace UnrealBuildTool
 		/// <param name="Arguments">Command line arguments</param>
 		/// <param name="Logger">Logger for output</param>
 		/// <returns>Exit code</returns>
-		private static Task<int> ExecuteInternal(CommandLineArguments Arguments, ILogger Logger)
+		private static async Task<int> ExecuteInternal(CommandLineArguments Arguments, ILogger Logger)
 		{
 			// Read the target info
 			FileReference FilePath = Arguments.GetFileReference("-Input=");
 			WriteMetadataTargetInfo TargetInfo;
 			using (FileStream FileStream = FileReference.Open(FilePath, FileMode.Open, FileAccess.Read, FileShare.Read))
 			{
-				WriteMetadataTargetInfo? MaybeTargetInfo = JsonSerializer.Deserialize<WriteMetadataTargetInfo>(FileStream, SerializerOptions);
+				WriteMetadataTargetInfo? MaybeTargetInfo = await JsonSerializer.DeserializeAsync<WriteMetadataTargetInfo>(FileStream, SerializerOptions);
 				TargetInfo = MaybeTargetInfo ?? throw new BuildLogEventException("Target info file '{FilePath}' was invalid", FilePath);
 			}
 
@@ -157,7 +155,7 @@ namespace UnrealBuildTool
 						// Check if the manifest has changed. Note that if a manifest is out of date, we should have generated a new build id causing the contents to differ.
 						if (bNoManifestChanges)
 						{
-							string CurrentText = FileReference.ReadAllText(ManifestFile);
+							string CurrentText = await FileReference.ReadAllTextAsync(ManifestFile);
 							if (CurrentText != OutputText)
 							{
 								Logger.LogError(
@@ -169,7 +167,7 @@ namespace UnrealBuildTool
 						}
 
 						// Write it to disk
-						FileReference.WriteAllText(ManifestFile, OutputText);
+						await FileReference.WriteAllTextAsync(ManifestFile, OutputText);
 					}
 				}
 			}
@@ -201,7 +199,7 @@ namespace UnrealBuildTool
 						// Check if the manifest has changed. Note that if a manifest is out of date, we should have generated a new build id causing the contents to differ.
 						if (bNoManifestChanges)
 						{
-							string CurrentText = FileReference.ReadAllText(ManifestFile);
+							string CurrentText = await FileReference.ReadAllTextAsync(ManifestFile);
 							if (CurrentText != OutputText)
 							{
 								Logger.LogError(
@@ -213,7 +211,7 @@ namespace UnrealBuildTool
 						}
 
 						// Write it to disk
-						FileReference.WriteAllText(ManifestFile, OutputText);
+						await FileReference.WriteAllTextAsync(ManifestFile, OutputText);
 					}
 				}
 			}
@@ -234,7 +232,7 @@ namespace UnrealBuildTool
 				TargetInfo.Receipt.Write(TargetInfo.ReceiptFile);
 			}
 
-			return Task.FromResult(0);
+			return 0;
 		}
 
 		/// <summary>
