UnrealBuildTool: Allow finding DIA SDKs from either VS2026 or VS2022 if VisualStudio is requested.

#Synced-CL 50395407

diff --git a/Engine/Source/Programs/UnrealBuildTool/Platform/Windows/MicrosoftPlatformSDK.cs b/Engine/Source/Programs/UnrealBuildTool/Platform/Windows/MicrosoftPlatformSDK.cs
--- a/Engine/Source/Programs/UnrealBuildTool/Platform/Windows/MicrosoftPlatformSDK.cs
+++ b/Engine/Source/Programs/UnrealBuildTool/Platform/Windows/MicrosoftPlatformSDK.cs
@@ -1523,37 +1523,45 @@ namespace UnrealBuildTool
 		{
 			return CachedDiaSdkDirs.GetOrAdd(Compiler, _ =>
 			{
-				List<DirectoryReference> DiaSdkDirs = new List<DirectoryReference>();
+				List<DirectoryReference> DiaSdkDirs = [];
 
-				DirectoryReference? PlatformDir;
-				if (UEBuildPlatformSDK.TryGetHostPlatformAutoSDKDir(out PlatformDir))
+				// VisualStudio and VisualStudio2026 are synonymous, so try both VS2026 and VS2022 if that's specified.
+				List<WindowsCompiler> CompilersToCheck = Compiler == WindowsCompiler.VisualStudio ?
+					[WindowsCompiler.VisualStudio2026, WindowsCompiler.VisualStudio2022] :
+					[Compiler];
+
+				foreach (WindowsCompiler CompilerToCheck in CompilersToCheck)
 				{
-					string VSDir = String.Empty;
-					switch (Compiler)
+					DirectoryReference? PlatformDir;
+					if (TryGetHostPlatformAutoSDKDir(out PlatformDir))
 					{
-						case WindowsCompiler.VisualStudio2022: VSDir = "VS2022"; break;
-						case WindowsCompiler.VisualStudio2026: VSDir = "VS2026"; break;
+						string VSDir = "";
+						switch (CompilerToCheck)
+						{
+							case WindowsCompiler.VisualStudio2026: VSDir = "VS2026"; break;
+							case WindowsCompiler.VisualStudio2022: VSDir = "VS2022"; break;
+						}
+
+						if (VSDir.Length != 0)
+						{
+							DirectoryReference DiaSdkDir = DirectoryReference.Combine(PlatformDir, "Win64", "DIA SDK", VSDir);
+							if (IsValidDiaSdkDir(DiaSdkDir))
+							{
+								DiaSdkDirs.Add(DiaSdkDir);
+							}
+						}
 					}
 
-					if (!String.IsNullOrEmpty(VSDir))
+					List<DirectoryReference> VisualStudioDirs = FindVisualStudioInstallations(CompilerToCheck, Log.Logger).ConvertAll(x => x.BaseDir);
+					foreach (DirectoryReference VisualStudioDir in VisualStudioDirs)
 					{
-						DirectoryReference DiaSdkDir = DirectoryReference.Combine(PlatformDir, "Win64", "DIA SDK", VSDir);
+						DirectoryReference DiaSdkDir = DirectoryReference.Combine(VisualStudioDir, "DIA SDK");
 						if (IsValidDiaSdkDir(DiaSdkDir))
 						{
 							DiaSdkDirs.Add(DiaSdkDir);
 						}
 					}
 				}
-
-				List<DirectoryReference> VisualStudioDirs = MicrosoftPlatformSDK.FindVisualStudioInstallations(Compiler, Log.Logger).ConvertAll(x => x.BaseDir);
-				foreach (DirectoryReference VisualStudioDir in VisualStudioDirs)
-				{
-					DirectoryReference DiaSdkDir = DirectoryReference.Combine(VisualStudioDir, "DIA SDK");
-					if (IsValidDiaSdkDir(DiaSdkDir))
-					{
-						DiaSdkDirs.Add(DiaSdkDir);
-					}
-				}
 				return DiaSdkDirs;
 			});
 		}
