UnrealBuildTool: Downgrade message when MSVC toolchain is not preferred

#rnx

#Synced-CL 50405662

diff --git a/Engine/Source/Programs/UnrealBuildTool/Platform/Windows/UEBuildWindows.cs b/Engine/Source/Programs/UnrealBuildTool/Platform/Windows/UEBuildWindows.cs
--- a/Engine/Source/Programs/UnrealBuildTool/Platform/Windows/UEBuildWindows.cs
+++ b/Engine/Source/Programs/UnrealBuildTool/Platform/Windows/UEBuildWindows.cs
@@ -1555,7 +1555,6 @@ namespace UnrealBuildTool
 				{
 					_toolchainWarningLogged = true;
 					VersionNumber preferred = MicrosoftPlatformSDK.GetLatestPreferredVersion(Target.WindowsPlatform.ToolChain);
-					MicrosoftPlatformSDK.DumpAllToolChainInstallations(Target.WindowsPlatform.ToolChain, Target.Architecture, Logger);
 					if (Target.WindowsPlatform.ToolchainVersionWarningLevel == WarningLevel.Error)
 					{
 						throw new BuildLogEventException("{Toolchain} toolchain version {Version} is not a preferred version. Please use the latest preferred version {PreferredVersion}", WindowsPlatform.GetCompilerName(Target.WindowsPlatform.ToolChain), Target.WindowsPlatform.Environment.ToolChainVersion, preferred);
diff --git a/Engine/Source/Programs/UnrealBuildTool/Platform/Windows/VCToolChain.cs b/Engine/Source/Programs/UnrealBuildTool/Platform/Windows/VCToolChain.cs
--- a/Engine/Source/Programs/UnrealBuildTool/Platform/Windows/VCToolChain.cs
+++ b/Engine/Source/Programs/UnrealBuildTool/Platform/Windows/VCToolChain.cs
@@ -187,12 +187,12 @@ namespace UnrealBuildTool
 			{
 				if (!MicrosoftPlatformSDK.IsPreferredVersion(EnvVars.Compiler, EnvVars.CompilerVersion))
 				{
-					Lines.Add($"Warning: {WindowsPlatform.GetCompilerName(EnvVars.Compiler)} compiler is not a preferred version");
+					Lines.Add($"Warning: {WindowsPlatform.GetCompilerName(EnvVars.Compiler)} compiler {EnvVars.CompilerVersion} is not a preferred version");
 				}
 
 				if (EnvVars.Compiler != EnvVars.ToolChain && !MicrosoftPlatformSDK.IsPreferredVersion(EnvVars.ToolChain, EnvVars.ToolChainVersion))
 				{
-					Lines.Add($"Warning: {WindowsPlatform.GetCompilerName(EnvVars.ToolChain)} runtime is not a preferred version");
+					Lines.Add($"Notice: {WindowsPlatform.GetCompilerName(EnvVars.ToolChain)} runtime {EnvVars.ToolChainVersion} is not a preferred version");
 				}
 			}
 		}
