[Backout] - CL50396306
#rnx
Original CL Desc
-----------------------------------------------------------------
UnrealBuildTool: In WriteMetadataMode, await the inner execute method so we still hold the mutex whilst we carry out this logic. Add a single-thread task scheduler to Core to achieve this.

diff --git a/Engine/Source/Programs/Shared/EpicGames.Build.Tests/GlobalSingleInstanceMutexTests.cs b/Engine/Source/Programs/Shared/EpicGames.Build.Tests/GlobalSingleInstanceMutexTests.cs
deleted file mode 100644
diff --git a/Engine/Source/Programs/Shared/EpicGames.Build/System/GlobalSingleInstanceMutex.cs b/Engine/Source/Programs/Shared/EpicGames.Build/System/GlobalSingleInstanceMutex.cs
--- a/Engine/Source/Programs/Shared/EpicGames.Build/System/GlobalSingleInstanceMutex.cs
+++ b/Engine/Source/Programs/Shared/EpicGames.Build/System/GlobalSingleInstanceMutex.cs
@@ -2,7 +2,6 @@
 
 using System;
 using System.Threading;
-using System.Threading.Tasks;
 using EpicGames.Core;
 
 namespace UnrealBuildBase
@@ -17,11 +16,6 @@ namespace UnrealBuildBase
 		/// </summary>
 		Mutex? _globalMutex;
 
-		/// <summary>
-		/// A single-threaded scheduler to run the mutex acquire / release on, to make this class safe for use within an async context.
-		/// </summary>
-		SingleThreadTaskScheduler _scheduler = new();
-
 		/// <summary>
 		/// Constructor. Attempts to acquire the global mutex
 		/// </summary>
@@ -29,16 +23,27 @@ namespace UnrealBuildBase
 		/// <param name="bWaitMutex">Allow waiting for the mutex to be acquired</param>
 		public GlobalSingleInstanceMutex(string mutexName, bool bWaitMutex)
 		{
-			try
-			{
-				RunActionOnScheduler(() => GetMutex(mutexName, bWaitMutex));
-			}
-			catch
+			// Try to create the mutex, with it initially locked
+			bool bCreatedMutex;
+			_globalMutex = new Mutex(true, mutexName, out bCreatedMutex);
+
+			// If we didn't create the mutex, we can wait for it or fail immediately
+			if (!bCreatedMutex)
 			{
-				// If we can't construct the object fully then we should dispose our resources otherwise they would leak.
-				_globalMutex?.Dispose();
-				_scheduler.Dispose();
-				throw;
+				if (bWaitMutex)
+				{
+					try
+					{
+						_globalMutex.WaitOne();
+					}
+					catch (AbandonedMutexException)
+					{
+					}
+				}
+				else
+				{
+					throw new BuildLogEventException(new CompilationResultException(CompilationResult.ConflictingInstance), "A conflicting instance of {Mutex} is already running.", mutexName);
+				}
 			}
 		}
 
@@ -77,41 +82,9 @@ namespace UnrealBuildBase
 		{
 			if (disposing)
 			{
-				RunActionOnScheduler(() => _globalMutex?.ReleaseMutex());
+				_globalMutex?.ReleaseMutex();
 				_globalMutex?.Dispose();
 				_globalMutex = null;
-				_scheduler.Dispose();
-			}
-		}
-
-		private void RunActionOnScheduler(Action action)
-		{
-			Task getMutexTask = Task.Factory.StartNew(action, CancellationToken.None, TaskCreationOptions.None, _scheduler);
-			getMutexTask.GetAwaiter().GetResult();
-		}
-
-		private void GetMutex(string mutexName, bool bWaitMutex)
-		{
-			// Try to create the mutex, with it initially locked
-			_globalMutex = new Mutex(true, mutexName, out bool bCreatedMutex);
-
-			// If we didn't create the mutex, we can wait for it or fail immediately
-			if (!bCreatedMutex)
-			{
-				if (bWaitMutex)
-				{
-					try
-					{
-						_globalMutex.WaitOne();
-					}
-					catch (AbandonedMutexException)
-					{
-					}
-				}
-				else
-				{
-					throw new BuildLogEventException(new CompilationResultException(CompilationResult.ConflictingInstance), "A conflicting instance of {Mutex} is already running.", mutexName);
-				}
 			}
 		}
 	}
diff --git a/Engine/Source/Programs/Shared/EpicGames.Core.Tests/SingleThreadTaskSchedulerTests.cs b/Engine/Source/Programs/Shared/EpicGames.Core.Tests/SingleThreadTaskSchedulerTests.cs
deleted file mode 100644
diff --git a/Engine/Source/Programs/Shared/EpicGames.Core/SingleThreadTaskScheduler.cs b/Engine/Source/Programs/Shared/EpicGames.Core/SingleThreadTaskScheduler.cs
deleted file mode 100644
diff --git a/Engine/Source/Programs/UnrealBuildTool/Modes/WriteMetadataMode.cs b/Engine/Source/Programs/UnrealBuildTool/Modes/WriteMetadataMode.cs
--- a/Engine/Source/Programs/UnrealBuildTool/Modes/WriteMetadataMode.cs
+++ b/Engine/Source/Programs/UnrealBuildTool/Modes/WriteMetadataMode.cs
@@ -10,6 +10,8 @@ using EpicGames.Core;
 using Microsoft.Extensions.Logging;
 using UnrealBuildBase;
 
+#pragma warning disable CA1849
+
 namespace UnrealBuildTool
 {
 	/// <summary>
@@ -58,14 +60,14 @@ namespace UnrealBuildTool
 		/// <param name="Arguments">Command line arguments</param>
 		/// <returns>Exit code</returns>
 		/// <param name="Logger"></param>
-		public override async Task<int> ExecuteAsync(CommandLineArguments Arguments, ILogger Logger)
+		public override Task<int> ExecuteAsync(CommandLineArguments Arguments, ILogger Logger)
 		{
 			// Acquire a different mutex to the regular UBT instance, since this mode will be called as part of a build. We need the mutex to ensure that building two modular configurations 
 			// in parallel don't clash over writing shared *.modules files (eg. DebugGame and Development editors).
 			string MutexName = GlobalSingleInstanceMutex.GetUniqueMutexForPath("UnrealBuildTool_WriteMetadata", Unreal.RootDirectory);
 			using (new GlobalSingleInstanceMutex(MutexName, true))
 			{
-				return await ExecuteInternal(Arguments, Logger);
+				return ExecuteInternal(Arguments, Logger);
 			}
 		}
 
@@ -75,14 +77,14 @@ namespace UnrealBuildTool
 		/// <param name="Arguments">Command line arguments</param>
 		/// <param name="Logger">Logger for output</param>
 		/// <returns>Exit code</returns>
-		private static async Task<int> ExecuteInternal(CommandLineArguments Arguments, ILogger Logger)
+		private static Task<int> ExecuteInternal(CommandLineArguments Arguments, ILogger Logger)
 		{
 			// Read the target info
 			FileReference FilePath = Arguments.GetFileReference("-Input=");
 			WriteMetadataTargetInfo TargetInfo;
 			using (FileStream FileStream = FileReference.Open(FilePath, FileMode.Open, FileAccess.Read, FileShare.Read))
 			{
-				WriteMetadataTargetInfo? MaybeTargetInfo = await JsonSerializer.DeserializeAsync<WriteMetadataTargetInfo>(FileStream, SerializerOptions);
+				WriteMetadataTargetInfo? MaybeTargetInfo = JsonSerializer.Deserialize<WriteMetadataTargetInfo>(FileStream, SerializerOptions);
 				TargetInfo = MaybeTargetInfo ?? throw new BuildLogEventException("Target info file '{FilePath}' was invalid", FilePath);
 			}
 
@@ -155,7 +157,7 @@ namespace UnrealBuildTool
 						// Check if the manifest has changed. Note that if a manifest is out of date, we should have generated a new build id causing the contents to differ.
 						if (bNoManifestChanges)
 						{
-							string CurrentText = await FileReference.ReadAllTextAsync(ManifestFile);
+							string CurrentText = FileReference.ReadAllText(ManifestFile);
 							if (CurrentText != OutputText)
 							{
 								Logger.LogError(
@@ -167,7 +169,7 @@ namespace UnrealBuildTool
 						}
 
 						// Write it to disk
-						await FileReference.WriteAllTextAsync(ManifestFile, OutputText);
+						FileReference.WriteAllText(ManifestFile, OutputText);
 					}
 				}
 			}
@@ -199,7 +201,7 @@ namespace UnrealBuildTool
 						// Check if the manifest has changed. Note that if a manifest is out of date, we should have generated a new build id causing the contents to differ.
 						if (bNoManifestChanges)
 						{
-							string CurrentText = await FileReference.ReadAllTextAsync(ManifestFile);
+							string CurrentText = FileReference.ReadAllText(ManifestFile);
 							if (CurrentText != OutputText)
 							{
 								Logger.LogError(
@@ -211,7 +213,7 @@ namespace UnrealBuildTool
 						}
 
 						// Write it to disk
-						await FileReference.WriteAllTextAsync(ManifestFile, OutputText);
+						FileReference.WriteAllText(ManifestFile, OutputText);
 					}
 				}
 			}
@@ -232,7 +234,7 @@ namespace UnrealBuildTool
 				TargetInfo.Receipt.Write(TargetInfo.ReceiptFile);
 			}
 
-			return 0;
+			return Task.FromResult(0);
 		}
 
 		/// <summary>
