Fix iOS ASan. To enable/disable, set "IOSPlatform.bEnableAddressSanitizer = true/false;" in <project>.Target.cs and re-run GenerateProjectFiles.

diff --git a/Engine/Source/Programs/UnrealBuildTool/ProjectFiles/Xcode/XcodeProject.cs b/Engine/Source/Programs/UnrealBuildTool/ProjectFiles/Xcode/XcodeProject.cs
--- a/Engine/Source/Programs/UnrealBuildTool/ProjectFiles/Xcode/XcodeProject.cs
+++ b/Engine/Source/Programs/UnrealBuildTool/ProjectFiles/Xcode/XcodeProject.cs
@@ -2206,6 +2206,11 @@ namespace UnrealBuildTool.XcodeProjectXcconfig
 			Xcconfig.AppendLine($"SUPPORTED_PLATFORMS = {SupportedPlatforms}");
 			Xcconfig.AppendLine($"SDKROOT = {SDKRoot}");
 
+			if (UnrealData.TargetRules.IOSPlatform.bEnableAddressSanitizer)
+			{
+				Xcconfig.AppendLine("ENABLE_ADDRESS_SANITIZER = YES");
+			}
+
 			// Xcode creates the Build Dir (where the .app is) by combining {SYMROOT}/{CONFIGURATION}{EFFECTIVE_PLATFORM_NAME}, so we set SYMROOT
 			// to the Parent directory of the diectory the binary is in, CONFIGURATION to nothing, and EFFECTIVE_PLATFORM_NAME to the directory
 			// the binary is in
@@ -2697,7 +2702,7 @@ namespace UnrealBuildTool.XcodeProjectXcconfig
 			Xcconfig.AppendLine("COMBINE_HIDPI_IMAGES = YES");
 			Xcconfig.AppendLine("USE_HEADERMAP = NO");
 			Xcconfig.AppendLine("ONLY_ACTIVE_ARCH = YES");
-
+			
 			Xcconfig.AppendLine("");
 			Xcconfig.AppendLine("// Platform settings");
 			Xcconfig.AppendLine($"SUPPORTED_PLATFORMS = {SupportedPlatforms}");
