UnrealBuildTool: For targets compiled as .dll, ensure any .lib and .exp are in the BuildProducts

#rnx

#Synced-CL 50424731

#Synced-CL 50424731

diff --git a/Engine/Source/Programs/UnrealBuildTool/Platform/Windows/VCToolChain.cs b/Engine/Source/Programs/UnrealBuildTool/Platform/Windows/VCToolChain.cs
--- a/Engine/Source/Programs/UnrealBuildTool/Platform/Windows/VCToolChain.cs
+++ b/Engine/Source/Programs/UnrealBuildTool/Platform/Windows/VCToolChain.cs
@@ -4026,11 +4026,12 @@ namespace UnrealBuildTool
 
 		public override void ModifyBuildProducts(ReadOnlyTargetRules Target, UEBuildBinary Binary, IEnumerable<string> Libraries, IEnumerable<UEBuildBundleResource> BundleResources, Dictionary<FileReference, BuildProductType> BuildProducts)
 		{
-			if (Binary.Type == UEBuildBinaryType.DynamicLinkLibrary)
+			if (Binary.Type == UEBuildBinaryType.DynamicLinkLibrary && !Target.WindowsPlatform.bAllowClangLinker)
 			{
 				if (Target.bShouldCompileAsDLL)
 				{
-					BuildProducts.Add(FileReference.Combine(Binary.OutputDir, Binary.OutputFilePath.GetFileNameWithoutExtension() + ".lib"), BuildProductType.BuildResource);
+					BuildProducts.Add(FileReference.Combine(Binary.OutputDir, Binary.OutputFilePath.GetFileNameWithoutExtension() + ".lib"), BuildProductType.SymbolFile);
+					BuildProducts.Add(FileReference.Combine(Binary.OutputDir, Binary.OutputFilePath.GetFileNameWithoutExtension() + ".exp"), BuildProductType.SymbolFile);
 				}
 				else
 				{
