UnrealBuildTool: Disable PCH chaining as there currently is mismatched preprocessor state that needs to be addressed

#rnx

#Synced-CL 50424731

diff --git a/Engine/Source/Programs/UnrealBuildTool/Configuration/Rules/TargetRules.cs b/Engine/Source/Programs/UnrealBuildTool/Configuration/Rules/TargetRules.cs
--- a/Engine/Source/Programs/UnrealBuildTool/Configuration/Rules/TargetRules.cs
+++ b/Engine/Source/Programs/UnrealBuildTool/Configuration/Rules/TargetRules.cs
@@ -2176,10 +2176,13 @@ namespace UnrealBuildTool
 
 		/// <summary>
 		/// Whether PCHs should be chained when compiling with clang.
+		/// Note: The earliest PCH a header is included in the chain will be included, so if there is a mismatch
+		/// in preprocessor state this can cause an unexpected version of a header to be included.
+		/// Currently disabled as this is the case.
 		/// </summary>
 		[CommandLine("-NoPCHChain", Value = "false")]
 		[XmlConfigFile(Category = "BuildConfiguration")]
-		public bool bChainPCHs { get; set; } = true;
+		public bool bChainPCHs { get; set; } = false;
 
 		/// <summary>
 		/// Whether PCH headers should be force included for gen.cpp files when PCH is disabled.
