UnrealBuildTool: Suppress existing instances of issue CA1034 (Nested types should not be visible), enable as a warning.

#Synced-CL 50482548

diff --git a/Engine/Source/Programs/UnrealBuildTool/Configuration/Rules/ModuleRules.cs b/Engine/Source/Programs/UnrealBuildTool/Configuration/Rules/ModuleRules.cs
--- a/Engine/Source/Programs/UnrealBuildTool/Configuration/Rules/ModuleRules.cs
+++ b/Engine/Source/Programs/UnrealBuildTool/Configuration/Rules/ModuleRules.cs
@@ -309,6 +309,7 @@ namespace UnrealBuildTool
 			MM,
 		}
 
+#pragma warning disable CA1034 // Nested types should not be visible - preserve API of APIs exposed to Rules files
 		/// <summary>
 		/// Information about a file which is required by the target at runtime, and must be moved around with it.
 		/// </summary>
@@ -609,6 +610,7 @@ namespace UnrealBuildTool
 				Include = include;
 			}
 		}
+#pragma warning restore CA1034
 
 		/// <summary>
 		/// Name of this module
diff --git a/Engine/Source/Programs/UnrealBuildTool/Configuration/Rules/TargetRules.cs b/Engine/Source/Programs/UnrealBuildTool/Configuration/Rules/TargetRules.cs
--- a/Engine/Source/Programs/UnrealBuildTool/Configuration/Rules/TargetRules.cs
+++ b/Engine/Source/Programs/UnrealBuildTool/Configuration/Rules/TargetRules.cs
@@ -569,6 +569,7 @@ namespace UnrealBuildTool
 	/// </summary>
 	public abstract partial class TargetRules
 	{
+#pragma warning disable CA1034 // Nested types should not be visible - preserve API of APIs exposed to Rules files
 		/// <summary>
 		/// Static class wrapping constants aliasing the global TargetType enum.
 		/// </summary>
@@ -599,6 +600,7 @@ namespace UnrealBuildTool
 			/// </summary>
 			public const global::UnrealBuildTool.TargetType Program = global::UnrealBuildTool.TargetType.Program;
 		}
+#pragma warning restore CA1034
 
 		/// <summary>
 		/// The name of this target
diff --git a/Engine/Source/Programs/UnrealBuildTool/Configuration/Rules/TestModuleRules.cs b/Engine/Source/Programs/UnrealBuildTool/Configuration/Rules/TestModuleRules.cs
--- a/Engine/Source/Programs/UnrealBuildTool/Configuration/Rules/TestModuleRules.cs
+++ b/Engine/Source/Programs/UnrealBuildTool/Configuration/Rules/TestModuleRules.cs
@@ -468,6 +468,7 @@ namespace UnrealBuildTool
 			System.IO.File.SetAttributes(InFilePath, System.IO.File.GetAttributes(InFilePath) & ~FileAttributes.ReadOnly);
 		}
 
+#pragma warning disable CA1034 // Nested types should not be visible - preserve API of APIs exposed to Rules files
 #pragma warning disable 8618
 		/// <summary>
 		/// Test metadata class.
@@ -606,5 +607,6 @@ namespace UnrealBuildTool
 #pragma warning restore CA2227
 		}
 #pragma warning restore 8618
+#pragma warning restore CA1034
 	}
 }
\ No newline at end of file
diff --git a/Engine/Source/Programs/UnrealBuildTool/System/BuildHostPlatform.cs b/Engine/Source/Programs/UnrealBuildTool/System/BuildHostPlatform.cs
--- a/Engine/Source/Programs/UnrealBuildTool/System/BuildHostPlatform.cs
+++ b/Engine/Source/Programs/UnrealBuildTool/System/BuildHostPlatform.cs
@@ -73,6 +73,7 @@ namespace UnrealBuildTool
 		/// </summary>
 		public abstract string BinarySuffix { get; }
 
+#pragma warning disable CA1034 // Nested types should not be visible - preserve API of APIs exposed to UAT
 		/// <summary>
 		/// Class that holds information about a running process
 		/// </summary>
@@ -124,6 +125,7 @@ namespace UnrealBuildTool
 			/// <inheritdoc/>
 			public override string? ToString() => $"{Name}, {Filename}";
 		}
+#pragma warning restore CA1034
 
 		/// <summary>
 		/// Gets all currently running processes.
diff --git a/Engine/Source/Programs/UnrealBuildTool/System/DataDrivenPlatformInfo.cs b/Engine/Source/Programs/UnrealBuildTool/System/DataDrivenPlatformInfo.cs
--- a/Engine/Source/Programs/UnrealBuildTool/System/DataDrivenPlatformInfo.cs
+++ b/Engine/Source/Programs/UnrealBuildTool/System/DataDrivenPlatformInfo.cs
@@ -14,6 +14,7 @@ namespace UnrealBuildTool
 	/// </summary>
 	public static class DataDrivenPlatformInfo
 	{
+#pragma warning disable CA1034 // Nested types should not be visible - preserve API of APIs exposed to UAT
 		/// <summary>
 		/// All data driven information about a platform
 		/// </summary>
@@ -58,7 +59,8 @@ namespace UnrealBuildTool
 			/// The raw DataDrivenPlatformInfo.ini file
 			/// </summary>
 			public ConfigFile? PlatformConfig = null;
-		};
+		}
+#pragma warning restore CA1034
 
 		static IReadOnlyDictionary<string, ConfigDataDrivenPlatformInfo>? PlatformInfos = null;
 		static readonly object LockObject = new object();
diff --git a/Engine/Source/Programs/UnrealBuildTool/System/EncryptionAndSigning.cs b/Engine/Source/Programs/UnrealBuildTool/System/EncryptionAndSigning.cs
--- a/Engine/Source/Programs/UnrealBuildTool/System/EncryptionAndSigning.cs
+++ b/Engine/Source/Programs/UnrealBuildTool/System/EncryptionAndSigning.cs
@@ -16,6 +16,7 @@ namespace UnrealBuildTool
 	/// </summary>
 	public static class EncryptionAndSigning
 	{
+#pragma warning disable CA1034 // Nested types should not be visible - preserve API of APIs exposed to UAT
 		/// <summary>
 		/// Wrapper class for a single RSA key
 		/// </summary>
@@ -211,6 +212,7 @@ namespace UnrealBuildTool
 				return false;
 			}
 		}
+#pragma warning restore CA1034
 
 		/// <summary>
 		/// Helper class for formatting incoming hex signing key strings
diff --git a/Engine/Source/Programs/UnrealBuildTool/System/InstalledPlatformInfo.cs b/Engine/Source/Programs/UnrealBuildTool/System/InstalledPlatformInfo.cs
--- a/Engine/Source/Programs/UnrealBuildTool/System/InstalledPlatformInfo.cs
+++ b/Engine/Source/Programs/UnrealBuildTool/System/InstalledPlatformInfo.cs
@@ -55,6 +55,7 @@ namespace UnrealBuildTool
 	/// </summary>
 	public static class InstalledPlatformInfo
 	{
+#pragma warning disable CA1034 // Nested types should not be visible - preserve API of APIs exposed to UAT
 		/// <summary>
 		/// Information about a single installed platform configuration
 		/// </summary>
@@ -116,6 +117,7 @@ namespace UnrealBuildTool
 				bCanBeDisplayed = bInCanBeDisplayed;
 			}
 		}
+#pragma warning restore CA1034
 
 		private static readonly List<InstalledPlatformConfiguration>? InstalledPlatformConfigurations;
 
