UnrealBuildTool: By popular demand, put the c# projects back in the vs2022 project

#rnx

#Synced-CL 50460628

diff --git a/Engine/Source/Programs/UnrealBuildTool/ProjectFiles/VisualStudio/VCProjectFileGenerator.cs b/Engine/Source/Programs/UnrealBuildTool/ProjectFiles/VisualStudio/VCProjectFileGenerator.cs
--- a/Engine/Source/Programs/UnrealBuildTool/ProjectFiles/VisualStudio/VCProjectFileGenerator.cs
+++ b/Engine/Source/Programs/UnrealBuildTool/ProjectFiles/VisualStudio/VCProjectFileGenerator.cs
@@ -423,7 +423,7 @@ namespace UnrealBuildTool
 				}
 			}
 
-			if (bIncludeDotNetPrograms && Settings.ProjectFileFormat == VCProjectFileFormat.VisualStudio2022)
+			/*if (bIncludeDotNetPrograms && Settings.ProjectFileFormat == VCProjectFileFormat.VisualStudio2022)
 			{
 #if NET10_0_OR_GREATER
 				Logger.LogInformation("Visual Studio 2022 does not support .NET 10.0 C# projects, these projects will not be added to the generated solution.");
@@ -437,7 +437,7 @@ namespace UnrealBuildTool
 					bIncludeDotNetPrograms = false;
 				}
 #endif
-			}
+			}*/
 		}
 
 		protected override void GenerateAutomationProjectFiles(
