UnrealBuildTool: Fix include chain string possibly being resolved incorrectly

#rnx

#Synced-CL 50498997

diff --git a/Engine/Source/Programs/UnrealBuildTool/Configuration/UEBuildModule.cs b/Engine/Source/Programs/UnrealBuildTool/Configuration/UEBuildModule.cs
--- a/Engine/Source/Programs/UnrealBuildTool/Configuration/UEBuildModule.cs
+++ b/Engine/Source/Programs/UnrealBuildTool/Configuration/UEBuildModule.cs
@@ -443,7 +443,11 @@ namespace UnrealBuildTool
 		/// <returns>An enumerable containing the dependencies of the module.</returns>
 		public IEnumerable<UEBuildModule> GetDirectDependencyModules()
 		{
-			return PublicDependencyModules!.Concat(PrivateDependencyModules!).Concat(DynamicallyLoadedModules!);
+			return [
+				.. PublicDependencyModules ?? [],
+				.. PrivateDependencyModules ?? [],
+				.. DynamicallyLoadedModules ?? []
+			];
 		}
 
 		/// <summary>
@@ -1228,11 +1232,11 @@ namespace UnrealBuildTool
 				{
 					if (Rules.CppCompileWarningSettings.ModuleUnsupportedWarningLevel == WarningLevel.Warning)
 					{
-						Logger.LogWarning("Warning: Referenced module '{Module}' not supported '{InvalidReason}' via {ReferenceChain}", ThisRefName, InvalidReason, ReferenceChain);
+						Logger.LogWarning("Warning: Referenced module '{Module}' not supported '{InvalidReason}' via {ReferenceChain}", ThisRefName, InvalidReason, ReferenceChain.ToString());
 					}
 					else if (Rules.CppCompileWarningSettings.ModuleUnsupportedWarningLevel == WarningLevel.Error)
 					{
-						Logger.LogError("Error: Referenced module '{Module}' not supported '{InvalidReason}' via {ReferenceChain}", ThisRefName, InvalidReason, ReferenceChain);
+						Logger.LogError("Error: Referenced module '{Module}' not supported '{InvalidReason}' via {ReferenceChain}", ThisRefName, InvalidReason, ReferenceChain.ToString());
 					}
 				}
 			}
@@ -1244,11 +1248,11 @@ namespace UnrealBuildTool
 				{
 					if (Rules.CppCompileWarningSettings.PluginModuleUnsupportedWarningLevel == WarningLevel.Warning)
 					{
-						Logger.LogWarning("Warning: Referenced plugin '{Plugin}' module '{Module}' not supported '{InvalidReason}' via {ReferenceChain}", Rules.Plugin.File.GetFileName(), ThisRefName, InvalidReason, ReferenceChain);
+						Logger.LogWarning("Warning: Referenced plugin '{Plugin}' module '{Module}' not supported '{InvalidReason}' via {ReferenceChain}", Rules.Plugin.File.GetFileName(), ThisRefName, InvalidReason, ReferenceChain.ToString());
 					}
 					else if (Rules.CppCompileWarningSettings.PluginModuleUnsupportedWarningLevel == WarningLevel.Error)
 					{
-						Logger.LogError("Error: Referenced plugin '{Plugin}' module '{Module}' not supported '{InvalidReason}' via {ReferenceChain}", Rules.Plugin.File.GetFileName(), ThisRefName, InvalidReason, ReferenceChain);
+						Logger.LogError("Error: Referenced plugin '{Plugin}' module '{Module}' not supported '{InvalidReason}' via {ReferenceChain}", Rules.Plugin.File.GetFileName(), ThisRefName, InvalidReason, ReferenceChain.ToString());
 					}
 				}
 			}
@@ -1268,7 +1272,7 @@ namespace UnrealBuildTool
 					LogDependencyNameList("Private Include Paths:", Rules.PrivateIncludePathModuleNames, Logger);
 				}
 
-				ReferenceStackParentModules ??= new(Comparer<UEBuildModule>.Create(static (a, b) => string.Compare(a.Name, b.Name, StringComparison.Ordinal)));
+				ReferenceStackParentModules ??= new(Comparer<UEBuildModule>.Create(static (a, b) => String.Compare(a.Name, b.Name, StringComparison.Ordinal)));
 
 				// Create all the dependency modules - pass through the reference stack so we can check for cycles
 				RecursivelyCreateModulesByName(Rules.PublicDependencyModuleNames, ref PublicDependencyModules, ref bDependsOnVerse, CreateModule, ReferenceChain, Logger);
