Fix for build manifest generation when compiling multiple architectures at the same time.

#rnx

#Synced-CL 50434826

diff --git a/Engine/Source/Programs/UnrealBuildTool/Configuration/UEBuildTarget.cs b/Engine/Source/Programs/UnrealBuildTool/Configuration/UEBuildTarget.cs
--- a/Engine/Source/Programs/UnrealBuildTool/Configuration/UEBuildTarget.cs
+++ b/Engine/Source/Programs/UnrealBuildTool/Configuration/UEBuildTarget.cs
@@ -2262,6 +2262,32 @@ namespace UnrealBuildTool
 			Utils.WriteClass<BuildManifest>(Manifest, ManifestPath.FullName, "", Logger);
 		}
 
+		/// <summary>
+		/// Combines multiple manifests into one. Target manifest can be included in the merge sources
+		/// </summary>
+		internal static void CombineManifests(FileReference TargetManifestPath, IEnumerable<FileReference> SourceManifestPaths, ILogger Logger)
+		{
+			HashSet<string> BuildProducts = [];
+			HashSet<string> DeployTargetFiles = [];
+
+			foreach (FileReference SourceManifestPath in SourceManifestPaths)
+			{
+				BuildManifest SourceManifest = Utils.ReadClass<BuildManifest>(SourceManifestPath.FullName, Logger);
+				if (SourceManifest != null)
+				{
+					BuildProducts.AddRange(SourceManifest.BuildProducts);
+					DeployTargetFiles.AddRange(SourceManifest.DeployTargetFiles);
+				}
+			}
+
+			BuildManifest Manifest = new BuildManifest();
+			Manifest.BuildProducts.AddRange(BuildProducts);
+			Manifest.DeployTargetFiles.AddRange(DeployTargetFiles);
+
+			Logger.LogInformation("Writing merged manifest to {ManifestPath}", TargetManifestPath);
+			Utils.WriteClass<BuildManifest>(Manifest, TargetManifestPath.FullName, "", Logger);
+		}
+
 		/// <summary>
 		/// Prepare all the module manifests for this target
 		/// </summary>
@@ -3040,7 +3066,7 @@ namespace UnrealBuildTool
 			return timestampTasks;
 		}
 
-		internal void BuildMakefile2(ISourceFileWorkingSet WorkingSet, ILogger Logger, BuildContext context, FileHashCache fileHashCache)
+		internal void BuildMakefile2(ISourceFileWorkingSet WorkingSet, ILogger Logger, BuildContext context, FileHashCache fileHashCache, IEnumerable<FileReference>? customManifestFileNames = null)
 		{
 			SourceFileMetadataCache MetadataCache = context.MetadataCache;
 			TargetDescriptor TargetDescriptor = context.Descriptor;
@@ -3365,7 +3391,7 @@ namespace UnrealBuildTool
 			}
 
 			// If we're only generating the manifest, return now
-			foreach (FileReference ManifestFileName in Rules.ManifestFileNames)
+			foreach (FileReference ManifestFileName in customManifestFileNames ?? Rules.ManifestFileNames)
 			{
 				GenerateManifest(ManifestFileName, BuildProducts, Logger);
 			}
diff --git a/Engine/Source/Programs/UnrealBuildTool/Modes/BuildMode.cs b/Engine/Source/Programs/UnrealBuildTool/Modes/BuildMode.cs
--- a/Engine/Source/Programs/UnrealBuildTool/Modes/BuildMode.cs
+++ b/Engine/Source/Programs/UnrealBuildTool/Modes/BuildMode.cs
@@ -1136,6 +1136,7 @@ namespace UnrealBuildTool
 
 				FileReference?[] makefileLocations = new FileReference[descriptorCount];
 				UEBuildTarget[] buildTargets = new UEBuildTarget[descriptorCount];
+				IReadOnlyList<FileReference>[] manifestFileNames = new IReadOnlyList<FileReference>[descriptorCount];
 
 				// Start loading unique makefiles
 				// Note, there are (very) rare instances of this being the wrong makefile
@@ -1394,8 +1395,9 @@ namespace UnrealBuildTool
 								if (buildContext != null)
 								{
 									UEBuildTarget buildTarget = buildTargets[index];
+									manifestFileNames[index] = (index == 0) ? buildTarget.Rules.ManifestFileNames : [.. buildTarget.Rules.ManifestFileNames.Select( x => x.ChangeExtension($".{index}.xml") )]; // generate unique names for the manifests for other threads
 
-									buildTarget.BuildMakefile2(workingSet, logger, buildContext, fileHashCache);
+									buildTarget.BuildMakefile2(workingSet, logger, buildContext, fileHashCache, manifestFileNames[index]);
 
 									// Save the makefile for next time
 									FileReference? makefileLocation = makefileLocations[index];
@@ -1497,6 +1499,25 @@ namespace UnrealBuildTool
 					Task.WaitAll(outTasks.ToArray());
 					outTasks.Clear();
 				}
+
+				// Merge the manifest files back together. @todo: could do better
+				for (int index = 1; index != descriptorCount; ++index)
+				{
+					if (manifestFileNames[index] != null && manifestFileNames[index].Count > 0)
+					{
+						IReadOnlyList<FileReference> targetManifests = buildTargets[index].Rules.ManifestFileNames;
+						IReadOnlyList<FileReference> workingSetManifests = manifestFileNames[index];
+						if (workingSetManifests.Count != targetManifests.Count)
+						{
+							throw new BuildException("Manifest count mismatch for {0} ... {1} != {2}", buildTargets[index].TargetName, workingSetManifests.Count, targetManifests.Count);
+						}
+						for (int manifestIndex = 0; manifestIndex < workingSetManifests.Count; ++manifestIndex)
+						{
+							FileReference[] sourceManifests = [targetManifests[manifestIndex], workingSetManifests[manifestIndex]];
+							UEBuildTarget.CombineManifests( targetManifests[manifestIndex], sourceManifests, logger);
+						}
+					}
+				}
 			}
 
 			return outTargetsPasses;
