Adding FMA intrinsics when compiling with MinimumCpuArchitectureX64.AVX2 on Clang. This fixes a compile error in the Eigen library when compiling AVX2 on Dedicated Servers using Clang. This is because when compiling using the VCToolChain using Clang, we still define _MSC_VER so the Eigen library assumes FMA is included because MSVC always includes those intrinsics when using /arch:avx2. This change brings our Clang configuration in line with MSVC when using AVX2. Note that this does not change the ffp-contract behaviour, so will not allow the compiler to optimise multiply add operations automatically.

[RN] Added FMA intrinsics when compiling with MinimumCpuArchitectureX64.AVX2 on Clang
#tests Built and tested on Dedicated Server Cooked Data, Editor and Client

#Synced-CL 50152864

diff --git a/Engine/Source/Programs/UnrealBuildTool/Platform/Clang/ClangToolChain.cs b/Engine/Source/Programs/UnrealBuildTool/Platform/Clang/ClangToolChain.cs
--- a/Engine/Source/Programs/UnrealBuildTool/Platform/Clang/ClangToolChain.cs
+++ b/Engine/Source/Programs/UnrealBuildTool/Platform/Clang/ClangToolChain.cs
@@ -876,6 +876,7 @@ namespace UnrealBuildTool
 				if (MinCpuArchX64 >= MinimumCpuArchitectureX64.AVX2)
 				{
 					Arguments.Add("-mavx2");
+					Arguments.Add("-mfma"); // Matching MSVC which includes FMA with /arch:AVX2
 					Arguments.Add("-DPLATFORM_ALWAYS_HAS_AVX_2=1");
 				}
 
diff --git a/Engine/Source/Programs/UnrealBuildTool/Platform/Windows/VCToolChain.cs b/Engine/Source/Programs/UnrealBuildTool/Platform/Windows/VCToolChain.cs
--- a/Engine/Source/Programs/UnrealBuildTool/Platform/Windows/VCToolChain.cs
+++ b/Engine/Source/Programs/UnrealBuildTool/Platform/Windows/VCToolChain.cs
@@ -1109,6 +1109,7 @@ namespace UnrealBuildTool
 					if (MinCpuArchX64 >= MinimumCpuArchitectureX64.AVX2)
 					{
 						Arguments.Add("-mavx2");
+						Arguments.Add("-mfma"); // Matching MSVC which includes FMA with /arch:AVX2
 					}
 					if (MinCpuArchX64 >= MinimumCpuArchitectureX64.AVX512)
 					{
