UnrealBuildTool: Serialize LogEvents to makefile so any messages can be rendered to the active logger if the makefile is up to date

#rnx

#Synced-CL 50547795

diff --git a/Engine/Source/Programs/Shared/EpicGames.Core/BinaryArchiveReader.cs b/Engine/Source/Programs/Shared/EpicGames.Core/BinaryArchiveReader.cs
--- a/Engine/Source/Programs/Shared/EpicGames.Core/BinaryArchiveReader.cs
+++ b/Engine/Source/Programs/Shared/EpicGames.Core/BinaryArchiveReader.cs
@@ -197,6 +197,15 @@ namespace EpicGames.Core
 			}
 		}
 
+		/// <summary>
+		/// Reads a LogEvent from the stream
+		/// </summary>
+		/// <returns>The value that was read</returns>
+		public LogEvent ReadLogEvent()
+		{
+			return LogEvent.Read(ReadByteArray());
+		}
+
 		/// <summary>
 		/// Reads a byte array from the stream
 		/// </summary>
diff --git a/Engine/Source/Programs/Shared/EpicGames.Core/BinaryArchiveWriter.cs b/Engine/Source/Programs/Shared/EpicGames.Core/BinaryArchiveWriter.cs
--- a/Engine/Source/Programs/Shared/EpicGames.Core/BinaryArchiveWriter.cs
+++ b/Engine/Source/Programs/Shared/EpicGames.Core/BinaryArchiveWriter.cs
@@ -250,6 +250,15 @@ namespace EpicGames.Core
 			WriteByteArray(bytes);
 		}
 
+		/// <summary>
+		/// Writes a LogEvent to the output
+		/// </summary>
+		/// <param name="value">Value to write</param>
+		public void WriteLogEvent(LogEvent value)
+		{
+			WriteByteArray(value.ToJsonBytes());
+		}
+
 		/// <summary>
 		/// Writes an array of bytes to the output
 		/// </summary>
diff --git a/Engine/Source/Programs/UnrealBuildTool/Modes/BuildMode.cs b/Engine/Source/Programs/UnrealBuildTool/Modes/BuildMode.cs
--- a/Engine/Source/Programs/UnrealBuildTool/Modes/BuildMode.cs
+++ b/Engine/Source/Programs/UnrealBuildTool/Modes/BuildMode.cs
@@ -1313,42 +1313,51 @@ namespace UnrealBuildTool
 								// If we couldn't load a valid makefile, create a new one
 								if (makefile == null)
 								{
-									buildTarget.PreBuildSetup(logger);
+									CaptureLogger captureLogger = new();
+									try
+									{
+										buildTarget.PreBuildSetup(captureLogger);
 
-									// Create the pre-build scripts
-									FileReference[] PreBuildScripts = buildTarget.CreatePreBuildScripts();
+										// Create the pre-build scripts
+										FileReference[] PreBuildScripts = buildTarget.CreatePreBuildScripts();
 
-									// Execute the pre-build scripts
-									if (!ranPrebuildScripts)
-									{
-										Utils.ExecuteCustomBuildSteps(PreBuildScripts, logger);
-										ranPrebuildScripts = true;
-									}
+										// Execute the pre-build scripts
+										if (!ranPrebuildScripts)
+										{
+											Utils.ExecuteCustomBuildSteps(PreBuildScripts, captureLogger);
+											ranPrebuildScripts = true;
+										}
 
-									SourceFileMetadataCache sourceFileMetaDataCache;
-									if (descriptor.ProjectFile != null)
-									{
-										sourceFileMetaDataCache = sourceFileMetadataCaches[descriptor.ProjectFile].Result;
-									}
-									else
-									{
-										sourceFileMetaDataCache = SourceFileMetadataCache.CreateHierarchy(null, logger);
-									}
+										SourceFileMetadataCache sourceFileMetaDataCache = descriptor.ProjectFile != null
+											? sourceFileMetadataCaches[descriptor.ProjectFile].Result
+											: SourceFileMetadataCache.CreateHierarchy(null, logger);
 
-									// Build the target
-									UEBuildTarget.BuildContext context = buildTarget.BuildMakefile1(buildConfiguration, descriptor, logger, sourceFileMetaDataCache);
-									buildContexts[index] = context;
+										// Build the target
+										UEBuildTarget.BuildContext context = buildTarget.BuildMakefile1(buildConfiguration, descriptor, captureLogger, sourceFileMetaDataCache);
+										buildContexts[index] = context;
 
-									makefile = context.Makefile;
+										makefile = context.Makefile;
 
-									makefile.PreBuildScripts = PreBuildScripts;
+										makefile.PreBuildScripts = PreBuildScripts;
 
-									makefile.MemoryPerActionGB = buildTarget.Rules.MemoryPerActionGB;
+										makefile.MemoryPerActionGB = buildTarget.Rules.MemoryPerActionGB;
 
-									// Save the additional command line arguments
-									makefile.AdditionalArguments = descriptor.AdditionalArguments.GetRawArray();
+										// Save the additional command line arguments
+										makefile.AdditionalArguments = descriptor.AdditionalArguments.GetRawArray();
 
-									generatedHeaderFuncs = context.GenerateHeaderFuncs;
+										generatedHeaderFuncs = context.GenerateHeaderFuncs;
+										makefile.LogEvents = [.. captureLogger.Events];
+									}
+									finally
+									{
+										captureLogger.RenderTo(logger);
+									}
+								}
+								else
+								{
+									CaptureLogger captureLogger = new();
+									captureLogger.Events.AddRange(makefile.LogEvents);
+									captureLogger.RenderTo(logger);
 								}
 
 								lock (contextsForCodeGen)
@@ -1356,7 +1365,7 @@ namespace UnrealBuildTool
 									contextsForCodeGen.Add(new UEBuildTarget.CodeGenContext()
 									{
 										Descriptor = descriptor,
-										Makefile = makefile!,
+										Makefile = makefile,
 										TargetName = buildTarget.TargetName,
 										GenerateHeaderFuncs = generatedHeaderFuncs,
 									});
diff --git a/Engine/Source/Programs/UnrealBuildTool/System/TargetMakefile.cs b/Engine/Source/Programs/UnrealBuildTool/System/TargetMakefile.cs
--- a/Engine/Source/Programs/UnrealBuildTool/System/TargetMakefile.cs
+++ b/Engine/Source/Programs/UnrealBuildTool/System/TargetMakefile.cs
@@ -306,6 +306,11 @@ namespace UnrealBuildTool
 		/// </summary>
 		public TargetInfo[] PreBuildTargets = Array.Empty<TargetInfo>();
 
+		/// <summary>
+		/// LogEvents captured during makefile creation
+		/// </summary>
+		public LogEvent[] LogEvents;
+
 		/// <summary>
 		/// Every action in the action graph
 		/// </summary>
@@ -479,6 +484,7 @@ namespace UnrealBuildTool
 			this.UbtPlugins = UbtPlugins;
 			this.EnabledUbtPlugins = EnabledUbtPlugins;
 			this.EnabledUhtPlugins = EnabledUhtPlugins;
+			LogEvents = Array.Empty<LogEvent>();
 			Actions = new List<IExternalAction>();
 			OutputItems = new List<FileItem>();
 			ModuleNameToOutputItems = new Dictionary<string, FileItem[]>(StringComparer.OrdinalIgnoreCase);
@@ -531,6 +537,7 @@ namespace UnrealBuildTool
 			UHTAdditionalArguments = Reader.ReadArray(() => Reader.ReadString())!;
 			PreBuildScripts = Reader.ReadArray(() => Reader.ReadFileReference())!;
 			PreBuildTargets = Reader.ReadArray(() => new TargetInfo(Reader))!;
+			LogEvents = Reader.ReadArray(() => Reader.ReadLogEvent())!;
 			Actions = Reader.ReadList(() => Reader.ReadAction())!;
 			EnvironmentVariables = Reader.ReadList(() => Tuple.Create(Reader.ReadString(), Reader.ReadString()))!;
 			OutputItems = Reader.ReadList(() => Reader.ReadFileItem())!;
@@ -622,6 +629,7 @@ namespace UnrealBuildTool
 			Writer.WriteArray(UHTAdditionalArguments, Item => Writer.WriteString(Item));
 			Writer.WriteArray(PreBuildScripts, Item => Writer.WriteFileReference(Item));
 			Writer.WriteArray(PreBuildTargets, Item => Item.Write(Writer));
+			Writer.WriteArray(LogEvents, Item => Writer.WriteLogEvent(Item));
 			Writer.WriteList(Actions, x => Writer.WriteAction(x));
 			Writer.WriteList(EnvironmentVariables, x =>
 			{
