When compiling with bIncludeHeaders, headers from unrelated platforms will be filtered out
#rnx
#Synced-CL 50505807

diff --git a/Engine/Source/Programs/UnrealBuildTool/Configuration/UEBuildModuleCPP.cs b/Engine/Source/Programs/UnrealBuildTool/Configuration/UEBuildModuleCPP.cs
--- a/Engine/Source/Programs/UnrealBuildTool/Configuration/UEBuildModuleCPP.cs
+++ b/Engine/Source/Programs/UnrealBuildTool/Configuration/UEBuildModuleCPP.cs
@@ -1299,10 +1299,19 @@ namespace UnrealBuildTool
 
 			HashSet<FileItem> CollidingFiles = new();
 
+			// Find all the platform folders to exclude from the list of valid modules
+			IReadOnlySet<string> ExcludeFolders = UEBuildPlatform.GetBuildPlatform(Rules.Target.Platform).GetExcludedFolderNames();
+
 			// Collect the headers that should be built
 			List<FileItem> HeaderFileItems = new();
 			foreach (FileItem HeaderFileItem in InputFiles.HeaderFiles)
 			{
+				// Skip headers for excluded platforms
+				if (HeaderFileItem.Location.ContainsAnyNames(ExcludeFolders, Unreal.EngineDirectory))
+				{
+					continue;
+				}
+
 				// We don't want to build pch files in iwyu, skip those.
 				if (HeaderFileItem == PrivatePchFileItem || HeaderFileItem == SharedPchFileItem)
 				{
