[UBT]
* Added documentation to AgentPlatform field in horde config class

diff --git a/Engine/Source/Programs/UnrealBuildTool/Executors/UnrealBuildAccelerator/UnrealBuildAcceleratorHordeConfig.cs b/Engine/Source/Programs/UnrealBuildTool/Executors/UnrealBuildAccelerator/UnrealBuildAcceleratorHordeConfig.cs
--- a/Engine/Source/Programs/UnrealBuildTool/Executors/UnrealBuildAccelerator/UnrealBuildAcceleratorHordeConfig.cs
+++ b/Engine/Source/Programs/UnrealBuildTool/Executors/UnrealBuildAccelerator/UnrealBuildAcceleratorHordeConfig.cs
@@ -237,6 +237,9 @@ namespace UnrealBuildTool
 		[CommandLine("-UBAHordeCacheDesiredConnectionCount=")]
 		public int CacheDesiredConnectionCount { get; set; } = 1;
 
+		/// <summary>
+		/// Set which platform the agents requested should have. If not set, default will be host platform 
+		/// </summary>
 		[ConfigFile(ConfigHierarchyType.Engine, "Uba.Provider.Horde", "UbaAgentPlatform")]
 		[XmlConfigFile(Category = "Horde", Name = "AgentPlatform")]
 		public string? AgentPlatform { get; set; }
