UnrealBuildTool: Add non-redistributable warning when creating modules if enabled

#rnx

#Synced-CL 50600646

diff --git a/Engine/Source/Programs/UnrealBuildTool/Configuration/Rules/ModuleRules.cs b/Engine/Source/Programs/UnrealBuildTool/Configuration/Rules/ModuleRules.cs
--- a/Engine/Source/Programs/UnrealBuildTool/Configuration/Rules/ModuleRules.cs
+++ b/Engine/Source/Programs/UnrealBuildTool/Configuration/Rules/ModuleRules.cs
@@ -1853,6 +1853,44 @@ namespace UnrealBuildTool
 			return true;
 		}
 
+		internal static bool IsValidForTarget(ModuleRules module, ReadOnlyTargetRules targetRules, [NotNullWhen(false)] out string? invalidReason)
+		{
+			// Additional checks for game/client/server based on module file path
+			if (targetRules.Type != TargetType.Editor && targetRules.Type != TargetType.Program)
+			{
+				// Check if a module is set as non-redistributable
+				if (targetRules.Configuration == UnrealTargetConfiguration.Shipping && targetRules.bCheckLicenseViolations && module.IsRedistributableOverride == false)
+				{
+					invalidReason = $"Target type '{targetRules.Type}' Configuration '{targetRules.Configuration}' and IsRedistributableOverride=false";
+					return false;
+				}
+
+				foreach (DirectoryReference ExtensionDir in Unreal.GetExtensionDirs(Unreal.EngineDirectory))
+				{
+					// Check if an engine Developer module is referenced
+					if (targetRules.Configuration == UnrealTargetConfiguration.Shipping && targetRules.bCheckLicenseViolations && module.IsRedistributableOverride == null)
+					{
+						DirectoryReference SourceDeveloperDir = DirectoryReference.Combine(ExtensionDir, "Source/Developer");
+						if (module.File.IsUnderDirectory(SourceDeveloperDir))
+						{
+							invalidReason = $"Non-Redistributable (rules file located under '{SourceDeveloperDir.MakeRelativeTo(Unreal.RootDirectory)}')";
+							return false;
+						}
+					}
+
+					// Check if an engine Editor module is referenced
+					DirectoryReference SourceEditorDir = DirectoryReference.Combine(ExtensionDir, "Source/Editor");
+					if (module.File.IsUnderDirectory(SourceEditorDir))
+					{
+						invalidReason = $"Editor module for target type '{targetRules.Type}' (rules file located under '{SourceEditorDir.MakeRelativeTo(Unreal.RootDirectory)}')";
+						return false;
+					}
+				}
+			}
+
+			return IsValidForTarget(module.GetType(), targetRules, out invalidReason);
+		}
+
 		/// <summary>
 		/// Determines if this module can be precompiled for the current target.
 		/// </summary>
diff --git a/Engine/Source/Programs/UnrealBuildTool/Configuration/UEBuildModule.cs b/Engine/Source/Programs/UnrealBuildTool/Configuration/UEBuildModule.cs
--- a/Engine/Source/Programs/UnrealBuildTool/Configuration/UEBuildModule.cs
+++ b/Engine/Source/Programs/UnrealBuildTool/Configuration/UEBuildModule.cs
@@ -1228,7 +1228,7 @@ namespace UnrealBuildTool
 			// Check if this module is invalid for the current target, based on Module attributes
 			if (Rules.CppCompileWarningSettings.ModuleUnsupportedWarningLevel != WarningLevel.Off)
 			{
-				if (!ModuleRules.IsValidForTarget(Rules.GetType(), Rules.Target, out string? InvalidReason))
+				if (!ModuleRules.IsValidForTarget(Rules, Rules.Target, out string? InvalidReason))
 				{
 					if (Rules.CppCompileWarningSettings.ModuleUnsupportedWarningLevel == WarningLevel.Warning)
 					{
diff --git a/Engine/Source/Programs/UnrealBuildTool/Configuration/UEBuildTarget.cs b/Engine/Source/Programs/UnrealBuildTool/Configuration/UEBuildTarget.cs
--- a/Engine/Source/Programs/UnrealBuildTool/Configuration/UEBuildTarget.cs
+++ b/Engine/Source/Programs/UnrealBuildTool/Configuration/UEBuildTarget.cs
@@ -3402,8 +3402,7 @@ namespace UnrealBuildTool
 				// Check the distribution level of all binaries based on the dependencies they have
 				if (ProjectFile == null && !Rules.bLegalToDistributeBinary)
 				{
-					List<DirectoryReference> RootDirectories = new List<DirectoryReference>();
-					RootDirectories.Add(Unreal.EngineDirectory);
+					List<DirectoryReference> RootDirectories = [Unreal.EngineDirectory];
 					DirectoryReference? ProjectDir = DirectoryReference.FromFile(ProjectFile);
 					if (ProjectDir != null)
 					{
@@ -3428,11 +3427,11 @@ namespace UnrealBuildTool
 						throw new BuildException("Unable to create binaries in less restricted locations than their input files.");
 					}
 				}
-
-				// Check for linking against modules prohibited by the EULA
-				CheckForEULAViolation(Logger);
 			}
 
+			// Check for referencing modules prohibited by the EULA
+			CheckForEULAViolation(Logger);
+
 			// Add all the plugins to be tracked
 			foreach (FileReference PluginFile in PluginsBase.EnumeratePlugins(ProjectFile))
 			{
