UnrealBuildTool: Update framework unzip script to fix case-sensitivity issue

#rnx
#Synced-CL 50600646

diff --git a/Engine/Source/Programs/UnrealBuildTool/Configuration/UEBuildModule.cs b/Engine/Source/Programs/UnrealBuildTool/Configuration/UEBuildModule.cs
--- a/Engine/Source/Programs/UnrealBuildTool/Configuration/UEBuildModule.cs
+++ b/Engine/Source/Programs/UnrealBuildTool/Configuration/UEBuildModule.cs
@@ -3,7 +3,6 @@
 using System;
 using System.Collections.Generic;
 using System.Diagnostics;
-using System.IO;
 using System.Linq;
 using System.Reflection;
 using System.Text;
@@ -287,11 +286,12 @@ namespace UnrealBuildTool
 					UEBuildFramework Framework;
 					if (FrameworkRules.IsZipFile())
 					{
+						FileReference ZipFile = FileReference.Combine(ModuleDirectory, FrameworkRules.Path);
 						// If FrameworkPath ends in .zip, it needs to be extracted
 						Framework = new UEBuildFramework(
 							Name: FrameworkRules.Name,
-							ZipFile: FileReference.Combine(ModuleDirectory, FrameworkRules.Path),
-							OutputDirectory: DirectoryReference.Combine(Unreal.EngineDirectory, "Intermediate", "UnzippedFrameworks", FrameworkRules.Name, Path.GetFileNameWithoutExtension(FrameworkRules.Path)),
+							ZipFile: ZipFile,
+							OutputDirectory: DirectoryReference.Combine(Unreal.EngineDirectory, "Intermediate", "UnzippedFrameworks", FrameworkRules.Name, ZipFile.GetFileNameWithoutExtension()),
 							CopyBundledAssets: FrameworkRules.CopyBundledAssets,
 							bLinkFramework, bCopyFramework, Logger);
 					}
diff --git a/Engine/Source/Programs/UnrealBuildTool/Platform/Apple/AppleToolChain.cs b/Engine/Source/Programs/UnrealBuildTool/Platform/Apple/AppleToolChain.cs
--- a/Engine/Source/Programs/UnrealBuildTool/Platform/Apple/AppleToolChain.cs
+++ b/Engine/Source/Programs/UnrealBuildTool/Platform/Apple/AppleToolChain.cs
@@ -7,7 +7,6 @@ using System.Diagnostics.CodeAnalysis;
 using System.IO;
 using System.Linq;
 using System.Security.Cryptography.X509Certificates;
-using System.Text;
 using System.Text.Json;
 using System.Text.Json.Serialization;
 using System.Text.RegularExpressions;
@@ -247,7 +246,7 @@ namespace UnrealBuildTool
 				{
 					// even in modern, we have to try to unzip in UBT, so that commandline builds extract the frameworks _before_ building (the Xcode finalize
 					// happens after building). We still need the unzip in Xcode pre-build phase, because Xcode wants the frameworks around for its build graph
-					ExtractFramework(Framework, Graph);
+					ExtractFramework(Framework, CompileEnvironment.RootPaths, Graph);
 					FrameworkTokenFiles.Add(Framework.ExtractedTokenFile!);
 				}
 			}
@@ -371,43 +370,56 @@ namespace UnrealBuildTool
 			return DestFile;
 		}
 
-		protected static FileItem ExtractFramework(UEBuildFramework Framework, IActionGraphBuilder Graph)
+		protected FileItem ExtractFramework(UEBuildFramework framework, CppRootPaths rootPaths, IActionGraphBuilder graph)
 		{
-			if (Framework.ZipFile == null)
+			if (framework.ZipFile == null)
 			{
-				throw new BuildException("Unable to extract framework '{0}' - no zip file specified", Framework.Name);
+				throw new BuildLogEventException("Unable to extract framework '{Name}': no zip file specified", framework.Name);
 			}
-			if (!Framework.bHasMadeUnzipAction)
+			if (!framework.bHasMadeUnzipAction)
 			{
-				Framework.bHasMadeUnzipAction = true;
+				framework.bHasMadeUnzipAction = true;
 
-				FileItem InputFile = FileItem.GetItemByFileReference(Framework.ZipFile);
+				rootPaths = new(rootPaths);
+				rootPaths.bUseVfs = false;
 
-				StringBuilder ExtractScript = new StringBuilder();
-				ExtractScript.AppendLine("#!/bin/sh");
-				ExtractScript.AppendLine("set -e");
-				// ExtractScript.AppendLine("set -x"); // For debugging
-				ExtractScript.AppendLine(String.Format("[ -d {0} ] && rm -rf {0}", Utils.MakePathSafeToUseWithCommandLine(Framework.ZipOutputDirectory!.FullName)));
-				ExtractScript.AppendLine(String.Format("unzip -q -o {0} -d {1}", Utils.MakePathSafeToUseWithCommandLine(Framework.ZipFile.FullName), Utils.MakePathSafeToUseWithCommandLine(Framework.ZipOutputDirectory.ParentDirectory!.FullName))); // Zip contains folder with the same name, hence ParentDirectory
-				ExtractScript.AppendLine(String.Format("touch {0}", Utils.MakePathSafeToUseWithCommandLine(Framework.ExtractedTokenFile!.AbsolutePath)));
+				FileItem inputFile = FileItem.GetItemByFileReference(framework.ZipFile);
+
+				string zipOutputDirectory = NormalizeCommandLinePath(framework.ZipOutputDirectory!, rootPaths);
+				string zipParentDirectory = NormalizeCommandLinePath(framework.ZipOutputDirectory!.ParentDirectory!, rootPaths);
+				string zipFile = NormalizeCommandLinePath(framework.ZipFile, rootPaths);
+				string extractedTokenFile = NormalizeCommandLinePath(framework.ExtractedTokenFile!, rootPaths);
+
+				IEnumerable<string> ExtractScript = [
+					$"#!{BuildHostPlatform.Current.Shell}",
+					"set -e",
+					//"set -x", // For debugging
+					$"[ -d \"{zipOutputDirectory}\" ] && rm -rf \"{zipOutputDirectory}\"",
+					$"unzip -q -o \"{zipFile}\" -d \"{zipParentDirectory}\"", // Zip contains folder with the same name, hence ParentDirectory
+					$"touch \"{extractedTokenFile}\""
+				];
+
+				// gather output files
+				IEnumerable<FileItem> producedItems = [framework.ExtractedTokenFile!];
 
 				// This is writing to folder shared by all configs.
 				//FileReference ScriptFile = FileReference.Combine(IntermediateDirectory, Framework.ZipOutputDirectory.GetDirectoryName() + ".sh");
 				//FileItem ExtractScriptFileItem = Graph.CreateIntermediateTextFile(ScriptFile, ExtractScript.ToString());
-				FileItem ExtractScriptFileItem = Graph.CreateIntermediateTextFile(new FileReference(Framework.ZipOutputDirectory.FullName + ".sh"), ExtractScript.ToString());
-
-				Action UnzipAction = Graph.CreateAction(ActionType.BuildProject);
-				UnzipAction.CommandPath = new FileReference("/bin/sh");
-				UnzipAction.CommandArguments = Utils.MakePathSafeToUseWithCommandLine(ExtractScriptFileItem.AbsolutePath);
-				UnzipAction.WorkingDirectory = Unreal.EngineDirectory;
-				UnzipAction.PrerequisiteItems.Add(InputFile);
-				UnzipAction.PrerequisiteItems.Add(ExtractScriptFileItem);
-				UnzipAction.ProducedItems.Add(Framework.ExtractedTokenFile);
-				UnzipAction.DeleteItems.Add(Framework.ExtractedTokenFile);
-				UnzipAction.StatusDescription = String.Format("Unzipping : {0} -> {1}", Framework.ZipFile, Framework.ZipOutputDirectory);
+				FileItem extractScriptFileItem = graph.CreateIntermediateTextFile(new FileReference($"{framework.ZipOutputDirectory.FullName}.sh"), ExtractScript);
+
+				Action UnzipAction = graph.CreateAction(ActionType.BuildProject);
+				UnzipAction.CommandDescription = "Unzip";
+				UnzipAction.CommandPath = BuildHostPlatform.Current.Shell;
+				UnzipAction.CommandArguments = $"\"{NormalizeCommandLinePath(extractScriptFileItem, rootPaths)}\"";
+				UnzipAction.WorkingDirectory = Unreal.EngineSourceDirectory;
+				UnzipAction.PrerequisiteItems.Add(inputFile);
+				UnzipAction.PrerequisiteItems.Add(extractScriptFileItem);
+				UnzipAction.ProducedItems.UnionWith(producedItems);
+				UnzipAction.DeleteItems.UnionWith(producedItems);
+				UnzipAction.StatusDescription = inputFile.Name;
 				UnzipAction.bCanExecuteRemotely = false;
 			}
-			return Framework.ExtractedTokenFile!;
+			return framework.ExtractedTokenFile!;
 		}
 
 		/// <summary>
diff --git a/Engine/Source/Programs/UnrealBuildTool/Platform/IOS/IOSToolChain.cs b/Engine/Source/Programs/UnrealBuildTool/Platform/IOS/IOSToolChain.cs
--- a/Engine/Source/Programs/UnrealBuildTool/Platform/IOS/IOSToolChain.cs
+++ b/Engine/Source/Programs/UnrealBuildTool/Platform/IOS/IOSToolChain.cs
@@ -433,7 +433,7 @@ namespace UnrealBuildTool
 			{
 				if (Framework.ZipFile != null)
 				{
-					FileItem ExtractedTokenFile = ExtractFramework(Framework, Graph);
+					FileItem ExtractedTokenFile = ExtractFramework(Framework, LinkEnvironment.RootPaths, Graph);
 					LinkAction.PrerequisiteItems.Add(ExtractedTokenFile);
 				}
 			}
diff --git a/Engine/Source/Programs/UnrealBuildTool/Platform/Mac/MacToolChain.cs b/Engine/Source/Programs/UnrealBuildTool/Platform/Mac/MacToolChain.cs
--- a/Engine/Source/Programs/UnrealBuildTool/Platform/Mac/MacToolChain.cs
+++ b/Engine/Source/Programs/UnrealBuildTool/Platform/Mac/MacToolChain.cs
@@ -719,7 +719,7 @@ namespace UnrealBuildTool
 				{
 					if (Framework.ZipFile != null)
 					{
-						FileItem ExtractedTokenFile = ExtractFramework(Framework, Graph);
+						FileItem ExtractedTokenFile = ExtractFramework(Framework, LinkEnvironment.RootPaths, Graph);
 						LinkAction.PrerequisiteItems.Add(ExtractedTokenFile);
 					}
 
