Adding an Apple Certificate checker to UBT so we can check that designated certificates are being used to sign Apple platform builds

#was-backout 50366409

#was-backout 50493135

#Synced-CL 50487889

diff --git a/Engine/Source/Programs/UnrealBuildTool/Platform/Apple/AppleCertificateVerification.cs b/Engine/Source/Programs/UnrealBuildTool/Platform/Apple/AppleCertificateVerification.cs
new file mode 100644
--- /dev/null
+++ b/Engine/Source/Programs/UnrealBuildTool/Platform/Apple/AppleCertificateVerification.cs
@@ -0,0 +1,106 @@
+// Copyright Epic Games, Inc. All Rights Reserved.
+
+using System;
+using System.Collections.Generic;
+using System.Security.Cryptography.X509Certificates;
+using Microsoft.Extensions.Logging;
+
+namespace UnrealBuildTool
+{
+	/// <summary>
+	/// Utility class for verifying that you have certain Apple Certificates installed
+	/// </summary>
+	public partial class AppleCertificateVerification
+	{
+		private static bool bCertificatesVerified = false;
+
+		[ConfigFile(ConfigHierarchyType.Engine, "/Script/MacTargetPlatform.XcodeProjectSettings", "CertificateHashes")]
+		private readonly List<string> CertOfTruthHashes = [];
+
+		[ConfigFile(ConfigHierarchyType.Engine, "/Script/MacTargetPlatform.XcodeProjectSettings", "CertificateTeamIDs")]
+		private readonly List<string> CertOfTruthTeamIds = [];
+
+		[ConfigFile(ConfigHierarchyType.Engine, "/Script/MacTargetPlatform.XcodeProjectSettings", "CertValidationFailMessage")]
+		private readonly string CertValidationFailMessage = "You don't have the correct cert(s) installed for auto codesigning, please install your team's cert(s)";
+
+		[ConfigFile(ConfigHierarchyType.Engine, "/Script/MacTargetPlatform.XcodeProjectSettings", "CodeSigningTeam")]
+		private readonly string CodeSigningTeamID = String.Empty;
+
+		private readonly ILogger Logger;
+
+		/// <summary>
+		/// Constructor
+		/// </summary>
+	    public AppleCertificateVerification(ILogger Logger)
+	    {
+	    	// the config entries shouldn't be in a platform-specific ini file, so we are just
+	    	// using the mac platform because we install the certs on mac and the certs are used across
+	    	// all apple platforms.
+	        ConfigCache.ReadSettings(null, UnrealTargetPlatform.Mac, this);
+
+	        this.Logger = Logger;
+	    }
+
+		private bool HasCertsOfTruthInSystemStore()
+		{
+			if (CertOfTruthHashes.Count == 0 || CertOfTruthTeamIds.Count == 0)
+			{
+				Logger.LogInformation("No certs of truth reported, skipping certificate verification...");
+				return true;
+			}
+			else if(CertOfTruthHashes.Count != CertOfTruthTeamIds.Count)
+			{
+				Logger.LogError("Invalid cert of truth config! There should be a hash in CertificateHashes for each team ID in CertificateTeamIDs!");
+				return false;
+			}
+
+			X509Certificate2Collection FoundCertCollection = AppleCodeSign.FindCertificateCollection();
+
+			for (int i = 0; CertOfTruthTeamIds.Count > i; ++i)
+			{
+				string TeamID = CertOfTruthTeamIds[i];
+				
+				if (TeamID == CodeSigningTeamID)
+				{
+					foreach (X509Certificate2 FoundCert in FoundCertCollection)
+					{
+						if (FoundCert.GetCertHashString() == CertOfTruthHashes[i])
+						{
+							Logger.LogDebug($"Found a locally installed Apple certificate matching known Apple certificate of truth.");
+							return true;
+						}
+					}
+
+					return false;
+				}
+			}
+
+			Logger.LogInformation("No certs of truth matching current team id {CodeSigningTeamID} reported, skipping certificate verification...",
+								  CodeSigningTeamID);
+
+			return true;
+		}
+
+		/// <summary>
+		/// Verifies you have the valid certificate for your associated team installed and
+		/// fails the build if not.
+		/// </summary> 
+		public void VerifyCertificateInstallation()
+		{
+			if (!bCertificatesVerified)
+			{
+				if (HasCertsOfTruthInSystemStore())
+				{
+					Logger.LogInformation("Known Apple certificate for this team is installed!");
+				}
+				else
+				{
+					Logger.LogError("{CertValidationFailMessage}", CertValidationFailMessage);
+					throw new BuildException("Stopping build due to broken Apple certificate setup...");
+				}
+
+				bCertificatesVerified = true;
+			}
+		}
+	}
+}
diff --git a/Engine/Source/Programs/UnrealBuildTool/Platform/Apple/AppleToolChain.cs b/Engine/Source/Programs/UnrealBuildTool/Platform/Apple/AppleToolChain.cs
--- a/Engine/Source/Programs/UnrealBuildTool/Platform/Apple/AppleToolChain.cs
+++ b/Engine/Source/Programs/UnrealBuildTool/Platform/Apple/AppleToolChain.cs
@@ -1077,6 +1077,12 @@ namespace UnrealBuildTool
 						string DestPlatform = AppleExports.GetDestinationPlatform(Platform, Architectures);
 						DestinationParam = String.Join(" ", DestinationIds.Select(x => $"-destination \"platform={DestPlatform},{IdOrNameParam(x)}\""));
 					}
+
+					if(bUseAutomaticCodeSigning && !Unreal.IsBuildMachine())
+					{
+						(new AppleCertificateVerification(Logger)).VerifyCertificateInstallation();
+					}
+
 					List<string> Arguments = new()
 					{
 						"UBT_NO_POST_DEPLOY=true",
diff --git a/Engine/Source/Programs/UnrealBuildTool/Platform/IOS/IOSCodeSigning.cs b/Engine/Source/Programs/UnrealBuildTool/Platform/IOS/IOSCodeSigning.cs
--- a/Engine/Source/Programs/UnrealBuildTool/Platform/IOS/IOSCodeSigning.cs
+++ b/Engine/Source/Programs/UnrealBuildTool/Platform/IOS/IOSCodeSigning.cs
@@ -85,9 +85,9 @@ namespace UnrealBuildTool
 		}
 
 		/// <summary>
-		/// Finds all valid installed certificates
+		/// Finds all installed certificates with known Apple development/distro prefixes as a X509Certificate2Collection
 		/// </summary>
-		public static List<string> FindCertificates()
+		public static X509Certificate2Collection FindCertificateCollection()
 		{
 			string[] ValidCertificatePrefixes = { "iPhone Developer", "iPhone Distribution", "Apple Development", "Apple Distribution" };
 
@@ -141,6 +141,16 @@ namespace UnrealBuildTool
 				Store.Close();
 			}
 
+			return FoundCerts;
+		}
+
+		/// <summary>
+		/// Finds all valid installed certificates
+		/// </summary>
+		public static List<string> FindCertificates()
+		{
+			X509Certificate2Collection FoundCerts = FindCertificateCollection();
+			
 			List<string> Certs = new();
 			foreach (X509Certificate2 TestCert in FoundCerts)
 			{
