UnrealBuildTool: Remove default enable of enum type conformance

#rnx

#Synced-CL 50684155

#Synced-CL 50684155

diff --git a/Engine/Source/Programs/UnrealBuildTool/Configuration/Rules/TargetRules.cs b/Engine/Source/Programs/UnrealBuildTool/Configuration/Rules/TargetRules.cs
--- a/Engine/Source/Programs/UnrealBuildTool/Configuration/Rules/TargetRules.cs
+++ b/Engine/Source/Programs/UnrealBuildTool/Configuration/Rules/TargetRules.cs
@@ -3715,7 +3715,6 @@ namespace UnrealBuildTool
 
 				if (BuildSettingsVersion.V7 <= latestVersion && DefaultBuildSettings < BuildSettingsVersion.V7)
 				{
-					modifiedSettings.Add(Tuple.Create($"{nameof(WindowsPlatform)}.{nameof(WindowsPlatform.bStrictEnumTypesConformance)} = true", "Updates MSVC enum types conformance to true (Previously: false)."));
 					modifiedSettings.Add(Tuple.Create($"{nameof(CppCompileWarningSettings.ReturnTypeWarningLevel)} = WarningLevel.Error", "Validates that non-void functions return a value. (Previously: Off for VCClang, Error for MSVC and non-VC Clang)."));
 					modifiedSettings.Add(Tuple.Create($"{nameof(CppCompileWarningSettings.DanglingWarningLevel)} = WarningLevel.Error", "Enables clang warnings related possible dangling references or pointers. (Previously: Off for VCClang, Error for non-VC Clang)."));
 				}
diff --git a/Engine/Source/Programs/UnrealBuildTool/Platform/Windows/UEBuildWindows.cs b/Engine/Source/Programs/UnrealBuildTool/Platform/Windows/UEBuildWindows.cs
--- a/Engine/Source/Programs/UnrealBuildTool/Platform/Windows/UEBuildWindows.cs
+++ b/Engine/Source/Programs/UnrealBuildTool/Platform/Windows/UEBuildWindows.cs
@@ -662,12 +662,7 @@ namespace UnrealBuildTool
 		[XmlConfigFile(Category = "WindowsPlatform")]
 		[CommandLine("-StrictEnumTypes")]
 		[CommandLine("-NoStrictEnumTypes", Value = "false")]
-		public bool bStrictEnumTypesConformance
-		{
-			get => bStrictEnumTypesConformancePrivate ?? Target.DefaultBuildSettings >= BuildSettingsVersion.V7;
-			set => bStrictEnumTypesConformancePrivate = value;
-		}
-		private bool? bStrictEnumTypesConformancePrivate;
+		public bool bStrictEnumTypesConformance { get; set; } = false;
 
 		/// <summary>
 		/// Enables enforcing standard C++ ODR violations (/Zc:checkGwOdr).
