UnrealBuildTool: Fix a bad path combining in IOSCodeSigning that caused mobile provision files to be copied to the wrong place.

#Synced-CL 50711080

diff --git a/Engine/Source/Programs/UnrealBuildTool/Platform/IOS/IOSCodeSigning.cs b/Engine/Source/Programs/UnrealBuildTool/Platform/IOS/IOSCodeSigning.cs
--- a/Engine/Source/Programs/UnrealBuildTool/Platform/IOS/IOSCodeSigning.cs
+++ b/Engine/Source/Programs/UnrealBuildTool/Platform/IOS/IOSCodeSigning.cs
@@ -817,7 +817,7 @@ namespace UnrealBuildTool
 					_logger?.LogInformation("Finding provisions in {ProvisionCopySrcDirectory}", ProvisionCopySrcDirectory);
 					foreach (string Provision in Directory.EnumerateFiles(ProvisionCopySrcDirectory, "*.mobileprovision", SearchOption.AllDirectories))
 					{
-						string TargetFile = ProvisionDirectory + Path.GetFileName(Provision);
+						string TargetFile = Path.Combine(ProvisionDirectory, Path.GetFileName(Provision));
 						if (!File.Exists(TargetFile) || File.GetLastWriteTime(TargetFile) < File.GetLastWriteTime(Provision))
 						{
 							FileInfo DestFileInfo;
