UnrealBuildTool: Update bAllowUbaCompression docstring as it works with all compilers now

#rnx

#Synced-CL 50718640

diff --git a/Engine/Source/Programs/UnrealBuildTool/Configuration/Rules/TargetRules.cs b/Engine/Source/Programs/UnrealBuildTool/Configuration/Rules/TargetRules.cs
--- a/Engine/Source/Programs/UnrealBuildTool/Configuration/Rules/TargetRules.cs
+++ b/Engine/Source/Programs/UnrealBuildTool/Configuration/Rules/TargetRules.cs
@@ -2205,8 +2205,7 @@ namespace UnrealBuildTool
 		public bool bWithAssembly { get; set; }
 
 		/// <summary>
-		/// Experimental: Store object (.obj) compressed on disk. Requires UBA to link, currently MSVC only. Toggling this flag will invalidate MSVC actions.
-		/// Warning, this option is not currently compatitable with PGO or the the cl-clang linker as those are not detoured and linking will fail.
+		/// Experimental: Store object (.obj) compressed on disk. Requires UBA to link. Toggling this flag will invalidate actions.
 		/// </summary>
 		[XmlConfigFile(Category = "UnrealBuildAccelerator", Name = "bStoreObjFilesCompressed")]
 		[ConfigFile(ConfigHierarchyType.Engine, "/Script/BuildSettings.BuildSettings", "bStoreObjFilesCompressed")]
