Mark symbols hidden on Linux even for monolithic builds.

#Synced-CL 50777006

diff --git a/Engine/Source/Programs/UnrealBuildTool/Platform/Linux/LinuxToolChain.cs b/Engine/Source/Programs/UnrealBuildTool/Platform/Linux/LinuxToolChain.cs
--- a/Engine/Source/Programs/UnrealBuildTool/Platform/Linux/LinuxToolChain.cs
+++ b/Engine/Source/Programs/UnrealBuildTool/Platform/Linux/LinuxToolChain.cs
@@ -397,9 +397,20 @@ namespace UnrealBuildTool
 
 			if (CompileEnvironment.bHideSymbolsByDefault)
 			{
-				if (!bMergeModules) // With merge modules we use version-script to control visibility instead.
+				// With merge modules we use version-script to control visibility instead.
+				if (!bMergeModules)
 				{
-					Arguments.Add("-fvisibility-ms-compat");
+					if (CompileEnvironment.bIsBuildingDLL)
+					{
+						Arguments.Add("-fvisibility-ms-compat");
+					}
+					else
+					{
+						// This is a *fun* one, but if we aren't building a DLL lets just mark
+						// all symbols hidden. This avoids us putting symbols into symbol tables
+						// when the only symbol that matters for monolithic is main!
+						Arguments.Add("-fvisibility=hidden");
+					}
 				}
 
 				Arguments.Add("-fvisibility-inlines-hidden");
