adding check for interfering apple identities in AppleCertificateVerification

#Synced-CL 50718145

diff --git a/Engine/Source/Programs/UnrealBuildTool/Platform/Apple/AppleCertificateVerification.cs b/Engine/Source/Programs/UnrealBuildTool/Platform/Apple/AppleCertificateVerification.cs
--- a/Engine/Source/Programs/UnrealBuildTool/Platform/Apple/AppleCertificateVerification.cs
+++ b/Engine/Source/Programs/UnrealBuildTool/Platform/Apple/AppleCertificateVerification.cs
@@ -3,6 +3,7 @@
 using System;
 using System.Collections.Generic;
 using System.Security.Cryptography.X509Certificates;
+using System.Text.RegularExpressions;
 using Microsoft.Extensions.Logging;
 
 namespace UnrealBuildTool
@@ -41,6 +42,40 @@ namespace UnrealBuildTool
 	        this.Logger = Logger;
 	    }
 
+		/// <summary>
+		/// Checks for certs that don't match the "certificates of truth" hashes but match the team ID
+		/// (and thus are certs which will interfere with xcode's autosigning process)
+		/// </summary>
+	    private bool FoundInterferingCerts(X509Certificate2Collection CertCollection)
+	    {
+	    	Regex CertOrgRegex = new Regex(@"OU=([^,]+)", RegexOptions.IgnoreCase);
+			Regex CertCommonNameRegex = new Regex(@"CN=([^,]+)", RegexOptions.IgnoreCase);
+
+			foreach (X509Certificate2 FoundCert in CertCollection)
+			{
+				Match CertOrgMatch = CertOrgRegex.Match(FoundCert.Subject);
+				string? CertOrgVal = CertOrgMatch.Success ? CertOrgMatch.Groups[1].Value.Trim() : null;
+
+				Match CertCommonNameMatch = CertCommonNameRegex.Match(FoundCert.Subject);
+				string? CertCommonNameVal = CertCommonNameMatch.Success ? CertCommonNameMatch.Groups[1].Value.Trim() : null;
+
+				if (String.IsNullOrEmpty(CertOrgVal) || String.IsNullOrEmpty(CertCommonNameVal))
+				{
+					continue;
+				}
+
+				if (CertCommonNameVal.StartsWith("Apple Development: ", StringComparison.Ordinal) && CodeSigningTeamID == CertOrgVal &&
+				    !CertOfTruthHashes.Contains(FoundCert.GetCertHashString()))
+				{
+					Logger.LogError("Found interfering cert identity: {FoundCert.Subject}!", FoundCert.Subject);
+				
+					return true;
+				}
+			}
+
+			return false;
+	    }
+
 		private bool HasCertsOfTruthInSystemStore()
 		{
 			if (CertOfTruthHashes.Count == 0 || CertOfTruthTeamIds.Count == 0)
@@ -67,7 +102,8 @@ namespace UnrealBuildTool
 						if (FoundCert.GetCertHashString() == CertOfTruthHashes[i])
 						{
 							Logger.LogDebug($"Found a locally installed Apple certificate matching known Apple certificate of truth.");
-							return true;
+
+							return !FoundInterferingCerts(FoundCertCollection);
 						}
 					}
 
