update experimental Windows project generation to match what RemoteWin supports
- default deployment folder is now derived from the target name not the project name

#rnx

#Synced-CL 50642529

diff --git a/Engine/Source/Programs/AutomationTool/Win/RemoteWinSupport.cs b/Engine/Source/Programs/AutomationTool/Win/RemoteWinSupport.cs
--- a/Engine/Source/Programs/AutomationTool/Win/RemoteWinSupport.cs
+++ b/Engine/Source/Programs/AutomationTool/Win/RemoteWinSupport.cs
@@ -260,7 +260,21 @@ public static class RemoteWinSupport
 
 	private static string GetDestinationRoot(UnrealTargetPlatform Platform, ProjectParams Params)
 	{
-		return $"{Params.ShortProjectName}_{Platform}";
+		string TargetName;
+		if (Params.HasClientCookedTargets)
+		{
+			TargetName = Params.ClientCookedTargets[0];
+		}
+		else if (Params.HasServerCookedTargets)
+		{
+			TargetName = Params.ServerCookedTargets[0];
+		}
+		else
+		{
+			TargetName = Params.ShortProjectName;
+		}
+
+		return (Platform == UnrealTargetPlatform.Win64) ? TargetName : $"{TargetName}_{Platform}";
 	}
 
 
diff --git a/Engine/Source/Programs/UnrealBuildTool/Platform/Windows/UEBuildWindows.cs b/Engine/Source/Programs/UnrealBuildTool/Platform/Windows/UEBuildWindows.cs
--- a/Engine/Source/Programs/UnrealBuildTool/Platform/Windows/UEBuildWindows.cs
+++ b/Engine/Source/Programs/UnrealBuildTool/Platform/Windows/UEBuildWindows.cs
@@ -849,12 +849,6 @@ namespace UnrealBuildTool
 		[CommandLine("-UseXCurl")]
 		public bool bUseXCurl = false;
 
-		/// <summary>
-		/// Whether to enable the experimental remote debugging configuration.
-		/// </summary>
-		[XmlConfigFile(Category = "WindowsPlatform")]
-		public bool bEnableExperimentalRemoteDebugging = false;
-
 		/// <summary>
 		/// Settings for PVS studio
 		/// </summary>
@@ -1127,8 +1121,6 @@ namespace UnrealBuildTool
 
 		public bool bUseXCurl => Inner.bUseXCurl;
 
-		public bool bEnableExperimentalRemoteDebugging => Inner.bEnableExperimentalRemoteDebugging;
-
 		public ReadOnlyPVSTargetSettings PVS
 		{
 			get; private set;
diff --git a/Engine/Source/Programs/UnrealBuildTool/Platform/Windows/VCToolChain.cs b/Engine/Source/Programs/UnrealBuildTool/Platform/Windows/VCToolChain.cs
--- a/Engine/Source/Programs/UnrealBuildTool/Platform/Windows/VCToolChain.cs
+++ b/Engine/Source/Programs/UnrealBuildTool/Platform/Windows/VCToolChain.cs
@@ -3896,7 +3896,6 @@ namespace UnrealBuildTool
 		{
 			// do not strip the PDB path from the executable if we want to remotely debug it (otherwise VS can't find the symbols)
 			if (LinkEnvironment.bCreateDebugInfo &&
-				Target.WindowsPlatform.bEnableExperimentalRemoteDebugging &&
 				Target.Type != TargetType.Editor &&
 				Target.Configuration != UnrealTargetConfiguration.Shipping &&
 				Target.Platform.IsInGroup(UnrealPlatformGroup.Windows))
diff --git a/Engine/Source/Programs/UnrealBuildTool/Platform/Windows/WindowsProjectGenerator.cs b/Engine/Source/Programs/UnrealBuildTool/Platform/Windows/WindowsProjectGenerator.cs
--- a/Engine/Source/Programs/UnrealBuildTool/Platform/Windows/WindowsProjectGenerator.cs
+++ b/Engine/Source/Programs/UnrealBuildTool/Platform/Windows/WindowsProjectGenerator.cs
@@ -15,55 +15,42 @@ namespace UnrealBuildTool
 	class WindowsProjectGenerator : PlatformProjectGenerator
 	{
 		/// <summary>
-		/// Whether to enable the experimental remote debugging support. 
-		/// This will configure the project settings for remote debugging based on the fields below.
+		/// This will configure the Visual Studio project settings for remote debugging based on the fields below.
 		/// Refer to https://learn.microsoft.com/en-us/visualstudio/debugger/remote-debugging-cpp for details on setting up Visual Studio and the remote device for remote debugging.
-		/// </summary>
-		protected bool bEnableExperimentalRemoteDebugging = false;
-
-		/// <summary>
-		/// When disabled, the remote debug directory is expected to be the root of a staged build.
-		/// When enabled, the project name and target will be appended to the remote debug directory, allowing multiple projects to be more easily debugged.
-		/// ...this can be used alongside RunUAT BuildCookRun [...] -archive -archivedirectory=\\remotepc\Share\MyProject\
+		/// 
+		/// The is also dependent on the new remote Windows dev tools. See https://aka.ms/GameRemoteDevtools and assumes the remote device has the game deployed already
 		/// </summary>
 		[XmlConfigFile(Category = "WindowsPlatform")]
-		protected bool bPerProjectRemoteDebugPaths = true;
+		protected bool bConfigureVSRemoteDebugger = true;
 
 		/// <summary>
-		/// Remote machine name for Windows X64 remote debugging
+		/// Default remote deployment path for Windows Remote Deployment. Must match wdendpoint
 		/// </summary>
 		[XmlConfigFile(Category = "WindowsPlatform")]
-		protected string? RemoteDebugMachineX64;
+		protected string RemoteWinRoot = @"C:\ProgramData\Microsoft GDK\gameroot\";
 
 		/// <summary>
-		/// Local shared folder path on the remote Windows X64 remote machine to deploy to. 
+		/// Default remote machine name for Windows X64 remote debugging. Must previously be paired
 		/// </summary>
 		[XmlConfigFile(Category = "WindowsPlatform")]
-		protected string? RemoteDebugDirectoryX64;
+		protected string? RemoteDebugMachineX64;
 
 		/// <summary>
-		/// Remote machine name for Windows Arm64 remote debugging
+		/// Default remote machine name for Windows Arm64 and Arm64ec remote debugging. Must previously be paired
 		/// </summary>
 		[XmlConfigFile(Category="WindowsPlatform")]
 		protected string? RemoteDebugMachineArm64;
 
 		/// <summary>
-		/// Local shared folder path on the remote Windows Arm64 remote machine to deploy to. 
-		/// </summary>
-		[XmlConfigFile(Category = "WindowsPlatform")]
-		protected string? RemoteDebugDirectoryArm64;
-
-		/// <summary>
-		/// Remote machine name for Windows Arm64ec remote debugging
+		/// Whether remote debugging should use authentication or not
 		/// </summary>
-		[XmlConfigFile(Category = "WindowsPlatform")]
-		protected string? RemoteDebugMachineArm64ec;
+		[XmlConfigFile(Category="WindowsPlatform")]
+		protected bool bRemoteDebuggerAuthentication = false;
 
 		/// <summary>
-		/// Local shared folder path on the remote Windows Arm64ec remote machine to deploy to. 
+		/// Location of the windows remote deployment tools (not user-configurable)
 		/// </summary>
-		[XmlConfigFile(Category = "WindowsPlatform")]
-		protected string? RemoteDebugDirectoryArm64ec;
+		private static string WdRemotePath => System.Environment.ExpandEnvironmentVariables("%LocalAppData%\\Microsoft\\WinGet\\Links\\wdRemote.exe");
 
 		/// <summary>
 		/// Constructor
@@ -76,12 +63,6 @@ namespace UnrealBuildTool
 #pragma warning restore IDE0060
 		{
 			XmlConfig.ApplyTo(this);
-
-			// must read this value directly because it exists in the windows target rules too
-			if (XmlConfig.TryGetValue(typeof(WindowsTargetRules), "bEnableExperimentalRemoteDebugging", out object? ConfigObject))
-			{
-				bEnableExperimentalRemoteDebugging = (bool)ConfigObject;
-			}
 		}
 
 		/// <inheritdoc/>
@@ -145,43 +126,30 @@ namespace UnrealBuildTool
 		{
 			base.GetVisualStudioPathsEntries(InVSSettings, TargetType, TargetRulesPath, ProjectFilePath, NMakeOutputPath, ProjectFileBuilder);
 
-			if (bEnableExperimentalRemoteDebugging && 
-				TargetType != TargetType.Editor &&
-				GetRemoteDebugProperties(InVSSettings.Architecture, out string? RemoteDebugMachine, out string? RemoteDebugDirectory))
+			if (bConfigureVSRemoteDebugger && 
+				(TargetType == TargetType.Game || TargetType == TargetType.Client || TargetType == TargetType.Server) &&
+				GetRemoteDebugProperties(InVSSettings, out string? RemoteDebugMachine))
 			{			
-				string RemoteDebugWorkingDir = RemoteDebugDirectory;
-
-				// can optionally use a subdirectory within the shared folder to make it simpler to debug multiple projects
-				if (bPerProjectRemoteDebugPaths)
-				{
-					string ProjectName = (TargetType == TargetType.Program)
-						? NMakeOutputPath.Directory.GetDirectoryName()
-						: NMakeOutputPath.Directory!.ParentDirectory!.ParentDirectory!.GetDirectoryName();
-
-					string PlatformTargetName = (TargetType == TargetType.Program) ? "" : "Windows";
-					if (TargetType != TargetType.Game && TargetType != TargetType.Program)
-					{
-						PlatformTargetName += TargetType.ToString();
-					}
-
-					DirectoryReference BaseDir = NMakeOutputPath.Directory!.ParentDirectory!.ParentDirectory!.ParentDirectory!;
-					string RelativeWorkingDir = NMakeOutputPath.Directory.MakeRelativeTo(BaseDir);
-					RemoteDebugWorkingDir = Path.Combine(RemoteDebugDirectory, ProjectName, PlatformTargetName, RelativeWorkingDir);
-				}
+				string TargetName = TargetRulesPath.GetFileNameWithoutAnyExtensions();
+				string RemoteDeployFolder = (InVSSettings.Platform == UnrealTargetPlatform.Win64) ? TargetName : $"{TargetName}_{InVSSettings.Platform}";
+				DirectoryReference BaseDir = DirectoryReference.Combine(NMakeOutputPath.Directory, @"..\..\..");
 
-				string OutDir = ProjectFile.NormalizeProjectPath(NMakeOutputPath.Directory.FullName); // need to override OutDir for remote debugging so that VS deploys the executable correctly
-				string RemoteDebugCommand = Path.Combine(RemoteDebugWorkingDir, NMakeOutputPath.GetFileName());
+				DirectoryReference RemoteWorkingDir = DirectoryReference.Combine( new DirectoryReference(RemoteWinRoot), RemoteDeployFolder);
+				FileReference RemoteDebugCommand = FileReference.Combine(RemoteWorkingDir, NMakeOutputPath.Directory.MakeRelativeTo(BaseDir), NMakeOutputPath.GetFileName());
+				string RemoteAuthentication = bRemoteDebuggerAuthentication ? "RemoteWithAuthentication" : "RemoteWithoutAuthentication";
+				string OutDir = ProjectFile.NormalizeProjectPath(BaseDir); // need to override OutDir for remote debugging so that VS deploys the executable relative to the project root, not the engine root
 
 				ProjectFileBuilder.AppendLine($"    <OutDir>{OutDir}</OutDir>");
-				ProjectFileBuilder.AppendLine($"    <DeploymentDirectory>{RemoteDebugWorkingDir}</DeploymentDirectory>");
+				ProjectFileBuilder.AppendLine($"    <DeploymentDirectory>{RemoteWorkingDir}</DeploymentDirectory>");
 				ProjectFileBuilder.AppendLine($"    <RemoteDebuggerCommand>{RemoteDebugCommand}</RemoteDebuggerCommand>");
-				ProjectFileBuilder.AppendLine($"    <RemoteDebuggerWorkingDirectory>{RemoteDebugWorkingDir}</RemoteDebuggerWorkingDirectory>");
+				ProjectFileBuilder.AppendLine($"    <RemoteDebuggerWorkingDirectory>{RemoteDebugCommand.Directory}</RemoteDebuggerWorkingDirectory>");
 				ProjectFileBuilder.AppendLine($"    <RemoteDebuggerServerName>{RemoteDebugMachine}</RemoteDebuggerServerName>");
 				ProjectFileBuilder.AppendLine($"    <RemoteDebuggerDebuggerType>NativeOnly</RemoteDebuggerDebuggerType>");
 				ProjectFileBuilder.AppendLine($"    <RemoteDebuggerDeployCppRuntime>false</RemoteDebuggerDeployCppRuntime>");
+				ProjectFileBuilder.AppendLine($"    <RemoteDebuggerConnection>{RemoteAuthentication}</RemoteDebuggerConnection>");
 
 				// default to remote debugging for other architectures
-				if (InVSSettings.Architecture != null && InVSSettings.Architecture.Value != UnrealArch.Host.Value)
+				if (InVSSettings.Architecture != null && InVSSettings.Architecture.Value.bIsX64 != UnrealArch.Host.Value.bIsX64)
 				{
 					ProjectFileBuilder.AppendLine($"    <DebuggerFlavor>WindowsRemoteDebugger</DebuggerFlavor>");
 				}
@@ -192,7 +160,7 @@ namespace UnrealBuildTool
 		public override bool GetVisualStudioDeploymentEnabled(VSSettings InVSSettings)
 		{
 			// remote debugging requires deployment
-			if (bEnableExperimentalRemoteDebugging && GetRemoteDebugProperties(InVSSettings.Architecture, out _, out _))
+			if (bConfigureVSRemoteDebugger && GetRemoteDebugProperties(InVSSettings, out _))
 			{
 				return true;
 			}
@@ -218,31 +186,42 @@ namespace UnrealBuildTool
 			return Result;
 		}
 
-		protected virtual bool GetRemoteDebugProperties(UnrealArch? Architecture, [NotNullWhen(true)] out string? RemoteDebugMachine, [NotNullWhen(true)] out string? RemoteDebugDirectory)
+		protected virtual bool GetRemoteDebugProperties(VSSettings InVSSettings, [NotNullWhen(true)] out string? RemoteDebugMachine)
 		{
 			RemoteDebugMachine = null;
-			RemoteDebugDirectory = null;
 
-			if (Architecture != null)
+			// check that the wdremote tool exists
+			if (!File.Exists(WdRemotePath))
 			{
-				if (Architecture == UnrealArch.X64)
+				return false;
+			}
+
+			// check for the default remote device for the current architecture
+			if (InVSSettings.Architecture != null)
+			{
+				if (((UnrealArch)InVSSettings.Architecture).bIsX64)
 				{
 					RemoteDebugMachine = RemoteDebugMachineX64;
-					RemoteDebugDirectory = RemoteDebugDirectoryX64;
 				}
-				else if (Architecture == UnrealArch.Arm64)
+				else
 				{
 					RemoteDebugMachine = RemoteDebugMachineArm64;
-					RemoteDebugDirectory = RemoteDebugDirectoryArm64;
 				}
-				else if (Architecture == UnrealArch.Arm64ec)
+			}
+
+			// if no specific remote machine was specified, default to the first one in the list of paired devices
+			if (RemoteDebugMachine == null)
+			{
+				// read device names from the engine ini. must already be paired
+				// best place is %LOCALAPPDATA%\Unreal Engine\Engine\Config\UserEngine.ini
+				ConfigHierarchy EngineConfig = ConfigCache.ReadHierarchy(ConfigHierarchyType.Engine, null, InVSSettings.Platform);
+				if (EngineConfig.GetArray("RemoteWin", "DeviceNames", out List<string>? DeviceNames) && DeviceNames != null && DeviceNames.Count > 0)
 				{
-					RemoteDebugMachine = RemoteDebugMachineArm64ec;
-					RemoteDebugDirectory = RemoteDebugDirectoryArm64ec;
+					RemoteDebugMachine = DeviceNames[0];
 				}
 			}
 
-			return !System.String.IsNullOrEmpty(RemoteDebugMachine) && !System.String.IsNullOrEmpty(RemoteDebugDirectory);
+			return !System.String.IsNullOrEmpty(RemoteDebugMachine);
 		}
 	}
 }
