Adding a unique build environment for when profiling compiles, -ProfilingCompile

#Synced-CL 50688484

diff --git a/Engine/Source/Programs/UnrealBuildTool/Configuration/Descriptors/OriginalTargetDescriptor.cs b/Engine/Source/Programs/UnrealBuildTool/Configuration/Descriptors/OriginalTargetDescriptor.cs
--- a/Engine/Source/Programs/UnrealBuildTool/Configuration/Descriptors/OriginalTargetDescriptor.cs
+++ b/Engine/Source/Programs/UnrealBuildTool/Configuration/Descriptors/OriginalTargetDescriptor.cs
@@ -95,6 +95,9 @@ internal record class OriginalTargetDescriptor
 	/// <inheritdoc cref="TargetDescriptor.bIWYU"/>
 	public bool bIWYU { get; init; }
 
+	/// <inheritdoc cref="TargetDescriptor.bProfilingCompile"/>
+	public bool bProfilingCompile { get; init; }
+
 	/// <inheritdoc cref="TargetDescriptor.IntermediateEnvironment"/>
 	public UnrealIntermediateEnvironment IntermediateEnvironment { get; init; }
 
@@ -133,6 +136,7 @@ internal record class OriginalTargetDescriptor
 			bUseUnityBuild = targetDescriptor.bUseUnityBuild,
 			bForceUnityBuild = targetDescriptor.bForceUnityBuild,
 			bIWYU = targetDescriptor.bIWYU,
+			bProfilingCompile = targetDescriptor.bProfilingCompile,
 			IntermediateEnvironment = targetDescriptor.IntermediateEnvironment,
 		};
 	}
diff --git a/Engine/Source/Programs/UnrealBuildTool/Configuration/Descriptors/TargetDescriptor.cs b/Engine/Source/Programs/UnrealBuildTool/Configuration/Descriptors/TargetDescriptor.cs
--- a/Engine/Source/Programs/UnrealBuildTool/Configuration/Descriptors/TargetDescriptor.cs
+++ b/Engine/Source/Programs/UnrealBuildTool/Configuration/Descriptors/TargetDescriptor.cs
@@ -150,6 +150,12 @@ namespace UnrealBuildTool
 		[CommandLine("-IWYU")]
 		public bool bIWYU { get; set; }
 
+		/// <summary>
+		/// Set when trying to profile a compile to make sure the environment is setup correctly.
+		/// </summary>
+		[CommandLine("-ProfilingCompile")]
+		public bool bProfilingCompile { get; set; }
+
 		/// <summary>
 		/// Intermediate environment. Determines if the intermediates end up in a different folder than normal.
 		/// </summary>
@@ -169,6 +175,10 @@ namespace UnrealBuildTool
 				{
 					return UnrealIntermediateEnvironment.NonUnity;
 				}
+				if (bProfilingCompile)
+				{
+					return UnrealIntermediateEnvironment.ProfilingCompile;
+				}
 				return UnrealIntermediateEnvironment.Default;
 			}
 			set => IntermediateEnvironmentOverride = value;
diff --git a/Engine/Source/Programs/UnrealBuildTool/Configuration/UEBuildTarget.cs b/Engine/Source/Programs/UnrealBuildTool/Configuration/UEBuildTarget.cs
--- a/Engine/Source/Programs/UnrealBuildTool/Configuration/UEBuildTarget.cs
+++ b/Engine/Source/Programs/UnrealBuildTool/Configuration/UEBuildTarget.cs
@@ -1208,6 +1208,11 @@ namespace UnrealBuildTool
 		/// Generate project files
 		/// </summary>
 		GenerateProjectFiles,
+
+		/// <summary>
+		/// Profiling compilation can modify a number of compilation settings.
+		/// </summary>
+		ProfilingCompile,
 	}
 
 	/// <summary>
@@ -2016,6 +2021,9 @@ namespace UnrealBuildTool
 				case UnrealIntermediateEnvironment.GenerateProjectFiles:
 					TargetFolderName += "GPF";
 					break;
+				case UnrealIntermediateEnvironment.ProfilingCompile:
+					TargetFolderName += "Prof";
+					break;
 			}
 			return TargetFolderName;
 		}
