Even with instrumentation disabled we still need to depend on the AutoRTFM runtime if we are using our compiler.

This is because we gate whether to use the runtime by `UE_AUTORTFM`, not `UE_AUTORTFM_ENABLED`, which is correct because we want instrumentation enabled/disabled code to see the same definitions of our symbols.

#Synced-CL 50814375

diff --git a/Engine/Source/Programs/UnrealBuildTool/Configuration/UEBuildTarget.cs b/Engine/Source/Programs/UnrealBuildTool/Configuration/UEBuildTarget.cs
--- a/Engine/Source/Programs/UnrealBuildTool/Configuration/UEBuildTarget.cs
+++ b/Engine/Source/Programs/UnrealBuildTool/Configuration/UEBuildTarget.cs
@@ -6769,7 +6769,7 @@ namespace UnrealBuildTool
 			}
 
 			// If compiling with the AutoRTFM compiler, implicitly depend on the AutoRTFM runtime.
-			if (Rules.bUseAutoRTFMCompiler && !RulesObject.bDisableAutoRTFMInstrumentation && ModuleName != "AutoRTFM")
+			if (Rules.bUseAutoRTFMCompiler && ModuleName != "AutoRTFM")
 			{
 				RulesObject.PrivateDependencyModuleNames.Add("AutoRTFM");
 			}
