UnrealBuildTool: Use TargetFramework from shared property sheet

#rnx
#Synced-CL 50830039

diff --git a/Engine/Source/Programs/UnrealBuildTool/UnrealBuildTool.csproj b/Engine/Source/Programs/UnrealBuildTool/UnrealBuildTool.csproj
--- a/Engine/Source/Programs/UnrealBuildTool/UnrealBuildTool.csproj
+++ b/Engine/Source/Programs/UnrealBuildTool/UnrealBuildTool.csproj
@@ -2,7 +2,6 @@
   <Import Project="..\Shared\UnrealEngine.csproj.props" />
 
   <PropertyGroup>
-    <TargetFramework>net8.0</TargetFramework>
     <OutputType>Exe</OutputType>
     <GenerateAssemblyInfo>false</GenerateAssemblyInfo>
     <GenerateTargetFrameworkAttribute>false</GenerateTargetFrameworkAttribute>
