UnrealBuildTool: Prevent more external environment variables from being used by compilers

#rnx

#Synced-CL 50830039

diff --git a/Engine/Source/Programs/UnrealBuildTool/Platform/Clang/ClangToolChain.cs b/Engine/Source/Programs/UnrealBuildTool/Platform/Clang/ClangToolChain.cs
--- a/Engine/Source/Programs/UnrealBuildTool/Platform/Clang/ClangToolChain.cs
+++ b/Engine/Source/Programs/UnrealBuildTool/Platform/Clang/ClangToolChain.cs
@@ -251,6 +251,15 @@ namespace UnrealBuildTool
 			LazyInfo = new Lazy<ClangToolChainInfo>(() => { return GetToolChainInfo(); }); // Don't change to => GetToolChainInfo().. it doesnt produce correct code
 		}
 
+		public override void SetEnvironmentVariables()
+		{
+			// Don't allow the INCLUDE environment variables to propagate.
+			Environment.SetEnvironmentVariable("C_INCLUDE_PATH", null);
+			Environment.SetEnvironmentVariable("OBJC_INCLUDE_PATH", null);
+			Environment.SetEnvironmentVariable("CPLUS_INCLUDE_PATH", null);
+			Environment.SetEnvironmentVariable("OBJCPLUS_INCLUDE_PATH", null);
+		}
+
 		/// <summary>
 		/// Create toolchain info. Do not call this function, use "Info" to get results
 		/// </summary>
diff --git a/Engine/Source/Programs/UnrealBuildTool/Platform/Windows/VCToolChain.cs b/Engine/Source/Programs/UnrealBuildTool/Platform/Windows/VCToolChain.cs
--- a/Engine/Source/Programs/UnrealBuildTool/Platform/Windows/VCToolChain.cs
+++ b/Engine/Source/Programs/UnrealBuildTool/Platform/Windows/VCToolChain.cs
@@ -462,6 +462,9 @@ namespace UnrealBuildTool
 
 		protected virtual void AppendCLArguments_Global(CppCompileEnvironment CompileEnvironment, List<string> Arguments)
 		{
+			// Prevents the compiler from searching for include files in directories specified in the PATH and INCLUDE environment variables
+			Arguments.Add("/X");
+
 			// Workaround for MSVC 14.31 compiler crash
 			if (EnvVars.Compiler.IsMSVC() && EnvVars.ToolChainVersion < new VersionNumber(14, 43) && EnvVars.ToolChainVersion >= new VersionNumber(14, 31))
 			{
