[UBT]
* Fixed "Done" text for horde info in uba visualizer being in the right place when using multiple horde sessions

#Synced-CL 50855769

diff --git a/Engine/Source/Programs/UnrealBuildTool/Executors/UnrealBuildAccelerator/UBAAgentCoordinatorHorde.cs b/Engine/Source/Programs/UnrealBuildTool/Executors/UnrealBuildAccelerator/UBAAgentCoordinatorHorde.cs
--- a/Engine/Source/Programs/UnrealBuildTool/Executors/UnrealBuildAccelerator/UBAAgentCoordinatorHorde.cs
+++ b/Engine/Source/Programs/UnrealBuildTool/Executors/UnrealBuildAccelerator/UBAAgentCoordinatorHorde.cs
@@ -430,7 +430,7 @@ namespace UnrealBuildTool
 			{
 				for (uint i = 0; i != _hordeSessionTasks?.Length; ++i)
 				{
-					_updateStatus(0, 6 + i, "Done", EpicGames.UBA.LogEntryType.Info, null);
+					_updateStatus(i, 6, "Done", EpicGames.UBA.LogEntryType.Info, null);
 				}
 			}
 		}
