Moving various bits of MSGaming code from Microsoft to Windows folders
- not intended for use by non-Windows platforms

#rnx

#Synced-CL 50857758

diff --git a/Engine/Source/Programs/UnrealBuildTool/Platform/Microsoft/GDK/MSGaming/MSGamingExports.cs b/Engine/Source/Programs/UnrealBuildTool/Platform/Microsoft/GDK/MSGaming/MSGamingExports.cs
deleted file mode 100644
diff --git a/Engine/Source/Programs/UnrealBuildTool/Platform/Microsoft/GDK/MSGaming/MSGamingGameConfigGenerator.cs b/Engine/Source/Programs/UnrealBuildTool/Platform/Microsoft/GDK/MSGaming/MSGamingGameConfigGenerator.cs
deleted file mode 100644
diff --git a/Engine/Source/Programs/UnrealBuildTool/Platform/Microsoft/GDK/MSGaming/MSGamingToolChain.cs b/Engine/Source/Programs/UnrealBuildTool/Platform/Microsoft/GDK/MSGaming/MSGamingToolChain.cs
deleted file mode 100644
diff --git a/Engine/Source/Programs/UnrealBuildTool/Platform/Windows/MSGaming/MSGamingExports.cs b/Engine/Source/Programs/UnrealBuildTool/Platform/Windows/MSGaming/MSGamingExports.cs
new file mode 100644
--- /dev/null
+++ b/Engine/Source/Programs/UnrealBuildTool/Platform/Windows/MSGaming/MSGamingExports.cs
@@ -0,0 +1,192 @@
+// Copyright Epic Games, Inc. All Rights Reserved.
+
+using EpicGames.Core;
+using Microsoft.Win32;
+using System;
+using System.Collections.Generic;
+using System.Linq;
+using System.Runtime.Versioning;
+
+namespace UnrealBuildTool
+{
+	/// <summary>
+	/// Public MSGaming functions exposed to UAT
+	/// </summary>
+	public static class MSGamingExports
+	{
+		#region custom install actions
+
+		/// Type of custom install action
+		public enum CustomInstallActionType
+		{
+			/// command for installing
+			Install,
+			/// command for repairing
+			Repair,
+			/// command for uninstalling
+			Uninstall,
+		};
+
+#pragma warning disable CA1034 // Nested types should not be visible - preserve API of APIs exposed to UAT
+		/// a custom install action
+		public struct CustomInstallAction
+		{
+			/// location to copy the custom installs to when staging
+			private const string DefaultTargetRootFolder = "CustomInstall";
+
+			/// name of the install action
+			public string Name;
+
+			/// source folder, relative to engine root
+			public DirectoryReference SourceFolder;
+
+			/// target folder, relative to targe root folder
+			public string TargetFolder;
+
+			/// a command to run and its arguments
+			public struct InstallCommand
+			{
+				/// command to run
+				public FileReference Command;
+
+				/// arguments to pass to the given command
+				public string Arguments;
+			};
+
+			/// the commands to run for each action
+			public Dictionary<CustomInstallActionType, InstallCommand> Commands;
+
+			/// read the given config key and return the actions
+			public static IEnumerable<CustomInstallAction> GetFromConfig(ConfigHierarchy EngineIni, string IniSection)
+			{
+				List<CustomInstallAction> CustomInstallActions = new();
+
+				if (EngineIni!.TryGetValuesGeneric(IniSection, "CustomInstallActions", out CustomInstallAction[]? CustomInstallActionsArray) && CustomInstallActionsArray != null)
+				{
+					// set the default target folder
+					for (int i = 0; i < CustomInstallActionsArray.Length; i++)
+					{
+						if (String.IsNullOrEmpty(CustomInstallActionsArray[i].TargetFolder))
+						{
+							CustomInstallActionsArray[i].TargetFolder = CustomInstallActionsArray[i].Name;
+						}
+						else if (CustomInstallActionsArray[i].TargetFolder == ".")
+						{
+							CustomInstallActionsArray[i].TargetFolder = "";
+						}
+					}
+
+					CustomInstallActions.AddRange(CustomInstallActionsArray);
+				}
+
+				// inject the GameInput dependency if requested
+				if (EngineIni.GetBool(IniSection, "bAlwaysIncludeGameInputDependency", out bool bAlwaysIncludeGameInputDependency) && bAlwaysIncludeGameInputDependency)
+				{
+					DirectoryReference GameInputRedistDir = DirectoryReference.Combine(UnrealBuildBase.Unreal.EngineDirectory, @"Extras\Redist\en-us\");
+					FileReference GameInputRedist = FileReference.Combine(GameInputRedistDir, "GameInputRedist.msi");
+
+					InstallCommand MSICommand = new()
+					{
+						Command = GameInputRedist,
+						Arguments = ""
+					};
+
+					CustomInstallAction GameInputAction = new();
+					GameInputAction.Name = "GameInput";
+					GameInputAction.SourceFolder = GameInputRedistDir;
+					GameInputAction.TargetFolder = "GameInput";
+					GameInputAction.Commands = new()
+					{
+						{ CustomInstallActionType.Install, MSICommand },
+						{ CustomInstallActionType.Uninstall, MSICommand },
+						{ CustomInstallActionType.Repair, MSICommand }
+					};
+
+					CustomInstallActions.Add(GameInputAction);
+				}
+
+				return CustomInstallActions.ToArray();
+			}
+
+			/// the target root folder for staging custom install actions, relative to the package root
+			public static string GetTargetRootFolder(ConfigHierarchy EngineIni, string IniSection)
+			{
+				if (EngineIni.TryGetValue(IniSection, "CustomInstallStagingRoot", out string? TargetRootFolder ) && !String.IsNullOrEmpty(TargetRootFolder))
+				{
+					return TargetRootFolder;
+				}
+				else
+				{
+					return DefaultTargetRootFolder;
+				}
+			}
+		}
+#pragma warning restore CA1034
+
+		#endregion
+
+		/// <summary>
+		/// Determine whether Windows is in developer mode. PC GDK packages can't be sideloaded without it
+		/// </summary>
+		[SupportedOSPlatform("windows")]
+		public static bool IsWindowsDeveloperMode()
+		{
+			try
+			{
+				// normal windows key
+				using RegistryKey? Key = Registry.LocalMachine.OpenSubKey(@"SOFTWARE\Microsoft\Windows\CurrentVersion\AppModelUnlock");
+				if (Key != null && Convert.ToInt64(Key.GetValue("AllowDevelopmentWithoutDevLicense", 0L)) != 0)
+				{
+					return true;
+				}
+
+				// group policy key
+				using RegistryKey? GroupKey = Registry.LocalMachine.OpenSubKey(@"SOFTWARE\Policies\Microsoft\Windows\Appx");
+				if (GroupKey != null && Convert.ToInt64(GroupKey.GetValue("AllowDevelopmentWithoutDevLicense", 0L)) != 0)
+				{
+					return true;
+				}
+			}
+			catch (Exception)
+			{
+			}
+
+			return false;
+		}
+
+		/// <summary>
+		/// Determine if the Microsoft.GamingServices package is installed.
+		/// </summary>
+		[SupportedOSPlatform("windows")]
+		public static bool IsMSGamingRuntimeInstalled()
+		{
+			try
+			{
+				using RegistryKey? Key = Registry.CurrentUser.OpenSubKey(@"SOFTWARE\Classes\Local Settings\Software\Microsoft\Windows\CurrentVersion\AppModel\Repository\Packages");
+				if (Key != null)
+				{
+					return Key.GetSubKeyNames().Any( X => X.StartsWith("Microsoft.GamingServices_", StringComparison.Ordinal));
+				}
+			}
+			catch (Exception)
+			{
+			}
+		
+			return false;
+		}
+
+		/// <summary>
+		/// Determine whether the local version of Windows can support GDK
+		/// </summary>
+		[SupportedOSPlatform("windows")]
+		public static bool IsWindowsVersionSupported()
+		{
+			if (System.Environment.OSVersion.Version.Major < 10 || (System.Environment.OSVersion.Version.Major == 10 && System.Environment.OSVersion.Version.Build < 18362)) //https://en.wikipedia.org/wiki/Template:Windows_10_versions
+			{
+				return false;
+			}
+
+			return true;
+		}
+	}
+}
diff --git a/Engine/Source/Programs/UnrealBuildTool/Platform/Windows/MSGaming/MSGamingGameConfigGenerator.cs b/Engine/Source/Programs/UnrealBuildTool/Platform/Windows/MSGaming/MSGamingGameConfigGenerator.cs
new file mode 100644
--- /dev/null
+++ b/Engine/Source/Programs/UnrealBuildTool/Platform/Windows/MSGaming/MSGamingGameConfigGenerator.cs
@@ -0,0 +1,258 @@
+// Copyright Epic Games, Inc. All Rights Reserved.
+
+using EpicGames.Core;
+using Microsoft.Extensions.Logging;
+using System;
+using System.Collections.Generic;
+using System.IO;
+using System.Linq;
+using System.Xml.Linq;
+using UnrealBuildBase;
+using static UnrealBuildTool.MSGamingExports;
+
+namespace UnrealBuildTool.GDK
+{
+	partial struct GDKGameConfigGeneratorFactory
+	{
+#pragma warning disable CA1823 // Avoid unused private fields - instantiation is side-effectful, preferable to a static ctor
+#pragma warning disable IDE0052
+		private static readonly Type MSGamingGen = RegisterGenerator(UnrealTargetPlatform.Win64, typeof(MSGamingGameConfigGenerator));
+#pragma warning restore CA1823
+#pragma warning restore IDE0052
+	}
+
+	class MSGamingGameConfigGenerator : GDKGameConfigGenerator
+	{
+		protected override string TargetDeviceFamily => "PC";
+		protected override string DefaultArchitecture => "x64";
+		protected override string IniSection_PlatformTargetSettings => "/Script/MSGamingSupport.MSGamingSettings";
+		protected override string? IniSection_GeneralPlatformSettings => "/Script/WindowsTargetPlatform.WindowsTargetSettings"; // for querying DefaultGraphicsRHI etc.
+
+		protected virtual string IniPlatformName => "Windows";
+
+		/// <summary>
+		/// Create a manifest generator for Win64 using MS Gaming Runtime and MS Game Store
+		/// </summary>
+		public MSGamingGameConfigGenerator(ILogger InLogger)
+			: base(UnrealTargetPlatform.Win64, InLogger)
+		{
+		}
+
+		/// <summary>
+		/// Create a manifest generator for other platforms
+		/// </summary>
+		internal MSGamingGameConfigGenerator(UnrealTargetPlatform InPlatform, ILogger InLogger)
+			: base(InPlatform, InLogger)
+		{
+		}
+
+		/// <summary>
+		/// Register the locations where resource binary files can be found
+		/// </summary>
+		protected override void PrepareResourceBinaryPaths()
+		{
+			base.PrepareResourceBinaryPaths();
+
+			// check for DLC images first
+			if (DLCFile != null)
+			{
+				AppXResources!.ProjectBinaryResourceDirectories.Add(DirectoryReference.Combine(DLCFile.Directory, "Build", IniPlatformName, "MSGaming", BuildResourceSubPath));
+				AppXResources!.ProjectBinaryResourceDirectories.Add(DirectoryReference.Combine(DLCFile.Directory, "Platforms", IniPlatformName, "Build", "MSGaming", BuildResourceSubPath));
+			}
+
+			// allow projects to keep their resources in an MSGaming subfolder to help with project organization
+			if (ProjectFile != null)
+			{
+				AppXResources!.ProjectBinaryResourceDirectories.Add(DirectoryReference.Combine(ProjectFile.Directory, "Build", IniPlatformName, "MSGaming", BuildResourceSubPath));
+				AppXResources!.ProjectBinaryResourceDirectories.Add(DirectoryReference.Combine(ProjectFile.Directory, "Platforms", IniPlatformName, "Build", "MSGaming", BuildResourceSubPath));
+			}
+
+			// default images are stored under MSGamingSupport plugin folder
+			DirectoryReference MSGamingSupportResourcePath = DirectoryReference.Combine(Unreal.EngineDirectory, "Build", EngineResourceSubPath);
+			PluginInfo? MSGamingSupportPlugin = Plugins.GetPlugin("MSGamingSupport");
+			if (MSGamingSupportPlugin != null)
+			{
+				MSGamingSupportResourcePath = DirectoryReference.Combine( MSGamingSupportPlugin.File.Directory, "Build", EngineResourceSubPath);
+			}
+			AppXResources!.EngineFallbackBinaryResourceDirectories.Add(MSGamingSupportResourcePath);
+
+			// fallback path because Program targets do not load plugins & this is still under the platforms folder for now
+			AppXResources!.EngineFallbackBinaryResourceDirectories.Add(DirectoryReference.Combine(Unreal.EngineDirectory, @"Platforms\GDK\Plugins\MSGamingSupport\Build\DefaultImages"));
+		}
+
+		protected override XElement GetGame(Dictionary<UnrealTargetConfiguration, string> InExecutablePairs, out string IdentityName)
+		{
+			XElement Game = base.GetGame(InExecutablePairs, out IdentityName);
+
+			Game.Add(GetDesktopRegistration());
+
+			if (!IsDLC)
+			{
+				if (EngineIni!.GetBool(IniSection_PlatformTargetSettings, "bMSAFullTrust", out bool bMSAFullTrust) && bMSAFullTrust)
+				{
+					Game.Add(new XElement(XName.Get("MSAFullTrust"), "true"));
+				}
+			}
+
+			return Game;
+		}
+
+		private XElement GetDesktopRegistration()
+		{
+			XElement DesktopRegistration = new(XName.Get("DesktopRegistration"),
+				new XElement(XName.Get("ProcessorArchitecture"), DefaultArchitecture)
+				);
+
+			if (!IsDLC)
+			{
+				int ManifestVersion = GetManifestVersion();
+
+				DesktopRegistration.Add(
+					new XElement(XName.Get("MultiplayerProtocol"), "true"),
+					GetDependencyList(),
+					GetCustomInstallActions(),
+					GetWindowsOSVersion()
+					);
+
+				// Note: Manifest Version 0 has been deprecated as of October 2023 GDK and is not supported,
+				// but keeping this here in case Manifest version is overridden.  
+				if (ManifestVersion == 0) // ModFolder is deprecated in Manifest Version 1
+				{
+					if (EngineIni!.GetBool(IniSection_PlatformTargetSettings, "bUseModFolder", out bool bUseModFolder) && bUseModFolder)
+					{
+						if (EngineIni.GetString(IniSection_PlatformTargetSettings, "ModFolder", out string ModFolder) && !String.IsNullOrEmpty(ModFolder))
+						{
+							DesktopRegistration.Add(new XElement(XName.Get("ModFolder"),
+								new XAttribute(XName.Get("Name"), ModFolder)
+								)
+							);
+						}
+					}
+				}
+			}
+			else
+			{
+				string MainPackageIdentityName = GetMainPackageIdentityPackageName();
+
+				DesktopRegistration.Add(new XElement(XName.Get("MainPackageDependency"),
+					new XAttribute(XName.Get("Name"), MainPackageIdentityName)
+					)
+				);
+			}
+
+			return DesktopRegistration;
+		}
+
+		private XElement? GetDependencyList()
+		{
+			XElement DependencyList = new(XName.Get("DependencyList"));
+
+			if (NoBootstrapExe)
+			{
+				// Without a bootstrapper, need Windows to install VC runtime libraries during package install
+				// this will cause VC DLL loads to be redirected from %System32% to %ProgramFiles%\WindowsApps\Microsoft.VCLibs.140.00.UWPDesktop_14.0.32530.0_x64__8wekyb3d8bbwe
+				DependencyList.Add(new XElement(XName.Get("KnownDependency"), new XAttribute(XName.Get("Name"), "VC14")));
+			}
+
+			bool bNeedsDX11 = (GetConfigString("DefaultGraphicsRHI", null) == "DefaultGraphicsRHI_DX11") ||
+							  (GetConfigBool("bAlwaysIncludeDX11Dependency", null));
+			if (bNeedsDX11)
+			{
+				// this will cause DX DLL loads to be redirected from %System32% to %ProgramFiles%\WindowsApps\Microsoft.DirectXRuntime_9.29.1974.0_x64__8wekyb3d8bbwe
+				Log.TraceInformationOnce("Adding KnownDependency on DX11 redist in MicrosoftGame.config");
+				DependencyList.Add(new XElement(XName.Get("KnownDependency"), new XAttribute(XName.Get("Name"), "DX11")));
+			}
+
+			return DependencyList.HasElements ? DependencyList : null;
+		}
+
+		private XElement? GetCustomInstallActions()
+		{
+			IEnumerable<CustomInstallAction> CustomInstallActions = CustomInstallAction.GetFromConfig(EngineIni!, IniSection_PlatformTargetSettings);
+			if (!CustomInstallActions.Any())
+			{
+				return null;
+			}
+
+			string TargetRoot = CustomInstallAction.GetTargetRootFolder(EngineIni!, IniSection_PlatformTargetSettings);
+
+			List<XElement> InstallActionList = new();
+			List<XElement> RepairActionList = new();
+			List<XElement> UninstallActionList = new();
+			foreach (CustomInstallAction Action in CustomInstallActions)
+			{
+				foreach (KeyValuePair<CustomInstallActionType, CustomInstallAction.InstallCommand> Command in Action.Commands)
+				{
+					List<XAttribute> Attribs = new List<XAttribute>
+					{
+						new("File",      Path.Combine( Action.TargetFolder, Command.Value.Command.MakeRelativeTo(Action.SourceFolder))),
+						new("Name",      Action.Name),
+					};
+					if (!String.IsNullOrEmpty(Command.Value.Arguments))
+					{
+						Attribs.Add(new("Arguments", Command.Value.Arguments));
+					}
+					switch (Command.Key)
+					{
+						case CustomInstallActionType.Install: InstallActionList.Add(new XElement("InstallAction", Attribs)); break;
+						case CustomInstallActionType.Repair: RepairActionList.Add(new XElement("RepairAction", Attribs)); break;
+						case CustomInstallActionType.Uninstall: UninstallActionList.Add(new XElement("UninstallAction", Attribs)); break;
+					}
+				}
+			}
+
+			XElement CustomInstallActionsElement = new("CustomInstallActions",
+				new XElement("Folder", new XText(TargetRoot)),
+				InstallActionList.Any() ? new XElement("InstallActionList", InstallActionList) : null,
+				RepairActionList.Any() ? new XElement("RepairActionList", RepairActionList) : null,
+				UninstallActionList.Any() ? new XElement("UninstallActionList", UninstallActionList) : null
+			);
+			return CustomInstallActionsElement;
+		}
+
+		private XElement? GetWindowsOSVersion()
+		{
+			EngineIni!.GetString(IniSection_PlatformTargetSettings, "RequiredMinimumWindowsVersion", out string RequiredMinVersion);
+			EngineIni!.GetString(IniSection_PlatformTargetSettings, "SuggestedMinimumWindowsVersion", out string SuggestedMinVersion);
+			EngineIni!.GetString(IniSection_PlatformTargetSettings, "RecommendedWindowsVersion", out string RecommendedVersion);
+
+			XElement WindowsOsVersion = new(XName.Get("WindowsOsVersion"));
+
+			if (!String.IsNullOrWhiteSpace(RequiredMinVersion))
+			{
+				WindowsOsVersion.SetAttributeValue("RequiredMinimum", RequiredMinVersion);
+			}
+
+			if (!String.IsNullOrWhiteSpace(SuggestedMinVersion))
+			{
+				WindowsOsVersion.SetAttributeValue("SuggestedMinimum", SuggestedMinVersion);
+			}
+
+			if (!String.IsNullOrWhiteSpace(RecommendedVersion))
+			{
+				WindowsOsVersion.SetAttributeValue("Recommended", RecommendedVersion);
+			}
+
+			// The <WindowsOsVersion> element is optional. If none of the values were configured, omit it.
+			return WindowsOsVersion.HasAttributes ? WindowsOsVersion : null;
+		}
+
+		protected override string? IncludeBuildVersionInPackageVersion(string? VersionNumber)
+		{
+			// GDK on Windows reserves the lowest version element (Revision) for internal use so package versioning needs to be a little different
+
+			if (VersionNumber != null && BuildVersion.TryRead(BuildVersion.GetDefaultFileName(), out BuildVersion? BuildVersionForPackage) && BuildVersionForPackage.Changelist != 0)
+			{
+				// Break apart the version number into individual elements
+				string[] SplitVersionString = VersionNumber.Split('.');
+				VersionNumber = String.Format("{0}.{1}.{2}.{3}",
+					SplitVersionString[0],
+					BuildVersionForPackage.Changelist / 10000,
+					BuildVersionForPackage.Changelist % 10000,
+					0); //must be zero
+			}
+
+			return VersionNumber;
+		}
+	}
+}
diff --git a/Engine/Source/Programs/UnrealBuildTool/Platform/Windows/MSGaming/MSGamingToolChain.cs b/Engine/Source/Programs/UnrealBuildTool/Platform/Windows/MSGaming/MSGamingToolChain.cs
new file mode 100644
--- /dev/null
+++ b/Engine/Source/Programs/UnrealBuildTool/Platform/Windows/MSGaming/MSGamingToolChain.cs
@@ -0,0 +1,103 @@
+// Copyright Epic Games, Inc. All Rights Reserved.
+
+using EpicGames.Core;
+using Microsoft.Extensions.Logging;
+using System;
+using System.Collections.Generic;
+using System.Linq;
+using UnrealBuildTool.GDK;
+
+namespace UnrealBuildTool
+{
+	class MSGamingToolChain : VCToolChain
+	{
+		protected virtual bool SupportsContentOnlyProjects => false; // there is no content-only project support for Win64+GDK because it requires bespoke target.cs files
+
+		public MSGamingToolChain(ReadOnlyTargetRules Target, ILogger Logger)
+			: base(Target, Logger)
+		{
+		}
+
+		public override void GetVersionInfo(List<string> Lines)
+		{
+			base.GetVersionInfo(Lines);
+
+			string? GDKEdition = GDKExports.GetSDKVersion();
+			if (GDKEdition != null)
+			{
+				string CurrentGDKDir = GDKExports.GetCurrentGSDKDir();
+				Lines.Add($"Using GDK version {GDKEdition} ({CurrentGDKDir}).");
+			}
+			else
+			{
+				Lines.Add("GameDKCoreLatest or GameDKLatest environment variable not found - is the GDK installed?");
+			}
+		}
+
+		protected override void AppendLinkArguments(LinkEnvironment LinkEnvironment, List<string> Arguments)
+		{
+			base.AppendLinkArguments(LinkEnvironment, Arguments);
+
+			// do not strip the PDB path from the executable if we want to debug it (otherwise it can't find the symbols when attempting to debug an installed package build)
+			if (LinkEnvironment.bCreateDebugInfo && LinkEnvironment.Configuration != CppConfiguration.Shipping)
+			{
+				Arguments.Remove("/PDBALTPATH:%_PDB%");
+			}
+		}
+
+		public override void FinalizeOutput(ReadOnlyTargetRules Target, TargetMakefileBuilder MakefileBuilder, DirectoryReference ExeDir)
+		{
+			base.FinalizeOutput(Target, MakefileBuilder, ExeDir);
+
+			// Look up the project
+			FileReference? ProjectFile = Target.ProjectFile;
+			if (SupportsContentOnlyProjects)
+			{
+				ProjectFile = GDKContentOnlyProjectSupport.GetProjectFileFromTarget(Target, Logger) ?? ProjectFile;
+			}
+
+			// Create a placeholder MicrosoftGame.config to be used during F5 debugging, and for Program targets. It does not reference executables because it could be shared between several configurations.
+			// This allows the game to run while using assets from the Saved/Cooked folder, without needing -basedir= to point to a Staged build
+			if ((Environment.CommandLine.Contains("-frommsbuild", StringComparison.OrdinalIgnoreCase) && ProjectFile != null) || Target.Type == TargetType.Program)
+			{
+				DirectoryReference ExecutableDirectory = MakefileBuilder.Makefile.ExecutableFile.Directory;
+				Action GenAction = GDKGameConfigGeneratorMode.CreateAction(MakefileBuilder, Target.Platform, ExecutableDirectory, Target.Name, ProjectFile, CustomConfig:Target.CustomConfig);
+				GenAction.CommandDescription = "Generate MicrosoftGame.config";
+				GenAction.bCanExecuteInUBA = false; // TODO: Fails with some missing file, revisit
+			}
+		}
+
+		protected override void ModifyFinalLinkArguments(LinkEnvironment LinkEnvironment, List<string> Arguments, bool bBuildImportLibraryOnly)
+		{
+			base.ModifyFinalLinkArguments(LinkEnvironment, Arguments, bBuildImportLibraryOnly);
+
+			// optionally audit the GDK editions of any libs we link to
+			GDKEditionAudit.VerifyLinkEnvironment(EnvVars, LinkEnvironment, bBuildImportLibraryOnly, Logger);
+		}
+
+		public override void PrepareRuntimeDependencies(List<RuntimeDependency> RuntimeDependencies, Dictionary<FileReference, FileReference> TargetFileToSourceFile, DirectoryReference ExeDir)
+		{
+			base.PrepareRuntimeDependencies(RuntimeDependencies, TargetFileToSourceFile, ExeDir);
+
+			// optionally audit the GDK editions of any libs we depend on
+			GDKEditionAudit.VerifyRuntimeDependencies(EnvVars, Target.Platform, TargetFileToSourceFile.Values.Distinct(), Logger);
+		}
+
+		public override void SetUpGlobalEnvironment(ReadOnlyTargetRules Target, CppCompileEnvironment GlobalCompileEnvironment, LinkEnvironment GlobalLinkEnvironment)
+		{
+			base.SetUpGlobalEnvironment(Target, GlobalCompileEnvironment, GlobalLinkEnvironment);
+
+			DirectoryReference platformGDK = new(GDKExports.GetCurrentGSDKDir());
+			GlobalCompileEnvironment.RootPaths.AddFolderName(CppRootPathFolder.PlatformSDK, $"{Target.Platform}SDK");
+			GlobalLinkEnvironment.RootPaths.AddFolderName(CppRootPathFolder.PlatformSDK, $"{Target.Platform}SDK");
+			GlobalCompileEnvironment.RootPaths[CppRootPathFolder.PlatformSDK] = platformGDK;
+			GlobalLinkEnvironment.RootPaths[CppRootPathFolder.PlatformSDK] = platformGDK;
+
+			string? currentGDKVersion = GDKExports.GetSDKVersion();
+			if (currentGDKVersion != null)
+			{
+				EpicGames.UBA.Utils.RegisterPathHash(platformGDK.FullName, currentGDKVersion);
+			}
+		}
+	}
+}
\ No newline at end of file
diff --git a/Engine/Source/Runtime/Microsoft/GDK/MSGameStorePlatformFeatures/MSGameStorePlatformFeatures.Build.cs b/Engine/Source/Runtime/Microsoft/GDK/MSGameStorePlatformFeatures/MSGameStorePlatformFeatures.Build.cs
deleted file mode 100644
diff --git a/Engine/Source/Runtime/Microsoft/GDK/MSGameStorePlatformFeatures/Private/MSGameStorePlatformFeatures.cpp b/Engine/Source/Runtime/Microsoft/GDK/MSGameStorePlatformFeatures/Private/MSGameStorePlatformFeatures.cpp
deleted file mode 100644
diff --git a/Engine/Source/Runtime/Microsoft/GDK/MSGameStorePlatformFeatures/Public/MSGameStorePlatformFeatures.h b/Engine/Source/Runtime/Microsoft/GDK/MSGameStorePlatformFeatures/Public/MSGameStorePlatformFeatures.h
deleted file mode 100644
diff --git a/Engine/Source/Runtime/Windows/MSGaming/MSGameStorePlatformFeatures/MSGameStorePlatformFeatures.Build.cs b/Engine/Source/Runtime/Windows/MSGaming/MSGameStorePlatformFeatures/MSGameStorePlatformFeatures.Build.cs
new file mode 100644
--- /dev/null
+++ b/Engine/Source/Runtime/Windows/MSGaming/MSGameStorePlatformFeatures/MSGameStorePlatformFeatures.Build.cs
@@ -0,0 +1,24 @@
+// Copyright Epic Games, Inc. All Rights Reserved.
+
+using UnrealBuildTool;
+
+[SupportedPlatformGroups("Windows")]
+public class MSGameStorePlatformFeatures : ModuleRules
+{
+	public MSGameStorePlatformFeatures( ReadOnlyTargetRules Target) : base(Target)
+	{
+		BinariesSubFolder = "GDK";
+
+		PrivateDependencyModuleNames.AddRange(new string[]
+			{
+				"Core",
+				"Engine",
+				"GRDK",
+			});
+
+		PublicDependencyModuleNames.AddRange(new string[]
+			{
+				"WindowsPlatformFeatures",
+			});
+	}
+}
diff --git a/Engine/Source/Runtime/Windows/MSGaming/MSGameStorePlatformFeatures/Private/MSGameStorePlatformFeatures.cpp b/Engine/Source/Runtime/Windows/MSGaming/MSGameStorePlatformFeatures/Private/MSGameStorePlatformFeatures.cpp
new file mode 100644
--- /dev/null
+++ b/Engine/Source/Runtime/Windows/MSGaming/MSGameStorePlatformFeatures/Private/MSGameStorePlatformFeatures.cpp
@@ -0,0 +1,37 @@
+// Copyright Epic Games, Inc. All Rights Reserved.
+
+#include "MSGameStorePlatformFeatures.h"
+#if WITH_GRDK
+#include "Microsoft/AllowMicrosoftPlatformTypes.h"
+#include <XGame.h>
+#include <XAppCapture.h>
+#include "Microsoft/HideMicrosoftPlatformTypes.h"
+
+void FMSGameStorePlatformFeaturesModule::SetScreenshotEnableState(bool bEnabled)
+{
+	if (bEnabled)
+	{
+		XAppCaptureEnableRecord();
+	}
+	else
+	{
+		XAppCaptureDisableRecord();
+	}
+}
+
+
+FString FMSGameStorePlatformFeaturesModule::GetUniqueAppId()
+{
+	uint32_t TitleId = 0;
+	if (SUCCEEDED(XGameGetXboxTitleId(&TitleId)))
+	{
+		return FString::Printf(TEXT("%X"), TitleId);
+	}
+	else
+	{
+		return FString();
+	}
+}
+#endif //WITH_GRDK
+
+IMPLEMENT_MODULE(FMSGameStorePlatformFeaturesModule, MSGameStorePlatformFeatures);
diff --git a/Engine/Source/Runtime/Windows/MSGaming/MSGameStorePlatformFeatures/Public/MSGameStorePlatformFeatures.h b/Engine/Source/Runtime/Windows/MSGaming/MSGameStorePlatformFeatures/Public/MSGameStorePlatformFeatures.h
new file mode 100644
--- /dev/null
+++ b/Engine/Source/Runtime/Windows/MSGaming/MSGameStorePlatformFeatures/Public/MSGameStorePlatformFeatures.h
@@ -0,0 +1,13 @@
+// Copyright Epic Games, Inc. All Rights Reserved.
+
+#include "WindowsPlatformFeatures.h"
+
+class FMSGameStorePlatformFeaturesModule : public FWindowsPlatformFeaturesModule
+{
+public:
+#if WITH_GRDK
+	virtual FString GetUniqueAppId() override;
+	virtual void SetScreenshotEnableState(bool bEnabled) override;
+#endif
+};
+
