UnrealBuildTool: Light refactor of AnalyzeMode to follow naming conventions.

#Synced-CL 50859582

diff --git a/Engine/Source/Programs/UnrealBuildTool/Modes/AnalyzeMode.cs b/Engine/Source/Programs/UnrealBuildTool/Modes/AnalyzeMode.cs
--- a/Engine/Source/Programs/UnrealBuildTool/Modes/AnalyzeMode.cs
+++ b/Engine/Source/Programs/UnrealBuildTool/Modes/AnalyzeMode.cs
@@ -14,44 +14,39 @@ namespace UnrealBuildTool.Modes
 	/// Outputs information about the given target, including a module dependecy graph (in .gefx format and list of module references)
 	/// </summary>
 	[ToolMode("Analyze", ToolModeOptions.XmlConfig | ToolModeOptions.BuildPlatforms | ToolModeOptions.SingleInstance | ToolModeOptions.StartPrefetchingEngine | ToolModeOptions.ShowExecutionTime)]
-	class AnalyzeMode : ToolMode
+	internal class AnalyzeMode : ToolMode
 	{
-		/// <summary>
-		/// Execute the command
-		/// </summary>
-		/// <param name="Arguments">Command line arguments</param>
-		/// <returns>Exit code</returns>
-		/// <param name="Logger"></param>
-		public override Task<int> ExecuteAsync(CommandLineArguments Arguments, ILogger Logger)
+		/// <inheritdoc/>
+		public override Task<int> ExecuteAsync(CommandLineArguments arguments, ILogger logger)
 		{
-			Arguments.ApplyTo(this);
+			arguments.ApplyTo(this);
 
 			// Create the build configuration object, and read the settings
-			BuildConfiguration BuildConfiguration = new();
-			XmlConfig.ApplyTo(BuildConfiguration);
-			Arguments.ApplyTo(BuildConfiguration);
+			BuildConfiguration buildConfiguration = new();
+			XmlConfig.ApplyTo(buildConfiguration);
+			arguments.ApplyTo(buildConfiguration);
 
 			// Parse all the target descriptors
-			List<TargetDescriptor> TargetDescriptors = TargetDescriptor.ParseCommandLine(Arguments, BuildConfiguration, Logger);
+			List<TargetDescriptor> targetDescriptors = TargetDescriptor.ParseCommandLine(arguments, buildConfiguration, logger);
 
 			// Generate the compile DB for each target
-			using (ISourceFileWorkingSet WorkingSet = new EmptySourceFileWorkingSet())
+			using (ISourceFileWorkingSet workingSet = new EmptySourceFileWorkingSet())
 			{
 				// Find the compile commands for each file in the target
-				Dictionary<FileReference, string> FileToCommand = [];
-				foreach (TargetDescriptor TargetDescriptor in TargetDescriptors)
+				Dictionary<FileReference, string> fileToCommand = [];
+				foreach (TargetDescriptor targetDescriptor in targetDescriptors)
 				{
-					AnalyzeTarget(TargetDescriptor, BuildConfiguration, Logger);
+					AnalyzeTarget(targetDescriptor, buildConfiguration, logger);
 				}
 			}
 
 			return Task.FromResult(0);
 		}
 
-		class ModuleInfo
+		private class ModuleInfo(UEBuildModule module, params string[] includeChain)
 		{
-			public UEBuildModule Module;
-			public string[] IncludeChain;
+			public UEBuildModule Module = module;
+			public string[] IncludeChain = includeChain;
 			public HashSet<UEBuildModule> InwardRefs = [];
 			public HashSet<UEBuildModule> UniqueInwardRefs = [];
 			public HashSet<UEBuildModule> OutwardRefs = [];
@@ -62,105 +57,105 @@ namespace UnrealBuildTool.Modes
 			public List<FileReference> BinaryFiles = [];
 			public long BinSize = 0;
 
-			public ModuleInfo(UEBuildModule Module, params string[] IncludeChain)
-			{
-				this.Module = Module;
-				this.IncludeChain = IncludeChain;
-			}
-
 			public string Chain => String.Join(" -> ", IncludeChain.Where(x => !String.IsNullOrEmpty(x)));
 		}
 
-		private static void AnalyzeModuleChains(string ModuleName, List<string> ParentChain, UEBuildTarget Target, Dictionary<UEBuildModule, ModuleInfo> ModuleToInfo, HashSet<UEBuildModule> Visited, ILogger Logger)
+		private static void AnalyzeModuleChains(
+			string moduleName,
+			List<string> parentChain,
+			UEBuildTarget target,
+			Dictionary<UEBuildModule, ModuleInfo> moduleToInfo,
+			HashSet<UEBuildModule> visited,
+			ILogger logger)
 		{
 			// Prevent recursive includes, they'll never be shorter
-			if (ParentChain.Contains(ModuleName))
+			if (parentChain.Contains(moduleName))
 			{
 				return;
 			}
 
-			UEBuildModule Module = Target.GetModuleByName(ModuleName);
+			UEBuildModule module = target.GetModuleByName(moduleName);
 
-			List<string> CurrentChain = [.. ParentChain, Module.Name];
+			List<string> currentChain = [.. parentChain, module.Name];
 
-			if (!ModuleToInfo.TryGetValue(Module, out ModuleInfo? moduleInfo))
+			if (!moduleToInfo.TryGetValue(module, out ModuleInfo? moduleInfo))
 			{
-				moduleInfo = new ModuleInfo(Module, [.. CurrentChain]);
-				ModuleToInfo[Module] = moduleInfo;
+				moduleInfo = new ModuleInfo(module, [.. currentChain]);
+				moduleToInfo[module] = moduleInfo;
 			}
 
-			if (moduleInfo.IncludeChain.Length > CurrentChain.Count)
+			if (moduleInfo.IncludeChain.Length > currentChain.Count)
 			{
 				// Now we need to recheck all downstream dependencies again because the chain may be shorter
-				List<UEBuildModule> RecheckModules = [];
-				Module.GetAllDependencyModules(RecheckModules, [], bIncludeDynamicallyLoaded: true, bForceCircular: true, bOnlyDirectDependencies: true);
-				Visited.ExceptWith(RecheckModules);
-				Logger.LogDebug("Found shorter chain for {Module} {Prev} -> {New}, rechecking {Count} already visited dependencies", Module.Name, moduleInfo.IncludeChain.Length, CurrentChain.Count, RecheckModules.Count);
-				moduleInfo.IncludeChain = [.. CurrentChain];
+				List<UEBuildModule> recheckModules = [];
+				module.GetAllDependencyModules(recheckModules, [], bIncludeDynamicallyLoaded: true, bForceCircular: true, bOnlyDirectDependencies: true);
+				visited.ExceptWith(recheckModules);
+				logger.LogDebug("Found shorter chain for {Module} {Prev} -> {New}, rechecking {Count} already visited dependencies", module.Name, moduleInfo.IncludeChain.Length, currentChain.Count, recheckModules.Count);
+				moduleInfo.IncludeChain = [.. currentChain];
 			}
-			else if (Visited.Contains(Module))
+			else if (visited.Contains(module))
 			{
 				return;
 			}
 
-			Visited.Add(Module);
+			visited.Add(module);
 
-			List<UEBuildModule> TargetModules = [];
-			Module.GetAllDependencyModules(TargetModules, [], bIncludeDynamicallyLoaded: true, bForceCircular: true, bOnlyDirectDependencies: true);
-			TargetModules.ForEach(x => AnalyzeModuleChains(x.Name, CurrentChain, Target, ModuleToInfo, Visited, Logger));
+			List<UEBuildModule> targetModules = [];
+			module.GetAllDependencyModules(targetModules, [], bIncludeDynamicallyLoaded: true, bForceCircular: true, bOnlyDirectDependencies: true);
+			targetModules.ForEach(x => AnalyzeModuleChains(x.Name, currentChain, target, moduleToInfo, visited, logger));
 		}
 
-		private static void AnalyzeTarget(TargetDescriptor TargetDescriptor, BuildConfiguration BuildConfiguration, ILogger Logger)
+		private static void AnalyzeTarget(TargetDescriptor targetDescriptor, BuildConfiguration buildConfiguration, ILogger logger)
 		{
 			// Create a makefile for the target
-			UEBuildTarget Target = UEBuildTarget.Create(TargetDescriptor, BuildConfiguration, Logger);
-			DirectoryReference.CreateDirectory(Target.ReceiptFileName.Directory);
+			UEBuildTarget target = UEBuildTarget.Create(targetDescriptor, buildConfiguration, logger);
+			DirectoryReference.CreateDirectory(target.ReceiptFileName.Directory);
 
 			// Find the shortest path from the target to each module
-			HashSet<UEBuildModule> Visited = [];
-			Dictionary<UEBuildModule, ModuleInfo> ModuleToInfo = [];
+			HashSet<UEBuildModule> visited = [];
+			Dictionary<UEBuildModule, ModuleInfo> moduleToInfo = [];
 
-			if (Target.Rules.LaunchModuleName != null)
+			if (target.Rules.LaunchModuleName != null)
 			{
-				AnalyzeModuleChains(Target.Rules.LaunchModuleName, ["target"], Target, ModuleToInfo, Visited, Logger);
+				AnalyzeModuleChains(target.Rules.LaunchModuleName, ["target"], target, moduleToInfo, visited, logger);
 			}
 
-			foreach (string RootModuleName in Target.Rules.ExtraModuleNames)
+			foreach (string rootModuleName in target.Rules.ExtraModuleNames)
 			{
-				AnalyzeModuleChains(RootModuleName, ["target"], Target, ModuleToInfo, Visited, Logger);
+				AnalyzeModuleChains(rootModuleName, ["target"], target, moduleToInfo, visited, logger);
 			}
 
 			// Also enable all the plugin modules
-			foreach (UEBuildPlugin Plugin in Target.BuildPlugins!)
+			foreach (UEBuildPlugin plugin in target.BuildPlugins!)
 			{
-				foreach (UEBuildModule Module in Plugin.Modules)
+				foreach (UEBuildModule module in plugin.Modules)
 				{
-					if (!ModuleToInfo.ContainsKey(Module))
+					if (!moduleToInfo.ContainsKey(module))
 					{
-						AnalyzeModuleChains(Module.Name, ["target", Plugin.ReferenceChain, Plugin.File.GetFileName()], Target, ModuleToInfo, Visited, Logger);
+						AnalyzeModuleChains(module.Name, ["target", plugin.ReferenceChain, plugin.File.GetFileName()], target, moduleToInfo, visited, logger);
 					}
 				}
 			}
 
-			if (Target.Rules.bBuildAllModules)
+			if (target.Rules.bBuildAllModules)
 			{
-				foreach (UEBuildBinary Binary in Target.Binaries)
+				foreach (UEBuildBinary binary in target.Binaries)
 				{
-					foreach (UEBuildModule Module in Binary.Modules)
+					foreach (UEBuildModule module in binary.Modules)
 					{
-						if (!ModuleToInfo.ContainsKey(Module))
+						if (!moduleToInfo.ContainsKey(module))
 						{
 							// quick hack to make allmodules always worse (empty entries are ignored when writing)
-							List<string> IncludeChain = [.. Enumerable.Repeat(String.Empty, 1000)];
-							IncludeChain.Add("allmodules option");
-							if (Module.Rules.Plugin != null)
+							List<string> includeChain = [.. Enumerable.Repeat(String.Empty, 1000)];
+							includeChain.Add("allmodules option");
+							if (module.Rules.Plugin != null)
 							{
-								IncludeChain.Add(Module.Rules.Plugin.File.GetFileName());
-								AnalyzeModuleChains(Module.Name, IncludeChain, Target, ModuleToInfo, Visited, Logger);
+								includeChain.Add(module.Rules.Plugin.File.GetFileName());
+								AnalyzeModuleChains(module.Name, includeChain, target, moduleToInfo, visited, logger);
 							}
 							else
 							{
-								AnalyzeModuleChains(Module.Name, IncludeChain, Target, ModuleToInfo, Visited, Logger);
+								AnalyzeModuleChains(module.Name, includeChain, target, moduleToInfo, visited, logger);
 							}
 						}
 					}
@@ -168,257 +163,266 @@ namespace UnrealBuildTool.Modes
 			}
 
 			// Find all the outward dependencies of each module
-			foreach ((UEBuildModule SourceModule, ModuleInfo SourceModuleInfo) in ModuleToInfo)
+			foreach ((UEBuildModule sourceModule, ModuleInfo sourceModuleInfo) in moduleToInfo)
 			{
-				SourceModuleInfo.OutwardRefs.Add(SourceModule);
-				SourceModule.GetAllDependencyModules([], SourceModuleInfo.OutwardRefs, bIncludeDynamicallyLoaded: false, bForceCircular: true, bOnlyDirectDependencies: false);
-				SourceModuleInfo.OutwardRefs.Remove(SourceModule);
+				sourceModuleInfo.OutwardRefs.Add(sourceModule);
+				sourceModule.GetAllDependencyModules([], sourceModuleInfo.OutwardRefs, bIncludeDynamicallyLoaded: false, bForceCircular: true, bOnlyDirectDependencies: false);
+				sourceModuleInfo.OutwardRefs.Remove(sourceModule);
 			}
 
 			// Find the unique output dependencies of each module
-			foreach ((UEBuildModule SourceModule, ModuleInfo SourceModuleInfo) in ModuleToInfo)
+			foreach ((UEBuildModule sourceModule, ModuleInfo sourceModuleInfo) in moduleToInfo)
 			{
-				SourceModuleInfo.UniqueOutwardRefs = [.. SourceModuleInfo.OutwardRefs];
+				sourceModuleInfo.UniqueOutwardRefs = [.. sourceModuleInfo.OutwardRefs];
 
-				foreach (UEBuildModule TargetModule in SourceModuleInfo.OutwardRefs)
+				foreach (UEBuildModule targetModule in sourceModuleInfo.OutwardRefs)
 				{
-					HashSet<UEBuildModule> VisitedTargetModules = [SourceModule];
+					HashSet<UEBuildModule> visitedTargetModules = [sourceModule];
 
-					List<UEBuildModule> DependencyModules = [];
-					TargetModule.GetAllDependencyModules(DependencyModules, VisitedTargetModules, bIncludeDynamicallyLoaded: false, bForceCircular: true, bOnlyDirectDependencies: false);
-					DependencyModules.Remove(TargetModule);
+					List<UEBuildModule> dependencyModules = [];
+					targetModule.GetAllDependencyModules(dependencyModules, visitedTargetModules, bIncludeDynamicallyLoaded: false, bForceCircular: true, bOnlyDirectDependencies: false);
+					dependencyModules.Remove(targetModule);
 
-					SourceModuleInfo.UniqueOutwardRefs.ExceptWith(DependencyModules);
+					sourceModuleInfo.UniqueOutwardRefs.ExceptWith(dependencyModules);
 				}
 			}
 
 			// Find the direct inward dependencies of each module
-			foreach ((UEBuildModule SourceModule, ModuleInfo SourceModuleInfo) in ModuleToInfo)
+			foreach ((UEBuildModule sourceModule, ModuleInfo sourceModuleInfo) in moduleToInfo)
 			{
-				foreach (UEBuildModule TargetModule in SourceModuleInfo.OutwardRefs)
+				foreach (UEBuildModule targetModule in sourceModuleInfo.OutwardRefs)
 				{
-					ModuleToInfo[TargetModule].InwardRefs.Add(SourceModule);
+					moduleToInfo[targetModule].InwardRefs.Add(sourceModule);
 				}
-				foreach (UEBuildModule TargetModule in SourceModuleInfo.UniqueOutwardRefs)
+				foreach (UEBuildModule targetModule in sourceModuleInfo.UniqueOutwardRefs)
 				{
-					ModuleToInfo[TargetModule].UniqueInwardRefs.Add(SourceModule);
+					moduleToInfo[targetModule].UniqueInwardRefs.Add(sourceModule);
 				}
 			}
 
 			// Estimate the size of object files for each module
-			foreach ((UEBuildModule SourceModule, ModuleInfo SourceModuleInfo) in ModuleToInfo)
+			foreach ((UEBuildModule sourceModule, ModuleInfo sourceModuleInfo) in moduleToInfo)
 			{
-				if (DirectoryReference.Exists(SourceModule.IntermediateDirectory))
+				if (DirectoryReference.Exists(sourceModule.IntermediateDirectory))
 				{
-					foreach (FileReference IntermediateFile in DirectoryReference.EnumerateFiles(SourceModule.IntermediateDirectory, "*", SearchOption.AllDirectories))
+					foreach (FileReference intermediateFile in DirectoryReference.EnumerateFiles(sourceModule.IntermediateDirectory, "*", SearchOption.AllDirectories))
 					{
-						if (IntermediateFile.HasExtension(".obj") || IntermediateFile.HasExtension(".o"))
+						if (intermediateFile.HasExtension(".obj") || intermediateFile.HasExtension(".o"))
 						{
-							SourceModuleInfo.ObjectFiles.Add(IntermediateFile);
-							SourceModuleInfo.ObjSize += IntermediateFile.ToFileInfo().Length;
+							sourceModuleInfo.ObjectFiles.Add(intermediateFile);
+							sourceModuleInfo.ObjSize += intermediateFile.ToFileInfo().Length;
 						}
 					}
 				}
 			}
 
-			HashSet<UEBuildModule> MissingModules = [];
-			foreach (UEBuildBinary Binary in Target.Binaries)
+			HashSet<UEBuildModule> missingModules = [];
+			foreach (UEBuildBinary binary in target.Binaries)
 			{
-				long BinSize = 0;
-				foreach (FileReference OutputFilePath in Binary.OutputFilePaths)
+				long binSize = 0;
+				foreach (FileReference outputFilePath in binary.OutputFilePaths)
 				{
-					FileInfo OutputFileInfo = OutputFilePath.ToFileInfo();
-					if (OutputFileInfo.Exists)
+					FileInfo outputFileInfo = outputFilePath.ToFileInfo();
+					if (outputFileInfo.Exists)
 					{
-						BinSize += OutputFileInfo.Length;
+						binSize += outputFileInfo.Length;
 					}
 				}
-				foreach (UEBuildModule Module in Binary.Modules)
+				foreach (UEBuildModule module in binary.Modules)
 				{
-					if (!ModuleToInfo.TryGetValue(Module, out ModuleInfo? ModuleInfo))
+					if (!moduleToInfo.TryGetValue(module, out ModuleInfo? moduleInfo))
 					{
-						MissingModules.Add(Module);
+						missingModules.Add(module);
 						continue;
 					}
 
-					ModuleInfo.BinaryFiles.AddRange(Binary.OutputFilePaths);
-					ModuleInfo.BinSize += BinSize;
+					moduleInfo.BinaryFiles.AddRange(binary.OutputFilePaths);
+					moduleInfo.BinSize += binSize;
 				}
 			}
 
 			// Warn about any missing modules
-			foreach (UEBuildModule MissingModule in MissingModules.OrderBy(x => x.Name))
+			foreach (UEBuildModule missingModule in missingModules.OrderBy(x => x.Name))
 			{
-				Logger.LogInformation("Missing module '{MissingModuleName}'", MissingModule.Name);
+				logger.LogInformation("Missing module '{MissingModuleName}'", missingModule.Name);
 			}
 
-			List<KeyValuePair<FileReference, BuildProductType>> AnalyzeProducts = [];
+			List<KeyValuePair<FileReference, BuildProductType>> analyzeProducts = [];
 
 			// Generate the dependency graph between modules
-			FileReference DependencyGraphFile = Target.ReceiptFileName.ChangeExtension(".Dependencies.gexf");
-			Logger.LogInformation("Writing dependency graph to {DependencyGraphFile}...", DependencyGraphFile);
-			WriteDependencyGraph(Target, ModuleToInfo, DependencyGraphFile);
-			AnalyzeProducts.Add(new(DependencyGraphFile, BuildProductType.BuildResource));
+			FileReference dependencyGraphFile = target.ReceiptFileName.ChangeExtension(".Dependencies.gexf");
+			logger.LogInformation("Writing dependency graph to {DependencyGraphFile}...", dependencyGraphFile);
+			WriteDependencyGraph(target, moduleToInfo, dependencyGraphFile);
+			analyzeProducts.Add(new(dependencyGraphFile, BuildProductType.BuildResource));
 
 			// Generate the dependency graph between modules
-			FileReference ShortestPathGraphFile = Target.ReceiptFileName.ChangeExtension(".ShortestPath.gexf");
-			Logger.LogInformation("Writing shortest-path graph to {ShortestPathGraphFile}...", ShortestPathGraphFile);
-			WriteShortestPathGraph(Target, ModuleToInfo, ShortestPathGraphFile);
-			AnalyzeProducts.Add(new(ShortestPathGraphFile, BuildProductType.BuildResource));
+			FileReference shortestPathGraphFile = target.ReceiptFileName.ChangeExtension(".ShortestPath.gexf");
+			logger.LogInformation("Writing shortest-path graph to {ShortestPathGraphFile}...", shortestPathGraphFile);
+			WriteShortestPathGraph(target, moduleToInfo, shortestPathGraphFile);
+			analyzeProducts.Add(new(shortestPathGraphFile, BuildProductType.BuildResource));
 
 			// Write all the target stats as a text file
-			FileReference TextFile = Target.ReceiptFileName.ChangeExtension(".txt");
-			Logger.LogInformation("Writing module information to {TextFile}", TextFile);
-			using (StreamWriter Writer = new(TextFile.FullName))
-			{
-				Writer.WriteLine("All modules in {0}, ordered by number of indirect references", Target.TargetName);
-
-				foreach (ModuleInfo ModuleInfo in ModuleToInfo.Values.OrderByDescending(x => x.InwardRefs.Count).ThenBy(x => x.BinSize))
-				{
-					Writer.WriteLine("");
-					Writer.WriteLine("Module:                  \"{0}\"", ModuleInfo.Module.Name);
-					Writer.WriteLine("Shortest path:           {0}", ModuleInfo.Chain);
-					WriteDependencyList(Writer, "Unique inward refs:     ", ModuleInfo.UniqueInwardRefs);
-					WriteDependencyList(Writer, "Unique outward refs:    ", ModuleInfo.UniqueOutwardRefs);
-					WriteDependencyList(Writer, "Recursive inward refs:  ", ModuleInfo.InwardRefs);
-					WriteDependencyList(Writer, "Recursive outward refs: ", ModuleInfo.OutwardRefs);
-					Writer.WriteLine("Object size:             {0:n0}kb", (ModuleInfo.ObjSize + 1023) / 1024);
-					Writer.WriteLine("Object files:            {0}", String.Join(", ", ModuleInfo.ObjectFiles.Select(x => x.GetFileName())));
-					Writer.WriteLine("Binary size:             {0:n0}kb", (ModuleInfo.BinSize + 1023) / 1024);
-					Writer.WriteLine("Binary files:            {0}", String.Join(", ", ModuleInfo.BinaryFiles.Select(x => x.GetFileName())));
-				}
-			}
-			AnalyzeProducts.Add(new(TextFile, BuildProductType.BuildResource));
+			FileReference textFile = target.ReceiptFileName.ChangeExtension(".txt");
+			logger.LogInformation("Writing module information to {TextFile}", textFile);
+			WriteInformationTxt(target, moduleToInfo, textFile);
+			analyzeProducts.Add(new(textFile, BuildProductType.BuildResource));
 
 			// Write all the target stats as a CSV file
-			FileReference CsvFile = Target.ReceiptFileName.ChangeExtension(".csv");
-			Logger.LogInformation("Writing module information to {CsvFile}", CsvFile);
-			using (StreamWriter Writer = new(CsvFile.FullName))
-			{
-				List<string> Columns =
-				[
-					"Module",
-					"ShortestPath",
-					"NumUniqueInwardRefs",
-					"UniqueInwardRefs",
-					"NumRecursiveInwardRefs",
-					"RecursiveInwardRefs",
-					"NumUniqueOutwardRefs",
-					"UniqueOutwardRefs",
-					"NumRecursiveOutwardRefs",
-					"RecursiveOutwardRefs",
-					"ObjSize",
-					"ObjFiles",
-					"BinSize",
-					"BinFiles",
-				];
-				Writer.WriteLine(String.Join(",", Columns));
+			FileReference csvFile = target.ReceiptFileName.ChangeExtension(".csv");
+			logger.LogInformation("Writing module information to {CsvFile}", csvFile);
+			WriteInformationCsv(moduleToInfo, csvFile);
+			analyzeProducts.Add(new(csvFile, BuildProductType.BuildResource));
 
-				foreach (ModuleInfo ModuleInfo in ModuleToInfo.Values.OrderByDescending(x => x.InwardRefs.Count).ThenBy(x => x.BinSize))
-				{
-					Columns.Clear();
-					Columns.Add(ModuleInfo.Module.Name);
-					Columns.Add(ModuleInfo.Chain);
-					Columns.Add($"{ModuleInfo.UniqueInwardRefs.Count}");
-					Columns.Add($"\"{String.Join(", ", ModuleInfo.UniqueInwardRefs.Select(x => x.Name))}\"");
-					Columns.Add($"{ModuleInfo.InwardRefs.Count}");
-					Columns.Add($"\"{String.Join(", ", ModuleInfo.InwardRefs.Select(x => x.Name))}\"");
-					Columns.Add($"{ModuleInfo.UniqueOutwardRefs.Count}");
-					Columns.Add($"\"{String.Join(", ", ModuleInfo.UniqueOutwardRefs.Select(x => x.Name))}\"");
-					Columns.Add($"{ModuleInfo.OutwardRefs.Count}");
-					Columns.Add($"\"{String.Join(", ", ModuleInfo.OutwardRefs.Select(x => x.Name))}\"");
-					Columns.Add($"{ModuleInfo.ObjSize}");
-					Columns.Add($"\"{String.Join(", ", ModuleInfo.ObjectFiles.Select(x => x.GetFileName()))}\"");
-					Columns.Add($"{ModuleInfo.BinSize}");
-					Columns.Add($"\"{String.Join(", ", ModuleInfo.BinaryFiles.Select(x => x.GetFileName()))}\"");
-					Writer.WriteLine(String.Join(",", Columns));
-				}
-			}
-			AnalyzeProducts.Add(new(CsvFile, BuildProductType.BuildResource));
-
-			foreach (FileReference ManifestFileName in Target.Rules.ManifestFileNames)
+			foreach (FileReference manifestFileName in target.Rules.ManifestFileNames)
 			{
-				Target.GenerateManifest(ManifestFileName, AnalyzeProducts, Logger);
+				target.GenerateManifest(manifestFileName, analyzeProducts, logger);
 			}
 		}
 
-		private static void WriteDependencyList(TextWriter Writer, string Prefix, HashSet<UEBuildModule> Modules)
+		private static void WriteDependencyList(TextWriter writer, string prefix, HashSet<UEBuildModule> modules)
 		{
-			if (Modules.Count == 0)
+			if (modules.Count == 0)
 			{
-				Writer.WriteLine("{0} 0", Prefix);
+				writer.WriteLine("{0} 0", prefix);
 			}
 			else
 			{
-				Writer.WriteLine("{0} {1} ({2})", Prefix, Modules.Count, String.Join(", ", Modules.Select(x => x.Name).OrderBy(x => x)));
+				writer.WriteLine("{0} {1} ({2})", prefix, modules.Count, String.Join(", ", modules.Select(x => x.Name).OrderBy(x => x)));
 			}
 		}
 
-		private static void WriteDependencyGraph(UEBuildTarget Target, Dictionary<UEBuildModule, ModuleInfo> ModuleToInfo, FileReference FileName)
+		private static void WriteDependencyGraph(UEBuildTarget target, Dictionary<UEBuildModule, ModuleInfo> moduleToInfo, FileReference fileName)
 		{
-			List<GraphNode> Nodes = [];
+			List<GraphNode> nodes = [];
 
-			Dictionary<UEBuildModule, GraphNode> ModuleToNode = [];
-			foreach (ModuleInfo ModuleInfo in ModuleToInfo.Values)
+			Dictionary<UEBuildModule, GraphNode> moduleToNode = [];
+			foreach (ModuleInfo moduleInfo in moduleToInfo.Values)
 			{
-				GraphNode Node = new(ModuleInfo.Module.Name);
+				GraphNode node = new(moduleInfo.Module.Name);
 
-				long Size = Target.ShouldCompileMonolithic() ? ModuleInfo.ObjSize : ModuleInfo.BinSize;
-				Node.Size = 1.0f + (Size / (50.0f * 1024.0f * 1024.0f));
-				Nodes.Add(Node);
-				ModuleToNode[ModuleInfo.Module] = Node;
+				long size = target.ShouldCompileMonolithic() ? moduleInfo.ObjSize : moduleInfo.BinSize;
+				node.Size = 1.0f + (size / (50.0f * 1024.0f * 1024.0f));
+				nodes.Add(node);
+				moduleToNode[moduleInfo.Module] = node;
 			}
 
-			List<GraphEdge> Edges = [];
-			foreach ((UEBuildModule SourceModule, ModuleInfo SourceModuleInfo) in ModuleToInfo)
+			List<GraphEdge> edges = [];
+			foreach ((UEBuildModule sourceModule, ModuleInfo sourceModuleInfo) in moduleToInfo)
 			{
-				GraphNode SourceNode = ModuleToNode[SourceModule];
-				foreach (UEBuildModule TargetModule in SourceModuleInfo.UniqueOutwardRefs)
+				GraphNode sourceNode = moduleToNode[sourceModule];
+				foreach (UEBuildModule targetModule in sourceModuleInfo.UniqueOutwardRefs)
 				{
-					ModuleInfo TargetModuleInfo = ModuleToInfo[TargetModule];
+					ModuleInfo targetModuleInfo = moduleToInfo[targetModule];
 
-					if (ModuleToNode.TryGetValue(TargetModule, out GraphNode? TargetNode))
+					if (moduleToNode.TryGetValue(targetModule, out GraphNode? targetNode))
 					{
-						Edges.Add(new(SourceNode, TargetNode)
+						edges.Add(new(sourceNode, targetNode)
 						{
-							Thickness = TargetModuleInfo.InwardRefs.Count
+							Thickness = targetModuleInfo.InwardRefs.Count
 						});
 					}
 				}
 			}
 
-			GraphVisualization.WriteGraphFile(FileName, $"Module dependency graph for {Target.TargetName}", Nodes, Edges);
+			GraphVisualization.WriteGraphFile(fileName, $"Module dependency graph for {target.TargetName}", nodes, edges);
 		}
 
-		private static void WriteShortestPathGraph(UEBuildTarget Target, Dictionary<UEBuildModule, ModuleInfo> ModuleToInfo, FileReference FileName)
+		private static void WriteShortestPathGraph(UEBuildTarget target, Dictionary<UEBuildModule, ModuleInfo> moduleToInfo, FileReference fileName)
 		{
-			Dictionary<string, GraphNode> NameToNode = new(StringComparer.Ordinal);
+			Dictionary<string, GraphNode> nameToNode = new(StringComparer.Ordinal);
 
-			HashSet<(GraphNode, GraphNode)> EdgesSet = [];
-			List<GraphEdge> Edges = [];
+			HashSet<(GraphNode, GraphNode)> edgesSet = [];
+			List<GraphEdge> edges = [];
 
-			foreach ((UEBuildModule Module, ModuleInfo ModuleInfo) in ModuleToInfo)
+			foreach ((UEBuildModule module, ModuleInfo moduleInfo) in moduleToInfo)
 			{
-				string[] Parts = ModuleInfo.Chain.Split(" -> ");
+				string[] parts = moduleInfo.Chain.Split(" -> ");
 
-				GraphNode? PrevNode = null;
-				foreach (string Part in Parts)
+				GraphNode? prevNode = null;
+				foreach (string part in parts)
 				{
-					if (!NameToNode.TryGetValue(Part, out GraphNode? NextNode))
+					if (!nameToNode.TryGetValue(part, out GraphNode? nextNode))
 					{
-						NextNode = new GraphNode(Part);
-						NameToNode[Part] = NextNode;
+						nextNode = new GraphNode(part);
+						nameToNode[part] = nextNode;
 					}
-					if (PrevNode != null && EdgesSet.Add((PrevNode, NextNode)))
+					if (prevNode != null && edgesSet.Add((prevNode, nextNode)))
 					{
-						Edges.Add(new(PrevNode, NextNode));
+						edges.Add(new(prevNode, nextNode));
 					}
-					PrevNode = NextNode;
+					prevNode = nextNode;
+				}
+			}
+
+			GraphVisualization.WriteGraphFile(fileName, $"Module dependency graph for {target.TargetName}", [.. nameToNode.Values], edges);
+		}
+
+		private static void WriteInformationTxt(UEBuildTarget target, Dictionary<UEBuildModule, ModuleInfo> moduleToInfo, FileReference textFile)
+		{
+			using (StreamWriter writer = new(textFile.FullName))
+			{
+				writer.WriteLine("All modules in {0}, ordered by number of indirect references", target.TargetName);
+
+				foreach (ModuleInfo moduleInfo in moduleToInfo.Values.OrderByDescending(x => x.InwardRefs.Count).ThenBy(x => x.BinSize))
+				{
+					writer.WriteLine("");
+					writer.WriteLine("Module:                  \"{0}\"", moduleInfo.Module.Name);
+					writer.WriteLine("Shortest path:           {0}", moduleInfo.Chain);
+					WriteDependencyList(writer, "Unique inward refs:     ", moduleInfo.UniqueInwardRefs);
+					WriteDependencyList(writer, "Unique outward refs:    ", moduleInfo.UniqueOutwardRefs);
+					WriteDependencyList(writer, "Recursive inward refs:  ", moduleInfo.InwardRefs);
+					WriteDependencyList(writer, "Recursive outward refs: ", moduleInfo.OutwardRefs);
+					writer.WriteLine("Object size:             {0:n0}kb", (moduleInfo.ObjSize + 1023) / 1024);
+					writer.WriteLine("Object files:            {0}", String.Join(", ", moduleInfo.ObjectFiles.Select(x => x.GetFileName())));
+					writer.WriteLine("Binary size:             {0:n0}kb", (moduleInfo.BinSize + 1023) / 1024);
+					writer.WriteLine("Binary files:            {0}", String.Join(", ", moduleInfo.BinaryFiles.Select(x => x.GetFileName())));
 				}
 			}
+		}
 
-			GraphVisualization.WriteGraphFile(FileName, $"Module dependency graph for {Target.TargetName}", [.. NameToNode.Values], Edges);
+		private static void WriteInformationCsv(Dictionary<UEBuildModule, ModuleInfo> moduleToInfo, FileReference csvFile)
+		{
+			using (StreamWriter writer = new(csvFile.FullName))
+			{
+				List<string> columns =
+				[
+					"Module",
+					"ShortestPath",
+					"NumUniqueInwardRefs",
+					"UniqueInwardRefs",
+					"NumRecursiveInwardRefs",
+					"RecursiveInwardRefs",
+					"NumUniqueOutwardRefs",
+					"UniqueOutwardRefs",
+					"NumRecursiveOutwardRefs",
+					"RecursiveOutwardRefs",
+					"ObjSize",
+					"ObjFiles",
+					"BinSize",
+					"BinFiles",
+				];
+				writer.WriteLine(String.Join(",", columns));
+
+				foreach (ModuleInfo moduleInfo in moduleToInfo.Values.OrderByDescending(x => x.InwardRefs.Count).ThenBy(x => x.BinSize))
+				{
+					columns.Clear();
+					columns.Add(moduleInfo.Module.Name);
+					columns.Add(moduleInfo.Chain);
+					columns.Add($"{moduleInfo.UniqueInwardRefs.Count}");
+					columns.Add($"\"{String.Join(", ", moduleInfo.UniqueInwardRefs.Select(x => x.Name))}\"");
+					columns.Add($"{moduleInfo.InwardRefs.Count}");
+					columns.Add($"\"{String.Join(", ", moduleInfo.InwardRefs.Select(x => x.Name))}\"");
+					columns.Add($"{moduleInfo.UniqueOutwardRefs.Count}");
+					columns.Add($"\"{String.Join(", ", moduleInfo.UniqueOutwardRefs.Select(x => x.Name))}\"");
+					columns.Add($"{moduleInfo.OutwardRefs.Count}");
+					columns.Add($"\"{String.Join(", ", moduleInfo.OutwardRefs.Select(x => x.Name))}\"");
+					columns.Add($"{moduleInfo.ObjSize}");
+					columns.Add($"\"{String.Join(", ", moduleInfo.ObjectFiles.Select(x => x.GetFileName()))}\"");
+					columns.Add($"{moduleInfo.BinSize}");
+					columns.Add($"\"{String.Join(", ", moduleInfo.BinaryFiles.Select(x => x.GetFileName()))}\"");
+					writer.WriteLine(String.Join(",", columns));
+				}
+			}
 		}
 	}
 }
-
