[UBA] Making the launch visualizer option configurable via UserEngine.ini file in addition to the xml config file and commandline.

#Synced-CL 50647103

diff --git a/Engine/Source/Programs/UnrealBuildTool/Executors/UnrealBuildAccelerator/UnrealBuildAcceleratorConfig.cs b/Engine/Source/Programs/UnrealBuildTool/Executors/UnrealBuildAccelerator/UnrealBuildAcceleratorConfig.cs
--- a/Engine/Source/Programs/UnrealBuildTool/Executors/UnrealBuildAccelerator/UnrealBuildAcceleratorConfig.cs
+++ b/Engine/Source/Programs/UnrealBuildTool/Executors/UnrealBuildAccelerator/UnrealBuildAcceleratorConfig.cs
@@ -202,6 +202,7 @@ namespace UnrealBuildTool
 		/// Launch visualizer application which shows build progress
 		/// </summary>
 		[XmlConfigFile(Category = "UnrealBuildAccelerator")]
+		[ConfigFile(ConfigHierarchyType.Engine, "UbaController", "bLaunchVisualizer")]
 		[CommandLine("-UBAVisualizer", Value = "true")]
 		public bool bLaunchVisualizer { get; set; } = false;
 
