[UBA]
* Enabled DeferRemoteCapableProcesses for macos to have local machine not take remote capable processes if there are non-capable queued later

#Synced-CL 50893966

diff --git a/Engine/Source/Programs/UnrealBuildTool/Executors/UnrealBuildAccelerator/UBAExecutor.cs b/Engine/Source/Programs/UnrealBuildTool/Executors/UnrealBuildAccelerator/UBAExecutor.cs
--- a/Engine/Source/Programs/UnrealBuildTool/Executors/UnrealBuildAccelerator/UBAExecutor.cs
+++ b/Engine/Source/Programs/UnrealBuildTool/Executors/UnrealBuildAccelerator/UBAExecutor.cs
@@ -370,6 +370,7 @@ namespace UnrealBuildTool
 			if (OperatingSystem.IsMacOS())
 			{ 				
 				AppleToolChain.RegisterCrossArchitecturePathsForUBA(logger);
+				_config.AddValue("Scheduler", "DeferRemoteCapableProcesses", "true");
 			}
 	
 			if (UBAConfig.MaxRacingPercent > 0)
@@ -1000,7 +1001,7 @@ namespace UnrealBuildTool
 
 								if (mustRunDetoured && userData.RetryCount <= 3)
 								{
-									logger.LogInformation("{Description} {StatusDescription}: Exited with error code {ExitCode}. This action will retry", action.CommandDescription, action.StatusDescription, exitCode);
+									logger.LogInformation("{Description} {StatusDescription}: Exited with error code {ExitCode}. This action will retry (More info in uba trace)", action.CommandDescription, action.StatusDescription, exitCode);
 									++forcedRetryActions;
 									return ProcessFinishedResponse.RerunLocal;
 								}
