[UBT]
* Exposed ways to disable individual agents in uba. This can be used to disable certain types of agents when getting others

#Synced-CL 50893966

diff --git a/Engine/Source/Programs/Shared/EpicGames.UBA/ISessionServer.cs b/Engine/Source/Programs/Shared/EpicGames.UBA/ISessionServer.cs
--- a/Engine/Source/Programs/Shared/EpicGames.UBA/ISessionServer.cs
+++ b/Engine/Source/Programs/Shared/EpicGames.UBA/ISessionServer.cs
@@ -155,6 +155,11 @@ namespace EpicGames.UBA
 		/// </summary>
 		public abstract void DisableRemoteExecution();
 
+		/// <summary>
+		/// Sets specific agent to finish current work and then disconnect
+		/// </summary>
+		public abstract bool DisableRemoteExecutionOnAgent(string agentName);
+
 		/// <summary>
 		/// Set max number of processes that can be executed remotely.
 		/// Setting this can let the backend disconnect remote workers earlier
diff --git a/Engine/Source/Programs/Shared/EpicGames.UBA/Impl/SessionServerImpl.cs b/Engine/Source/Programs/Shared/EpicGames.UBA/Impl/SessionServerImpl.cs
--- a/Engine/Source/Programs/Shared/EpicGames.UBA/Impl/SessionServerImpl.cs
+++ b/Engine/Source/Programs/Shared/EpicGames.UBA/Impl/SessionServerImpl.cs
@@ -119,6 +119,9 @@ namespace EpicGames.UBA.Impl
 		[DllImport("UbaHost", CharSet = CharSet.Auto)]
 		static extern void SessionServer_DisableRemoteExecution(nint server);
 
+		[DllImport("UbaHost", CharSet = CharSet.Auto)]
+		static extern byte SessionServer_DisableRemoteExecutionOnAgent(nint server, string agentName);
+
 		[DllImport("UbaHost", CharSet = CharSet.Auto)]
 		static extern void SessionServer_SetCustomCasKeyFromTrackedInputs(nint server, nint process, string filename, string workingdir);
 
@@ -234,6 +237,8 @@ namespace EpicGames.UBA.Impl
 
 		public void DisableRemoteExecution() => SessionServer_DisableRemoteExecution(_handle);
 
+		public bool DisableRemoteExecutionOnAgent(string agentName) => SessionServer_DisableRemoteExecutionOnAgent(_handle, agentName) != 0;
+
 		public void AddInfo(string info) => SessionServer_AddInfo(_handle, info);
 
 		public void SetMaxRemoteProcessCount(uint count) => SessionServer_SetMaxRemoteProcessCount(_handle, count);
diff --git a/Engine/Source/Programs/UnrealBuildTool/Executors/UnrealBuildAccelerator/UBAAgentCoordinator.cs b/Engine/Source/Programs/UnrealBuildTool/Executors/UnrealBuildAccelerator/UBAAgentCoordinator.cs
--- a/Engine/Source/Programs/UnrealBuildTool/Executors/UnrealBuildAccelerator/UBAAgentCoordinator.cs
+++ b/Engine/Source/Programs/UnrealBuildTool/Executors/UnrealBuildAccelerator/UBAAgentCoordinator.cs
@@ -16,6 +16,7 @@ namespace UnrealBuildTool
 		bool IsEmpty { get; }
 		void GetProcessWeightThatCanRunRemotelyNow(out double totalWeight, out double crossArchitecture, out double crossPlatform);
 		bool AddClient(string ip, int port, string crypto = "");
+		bool DisableRemoteExecutionOnAgent(string agentName);
 	}
 
 	interface IUBAAgentCoordinator
diff --git a/Engine/Source/Programs/UnrealBuildTool/Executors/UnrealBuildAccelerator/UBAExecutor.cs b/Engine/Source/Programs/UnrealBuildTool/Executors/UnrealBuildAccelerator/UBAExecutor.cs
--- a/Engine/Source/Programs/UnrealBuildTool/Executors/UnrealBuildAccelerator/UBAExecutor.cs
+++ b/Engine/Source/Programs/UnrealBuildTool/Executors/UnrealBuildAccelerator/UBAExecutor.cs
@@ -799,13 +799,15 @@ namespace UnrealBuildTool
 
 		class AgentCoordinatorScheduler : IUbaAgentCoordinatorScheduler
 		{
-			internal AgentCoordinatorScheduler(IScheduler scheduler, IServer server)
+			internal AgentCoordinatorScheduler(IScheduler scheduler, ISessionServer session, IServer server)
 			{
 				_scheduler = scheduler;
+				_session = session;
 				_server = server;
 			}
 
 			internal IScheduler? _scheduler;
+			internal ISessionServer? _session;
 			internal IServer? _server;
 			private readonly object _lockObject = new();
 
@@ -814,6 +816,7 @@ namespace UnrealBuildTool
 				lock (_lockObject)
 				{
 					_scheduler = null;
+					_session = null;
 					_server = null;
 				}
 			}
@@ -854,6 +857,18 @@ namespace UnrealBuildTool
 				}
 				return false;
 			}
+			
+			public bool DisableRemoteExecutionOnAgent(string agentName)
+			{
+				lock (_lockObject)
+				{
+					if (_session != null)
+					{
+						return _session.DisableRemoteExecutionOnAgent(agentName);
+					}
+				}
+				return false;
+			}
 		}
 
 		class ProcessUserData
@@ -1154,7 +1169,7 @@ namespace UnrealBuildTool
 
 			try
 			{
-				AgentCoordinatorScheduler coordinatorScheduler = new(scheduler, server);
+				AgentCoordinatorScheduler coordinatorScheduler = new(scheduler, session, server);
 				foreach (IUBAAgentCoordinator coordinator in _agentCoordinators)
 				{
 					coordinator.Start(coordinatorScheduler, CanRunRemotely);
