UnrealBuildTool: Replace usage of Ionic.Zip.Reduced.dll with System.IO.Compression. Remove unused internal IOS-specific code.
#Synced-CL 50903846

diff --git a/Engine/Source/Programs/UnrealBuildTool/Platform/Android/AndroidAARHandler.cs b/Engine/Source/Programs/UnrealBuildTool/Platform/Android/AndroidAARHandler.cs
--- a/Engine/Source/Programs/UnrealBuildTool/Platform/Android/AndroidAARHandler.cs
+++ b/Engine/Source/Programs/UnrealBuildTool/Platform/Android/AndroidAARHandler.cs
@@ -105,19 +105,6 @@ namespace UnrealBuildTool
 			}
 		}
 
-		public void DumpAAR(ILogger Logger)
-		{
-			Logger.LogInformation("ALL DEPENDENCIES");
-			foreach (AndroidAAREntry Entry in AARList!)
-			{
-				Logger.LogInformation("{EntryFilename}", Entry.Filename);
-			}
-			foreach (AndroidAAREntry Entry in JARList!)
-			{
-				Logger.LogInformation("{EntryFilename}", Entry.Filename);
-			}
-		}
-
 		private static string GetElementValue(XElement SourceElement, XName ElementName, string DefaultValue)
 		{
 			XElement? Element = SourceElement.Element(ElementName);
@@ -447,194 +434,5 @@ namespace UnrealBuildTool
 				}
 			}
 		}
-
-		private static void MakeDirectoryIfRequiredForFile(string DestFilename)
-		{
-			string DestSubdir = Path.GetDirectoryName(DestFilename)!;
-			if (!Directory.Exists(DestSubdir))
-			{
-				Directory.CreateDirectory(DestSubdir);
-			}
-		}
-
-		private static void MakeDirectoryIfRequired(string DestDirectory)
-		{
-			if (!Directory.Exists(DestDirectory))
-			{
-				Directory.CreateDirectory(DestDirectory);
-			}
-		}
-
-		/// <summary>
-		/// Copies the required JAR files to the provided directory
-		/// </summary>
-		/// <param name="DestinationPath">Destination path for JAR files</param>
-		/// <param name="Logger">Logger for output</param>
-		public void CopyJARs(string DestinationPath, ILogger Logger)
-		{
-			MakeDirectoryIfRequired(DestinationPath);
-			DestinationPath = Path.Combine(DestinationPath, "libs");
-			MakeDirectoryIfRequired(DestinationPath);
-
-			foreach (AndroidAAREntry Entry in JARList!)
-			{
-				string Filename = Entry.Filename + ".jar";
-				string BaseName = Path.GetFileName(Filename);
-				string TargetPath = Path.Combine(DestinationPath, BaseName);
-				//Logger.LogInformation("Attempting to copy JAR {Filename} {BaseName} {TargetPath}", Filename, BaseName, TargetPath);
-
-				if (!File.Exists(Filename))
-				{
-					Logger.LogInformation("JAR doesn't exist! {Filename}", Filename);
-				}
-				if (!File.Exists(TargetPath))
-				{
-					Logger.LogInformation("Copying JAR {BaseName}", BaseName);
-					File.Copy(Filename, TargetPath);
-				}
-			}
-		}
-
-		/// <summary>
-		/// Extracts the required AAR files to the provided directory
-		/// </summary>
-		/// <param name="DestinationPath">Destination path for AAR files</param>
-		/// <param name="AppPackageName">Name of the package these AARs are being used with</param>
-		/// <param name="Logger">Logger for output</param>
-		public void ExtractAARs(string DestinationPath, string AppPackageName, ILogger Logger)
-		{
-			MakeDirectoryIfRequired(DestinationPath);
-			DestinationPath = Path.Combine(DestinationPath, "JavaLibs");
-			MakeDirectoryIfRequired(DestinationPath);
-
-			Logger.LogInformation("Extracting AARs");
-			foreach (AndroidAAREntry Entry in AARList!)
-			{
-				string BaseName = Path.GetFileName(Entry.Filename);
-				string TargetPath = Path.Combine(DestinationPath, BaseName);
-
-				// Only extract if haven't before to prevent changing timestamps
-				string TargetManifestFileName = Path.Combine(TargetPath, "AndroidManifest.xml");
-				if (!File.Exists(TargetManifestFileName))
-				{
-					Logger.LogInformation("Extracting AAR {BaseName}", BaseName);
-					/*IEnumerable<string> FileNames =*/
-					UnzipFiles(Entry.Filename + ".aar", TargetPath, Logger);
-
-					// Must have a src directory (even if empty)
-					string SrcDirectory = Path.Combine(TargetPath, "src");
-					if (!Directory.Exists(SrcDirectory))
-					{
-						Directory.CreateDirectory(SrcDirectory);
-						//FileStream Placeholder = new FileStream(Path.Combine(SrcDirectory, ".placeholder_" + BaseName), FileMode.Create, FileAccess.Write);
-					}
-
-					// Move the class.jar file
-					string ClassSourceFile = Path.Combine(TargetPath, "classes.jar");
-					if (File.Exists(ClassSourceFile))
-					{
-						string ClassDestFile = Path.Combine(TargetPath, "libs", BaseName + ".jar");
-						MakeDirectoryIfRequiredForFile(ClassDestFile);
-						File.Move(ClassSourceFile, ClassDestFile);
-					}
-
-					// Now create the project.properties file
-					string PropertyFilename = Path.Combine(TargetPath, "project.properties");
-					if (!File.Exists(PropertyFilename))
-					{
-						// Try to get the SDK target from the AndroidManifest.xml
-						string MinSDK = "9";
-						string ManifestFilename = Path.Combine(TargetPath, "AndroidManifest.xml");
-						XDocument ManifestXML;
-						if (File.Exists(ManifestFilename))
-						{
-							try
-							{
-								// Replace all instances of the application id with the packagename
-								string Contents = File.ReadAllText(ManifestFilename);
-								string NewContents = Contents.Replace("${applicationId}", AppPackageName, StringComparison.Ordinal);
-								if (Contents != NewContents)
-								{
-									File.WriteAllText(ManifestFilename, NewContents);
-								}
-
-								ManifestXML = XDocument.Load(ManifestFilename);
-								XElement UsesSdk = ManifestXML.Root!.Element(XName.Get("uses-sdk", ManifestXML.Root.Name.NamespaceName))!;
-								XAttribute Target = UsesSdk.Attribute(XName.Get("minSdkVersion", "http://schemas.android.com/apk/res/android"))!;
-								MinSDK = Target.Value;
-							}
-							catch (Exception e)
-							{
-								Logger.LogError("AAR Manifest file {ManifestFile} parsing error! {Ex}", ManifestFilename, e);
-							}
-						}
-
-						// Project contents
-						string ProjectPropertiesContents = "target=android-" + MinSDK + "\nandroid.library=true\n";
-
-						// Add the dependencies
-						int RefCount = 0;
-						foreach (string DependencyName in Entry.Dependencies)
-						{
-							// Find the version
-							foreach (AndroidAAREntry ScanEntry in AARList)
-							{
-								if (ScanEntry.BaseName == DependencyName)
-								{
-									RefCount++;
-									ProjectPropertiesContents += "android.library.reference." + RefCount + "=../" + DependencyName + "-" + ScanEntry.Version + "\n";
-									break;
-								}
-							}
-						}
-
-						File.WriteAllText(PropertyFilename, ProjectPropertiesContents);
-					}
-				}
-			}
-		}
-
-		/// <summary>
-		/// Extracts the contents of a zip file
-		/// </summary>
-		/// <param name="ZipFileName">Name of the zip file</param>
-		/// <param name="BaseDirectory">Output directory</param>
-		/// <param name="Logger">Logger for output</param>
-		/// <returns>List of files written</returns>
-		public static IEnumerable<string> UnzipFiles(string ZipFileName, string BaseDirectory, ILogger Logger)
-		{
-			// manually extract the files. There was a problem with the Ionic.Zip library that required this on non-PC at one point,
-			// but that problem is now fixed. Leaving this code as is as we need to return the list of created files and fix up their permissions anyway.
-			using (Ionic.Zip.ZipFile Zip = new Ionic.Zip.ZipFile(ZipFileName))
-			{
-				List<string> OutputFileNames = new List<string>();
-				foreach (Ionic.Zip.ZipEntry Entry in Zip.Entries)
-				{
-					// support-v4 and support-v13 has the jar file named with "internal_impl-XX.X.X.jar"
-					// this causes error "Found 2 versions of internal_impl-XX.X.X.jar"
-					// following codes adds "support-v4-..." to the output jar file name to avoid the collision
-					string OutputFileName = Path.Combine(BaseDirectory, Entry.FileName);
-					if (Entry.FileName.Contains("internal_impl", StringComparison.Ordinal))
-					{
-						string _ZipName = Path.GetFileNameWithoutExtension(ZipFileName);
-						string NewOutputFileName = Path.Combine(Path.GetDirectoryName(OutputFileName)!,
-							_ZipName + '-' + Path.GetFileNameWithoutExtension(OutputFileName) + '.' + Path.GetExtension(OutputFileName));
-						Logger.LogInformation("Changed FileName {EntryFileName} => {NewOutputFileName}", Entry.FileName, NewOutputFileName);
-						OutputFileName = NewOutputFileName;
-					}
-
-					Directory.CreateDirectory(Path.GetDirectoryName(OutputFileName)!);
-					if (!Entry.IsDirectory)
-					{
-						using (FileStream OutputStream = new FileStream(OutputFileName, FileMode.Create, FileAccess.Write))
-						{
-							Entry.Extract(OutputStream);
-						}
-						OutputFileNames.Add(OutputFileName);
-					}
-				}
-				return OutputFileNames;
-			}
-		}
 	}
 }
diff --git a/Engine/Source/Programs/UnrealBuildTool/Platform/IOS/IOSToolChain.cs b/Engine/Source/Programs/UnrealBuildTool/Platform/IOS/IOSToolChain.cs
--- a/Engine/Source/Programs/UnrealBuildTool/Platform/IOS/IOSToolChain.cs
+++ b/Engine/Source/Programs/UnrealBuildTool/Platform/IOS/IOSToolChain.cs
@@ -4,11 +4,10 @@ using System;
 using System.Collections.Generic;
 using System.Diagnostics;
 using System.IO;
+using System.IO.Compression;
 using System.Linq;
 using System.Text;
 using EpicGames.Core;
-using Ionic.Zip;
-using Ionic.Zlib;
 using Microsoft.Extensions.Logging;
 using UnrealBuildBase;
 
@@ -720,95 +719,23 @@ namespace UnrealBuildTool
 				}
 			}
 
-			// create the file
-			using (ZipFile Zip = new ZipFile())
-			{
-				// add the entire directory
-				Zip.AddDirectory(ZipSourceDir, ZipWorkingDir);
-
-				// Update permissions to be UNIX-style
-				// Modify the file attributes of any added file to unix format
-				foreach (ZipEntry E in Zip.Entries)
-				{
-					const byte FileAttributePlatform_NTFS = 0x0A;
-					const byte FileAttributePlatform_UNIX = 0x03;
-					const byte FileAttributePlatform_FAT = 0x00;
-
-					const int UNIX_FILETYPE_NORMAL_FILE = 0x8000;
-					//const int UNIX_FILETYPE_SOCKET = 0xC000;
-					//const int UNIX_FILETYPE_SYMLINK = 0xA000;
-					//const int UNIX_FILETYPE_BLOCKSPECIAL = 0x6000;
-					const int UNIX_FILETYPE_DIRECTORY = 0x4000;
-					//const int UNIX_FILETYPE_CHARSPECIAL = 0x2000;
-					//const int UNIX_FILETYPE_FIFO = 0x1000;
-
-					const int UNIX_EXEC = 1;
-					const int UNIX_WRITE = 2;
-					const int UNIX_READ = 4;
-
-					int MyPermissions = UNIX_READ | UNIX_WRITE;
-					int OtherPermissions = UNIX_READ;
-
-					int PlatformEncodedBy = (E.VersionMadeBy >> 8) & 0xFF;
-					int LowerBits = 0;
-
-					// Try to preserve read-only if it was set
-					bool bIsDirectory = E.IsDirectory;
-
-					// Check to see if this 
-					bool bIsExecutable = false;
-					if (Path.GetFileNameWithoutExtension(E.FileName).Equals(GameName, StringComparison.OrdinalIgnoreCase))
-					{
-						bIsExecutable = true;
-					}
-
-					if (bIsExecutable)
-					{
-						// The executable will be encrypted in the final distribution IPA and will compress very poorly, so keeping it
-						// uncompressed gives a better indicator of IPA size for our distro builds
-						E.CompressionLevel = CompressionLevel.None;
-					}
-
-					if ((PlatformEncodedBy == FileAttributePlatform_NTFS) || (PlatformEncodedBy == FileAttributePlatform_FAT))
-					{
-						FileAttributes OldAttributes = E.Attributes;
-						//LowerBits = ((int)E.Attributes) & 0xFFFF;
-
-						if ((OldAttributes & FileAttributes.Directory) != 0)
-						{
-							bIsDirectory = true;
-						}
+			// Please note this archiving code does not support Apple extended file attributes (xattrs).
+			// Currently, this code is only reachable on Mac platforms - bear in mind that it may not set Unix file attribute bits correctly
+			// if it is made to run from Windows, and thus may produce app packages that are not capable of running without further modification.
+			using ZipArchive Zip = ZipFile.Open(IPAName, ZipArchiveMode.Create);
+			DirectoryReference SourceDir = new(ZipSourceDir);
 
-						// Permissions
-						if ((OldAttributes & FileAttributes.ReadOnly) != 0)
-						{
-							MyPermissions &= ~UNIX_WRITE;
-							OtherPermissions &= ~UNIX_WRITE;
-						}
-					}
-
-					if (bIsDirectory || bIsExecutable)
-					{
-						MyPermissions |= UNIX_EXEC;
-						OtherPermissions |= UNIX_EXEC;
-					}
-
-					// Re-jigger the external file attributes to UNIX style if they're not already that way
-					if (PlatformEncodedBy != FileAttributePlatform_UNIX)
-					{
-						int NewAttributes = bIsDirectory ? UNIX_FILETYPE_DIRECTORY : UNIX_FILETYPE_NORMAL_FILE;
-
-						NewAttributes |= (MyPermissions << 6);
-						NewAttributes |= (OtherPermissions << 3);
-						NewAttributes |= (OtherPermissions << 0);
+			foreach (string FileSystemEntryPath in Directory.EnumerateFiles(ZipSourceDir, "*", SearchOption.AllDirectories))
+			{
+				string RelativeSourcePath = new FileReference(FileSystemEntryPath).MakeRelativeTo(SourceDir);
+				string DestinationPath = Path.Combine(ZipWorkingDir, RelativeSourcePath);
 
-						// Now modify the properties
-						E.AdjustExternalFileAttributes(FileAttributePlatform_UNIX, (NewAttributes << 16) | LowerBits);
-					}
-				}
+				// The executable will be encrypted in the final distribution IPA and will compress very poorly, so keeping it
+				// uncompressed gives a better indicator of IPA size for our distro builds
+				bool bIsExecutable = Path.GetFileNameWithoutExtension(FileSystemEntryPath).Equals(GameName, StringComparison.OrdinalIgnoreCase);
+				CompressionLevel CompressionLevel = bIsExecutable ? CompressionLevel.NoCompression : CompressionLevel.Optimal;
 
-				// Save it out
-				Zip.Save(IPAName);
+				Zip.CreateEntryFromFile(FileSystemEntryPath, DestinationPath, CompressionLevel);
 			}
 		}
 
diff --git a/Engine/Source/Programs/UnrealBuildTool/UnrealBuildTool.csproj b/Engine/Source/Programs/UnrealBuildTool/UnrealBuildTool.csproj
--- a/Engine/Source/Programs/UnrealBuildTool/UnrealBuildTool.csproj
+++ b/Engine/Source/Programs/UnrealBuildTool/UnrealBuildTool.csproj
@@ -11,12 +11,6 @@
     <GenerateDocumentationFile>True</GenerateDocumentationFile>
   </PropertyGroup>
 
-  <ItemGroup>
-    <Reference Include="Ionic.Zip.Reduced">
-      <HintPath>..\..\..\Binaries\DotNET\Ionic.Zip.Reduced.dll</HintPath>
-    </Reference>
-  </ItemGroup>
-
   <PropertyGroup Condition="Exists('..\..\..\Restricted\NotForLicensees\Source\Programs\UnrealBuildTool\ProjectFiles\VProject\VProjectFileGenerator.cs')">
     <DefineConstants>$(DefineConstants);__VPROJECT_AVAILABLE__</DefineConstants>
   </PropertyGroup>
