UnrealBuildTool: Remove unused code from Apple / IOS support.

#Synced-CL 50903846

diff --git a/Engine/Source/Programs/UnrealBuildTool/Platform/Apple/AppleOnWindowsAppMaker.cs b/Engine/Source/Programs/UnrealBuildTool/Platform/Apple/AppleOnWindowsAppMaker.cs
--- a/Engine/Source/Programs/UnrealBuildTool/Platform/Apple/AppleOnWindowsAppMaker.cs
+++ b/Engine/Source/Programs/UnrealBuildTool/Platform/Apple/AppleOnWindowsAppMaker.cs
@@ -2,11 +2,8 @@
 
 using System;
 using System.Collections.Generic;
-using System.Diagnostics;
 using System.IO;
-using System.Security.Cryptography.X509Certificates;
 using System.Text;
-using System.Text.Json;
 using System.Xml.Linq;
 using EpicGames.Core;
 using Microsoft.Extensions.Logging;
@@ -14,616 +11,13 @@ using UnrealBuildBase;
 
 namespace UnrealBuildTool
 {
-	/** Helper class to make Apple apps on Windows.
-	 * It contains most of the functionality that used to exist as "Legacy Xcode"
-	 */
+	/// <summary>
+	/// Helper class to make Apple apps on Windows.
+	/// It contains most of the functionality that used to exist as "Legacy Xcode".
+	/// </summary>
 	class AppleOnWindowsAppMaker
 	{
-		class ProcessOutput
-		{
-			/// <summary>
-			/// Substrings that indicate a line contains an error
-			/// </summary>
-			protected static readonly string[] ErrorMessageTokens =
-			{
-				"ERROR ", "** BUILD FAILED **", "[BEROR]", "IPP ERROR", "System.Net.Sockets.SocketException"
-			};
-
-			/// <summary>
-			/// Helper function to sync source files to and from the local system and a remote Mac
-			/// </summary>
-			//This chunk looks to be required to pipe output to VS giving information on the status of a remote build.
-			public bool OutputReceivedDataEventHandlerEncounteredError = false;
-
-			public string OutputReceivedDataEventHandlerEncounteredErrorMessage = "";
-
-			public void OutputReceivedDataEventHandler(DataReceivedEventArgs Line, ILogger Logger)
-			{
-				if ((Line != null) && (Line.Data != null))
-				{
-					Logger.LogInformation("{Message}", Line.Data);
-
-					foreach (string ErrorToken in ErrorMessageTokens)
-					{
-						if (Line.Data.Contains(ErrorToken, StringComparison.Ordinal))
-						{
-							OutputReceivedDataEventHandlerEncounteredError = true;
-							OutputReceivedDataEventHandlerEncounteredErrorMessage += Line.Data;
-							break;
-						}
-					}
-				}
-			}
-
-			public static void OutputReceivedDataEventLogger(DataReceivedEventArgs Line, ILogger Logger)
-			{
-				if ((Line != null) && (Line.Data != null))
-				{
-					Logger.LogInformation("{Output}", Line.Data);
-				}
-			}
-		}
-
-		private static void WriteEntitlements(IOSPostBuildSyncTarget Target)
-		{
-			string AppName = Target.TargetName;
-			FileReference? MobileProvisionFile;
-			IOSProjectSettings ProjectSettings =
-				((IOSPlatform)UEBuildPlatform.GetBuildPlatform(Target.Platform)).ReadProjectSettings(Target.ProjectFile);
-
-			if (Target.ImportProvision == null)
-			{
-				IOSProvisioningData ProvisioningData =
-					((IOSPlatform)UEBuildPlatform.GetBuildPlatform(Target.Platform)).ReadProvisioningData(ProjectSettings, Target.bForDistribution);
-				MobileProvisionFile = ProvisioningData.MobileProvisionFile;
-			}
-			else
-			{
-				MobileProvisionFile = new FileReference(Target.ImportProvision);
-			}
-
-			string IntermediateDir = (((Target.ProjectFile != null) ? Target.ProjectFile.Directory.ToString() : Unreal.EngineDirectory.ToString())) +
-			                         "/Intermediate/" + (Target.Platform == UnrealTargetPlatform.IOS ? "IOS" : "TVOS");
-			// get the settings from the ini file
-			ConfigHierarchy Ini = ConfigCache.ReadHierarchy(ConfigHierarchyType.Engine, GetActualProjectDirectory(Target.ProjectFile),
-				UnrealTargetPlatform.IOS);
-
-			IOSExports.WriteEntitlements(Target.Platform, Ini, AppName, MobileProvisionFile, Target.bForDistribution, IntermediateDir);
-		}
-
-		private static void WriteMissingEmbeddedPlist(FileReference MissingPlistPath, DirectoryReference ProjectDirectory, FileReference? ProjectFile)
-		{
-			FileReference PlistTemplatePath =
-				FileReference.Combine(Unreal.EngineDirectory, "Build", "IOS", "Resources", "FrameworkPlist", "Embedded.plist");
-			if (!DirectoryReference.Exists(MissingPlistPath.Directory))
-			{
-				DirectoryReference.CreateDirectory(MissingPlistPath.Directory);
-			}
-
-			string BundleID = ProjectFiles.Xcode.XcodeFrameworkWrapperUtils.GetBundleID(ProjectDirectory, ProjectFile);
-			string BundleName = ProjectFiles.Xcode.XcodeFrameworkWrapperUtils.GetBundleName(ProjectDirectory, ProjectFile);
-
-			BundleID += ".embedded"; // This is necessary because the wrapper that contains this framework cannot have the same BundleID.
-
-			string Plist = FileReference.ReadAllText(PlistTemplatePath);
-			Plist = Plist.Replace("BUNDLE_NAME", BundleName, StringComparison.Ordinal);
-			Plist = Plist.Replace("BUNDLE_ID", BundleID, StringComparison.Ordinal);
-
-			FileReference.WriteAllText(MissingPlistPath, Plist);
-		}
-
-		/// <summary>
-		/// If the project is a UnrealGame project, Target.ProjectDirectory refers to the engine dir, not the actual dir of the project. So this method gets the 
-		/// actual directory of the project whether it is a UnrealGame project or not.
-		/// </summary>
-		/// <returns>The actual project directory.</returns>
-		/// <param name="ProjectFile">The path to the project file</param>
-		internal static DirectoryReference GetActualProjectDirectory(FileReference? ProjectFile)
-		{
-			DirectoryReference ProjectDirectory = (ProjectFile == null ? Unreal.EngineDirectory : DirectoryReference.FromFile(ProjectFile)!);
-			return ProjectDirectory;
-		}
-
-		private static void GenerateFrameworkWrapperIfNonexistent(IOSPostBuildSyncTarget Target)
-		{
-			DirectoryReference ProjectDirectory = GetActualProjectDirectory(Target.ProjectFile);
-			DirectoryReference WrapperDirectory = DirectoryReference.Combine(ProjectDirectory, "Wrapper");
-
-			string ProjectName = ProjectDirectory.GetDirectoryName();
-			string FrameworkName = Target.TargetName;
-			string BundleId = ProjectFiles.Xcode.XcodeFrameworkWrapperUtils.GetBundleID(Target.ProjectDirectory, Target.ProjectFile);
-			string EnginePath = Unreal.EngineDirectory.ToString();
-			string SrcFrameworkPath = DirectoryReference.Combine(Target.ProjectDirectory, "Binaries", "IOS", Target.Configuration.ToString())
-				.ToString(); // We use Target.ProjectDirectory because if it is a UnrealGame we want the engine dir and not the actual project dir.
-			string CookedDataPath = DirectoryReference.Combine(ProjectDirectory, "Saved", "StagedBuilds", "IOS", "cookeddata").ToString();
-			bool GenerateFrameworkWrapperProject =
-				ProjectFiles.Xcode.XcodeFrameworkWrapperUtils.GetGenerateFrameworkWrapperProject(Target.ProjectDirectory);
-
-			string? ProvisionName = String.Empty;
-			string? TeamUUID = String.Empty;
-			IOSProjectSettings ProjectSettings =
-				((IOSPlatform)UEBuildPlatform.GetBuildPlatform(UnrealTargetPlatform.IOS)).ReadProjectSettings(Target.ProjectFile);
-
-			IOSProvisioningData Data =
-				((IOSPlatform)UEBuildPlatform.GetBuildPlatform(UnrealTargetPlatform.IOS)).ReadProvisioningData(ProjectSettings,
-					Target.bForDistribution);
-			if (Data != null)
-			{
-				ProvisionName = Data.MobileProvisionName;
-				TeamUUID = Data.TeamUUID;
-			}
-
-			// Create the framework wrapper project, only if it doesn't exist and only if the user wants to generate the wrapper.
-			if (!DirectoryReference.Exists(DirectoryReference.Combine(WrapperDirectory, FrameworkName)) && GenerateFrameworkWrapperProject)
-			{
-				ProjectFiles.Xcode.XcodeFrameworkWrapperProject.GenerateXcodeFrameworkWrapper(WrapperDirectory.ToString(), ProjectName, FrameworkName,
-					BundleId, SrcFrameworkPath, EnginePath, CookedDataPath, ProvisionName, TeamUUID);
-			}
-		}
-
-		private static void GenerateCrashlyticsData(string DsymZip, ILogger Logger)
-		{
-			Logger.LogInformation("Generating and uploading Crashlytics Data");
-
-			// Clean this folder as it's used for extraction
-			string TempPath = Path.Combine(Unreal.EngineDirectory.FullName, "Intermediate", "Unzipped");
-
-			if (Directory.Exists(TempPath))
-			{
-				Logger.LogInformation("Deleting temp path {TempPath}", TempPath);
-				Directory.Delete(TempPath, true);
-			}
-
-			string FabricPath = Unreal.EngineDirectory + "/Intermediate/UnzippedFrameworks/Crashlytics/Fabric.embeddedframework";
-			if (Directory.Exists(FabricPath) && UnrealBuildBase.Unreal.IsBuildMachine())
-			{
-				//string PlistFile = ProjectDir + "/Intermediate/IOS/" + ProjectName + "-Info.plist";
-				Process FabricProcess = new Process();
-				FabricProcess.StartInfo.WorkingDirectory = Path.GetDirectoryName(DsymZip);
-				FabricProcess.StartInfo.FileName = "/bin/sh";
-				FabricProcess.StartInfo.Arguments = String.Format(
-					"-c \"chmod 777 \\\"{0}/Fabric.framework/upload-symbols\\\"; \\\"{0}/Fabric.framework/upload-symbols\\\" -a 7a4cebd0324af21696e5e321802c5e26ba541cad -p ios {1}\"",
-					FabricPath,
-					DsymZip);
-
-				ProcessOutput Output = new ProcessOutput();
-
-				FabricProcess.OutputDataReceived += new DataReceivedEventHandler((_, e) => Output.OutputReceivedDataEventHandler(e, Logger));
-				FabricProcess.ErrorDataReceived += new DataReceivedEventHandler((_, e) => Output.OutputReceivedDataEventHandler(e, Logger));
-
-				Utils.RunLocalProcess(FabricProcess);
-				if (Output.OutputReceivedDataEventHandlerEncounteredError)
-				{
-					throw new Exception(Output.OutputReceivedDataEventHandlerEncounteredErrorMessage);
-				}
-			}
-		}
-
-		static void CleanIntermediateDirectory(string Path)
-		{
-			// Delete the local dest directory if it exists
-			if (Directory.Exists(Path))
-			{
-				// this can deal with linked files
-				IOSToolChain.RunExecutableAndWait("rm", String.Format("-rf \"{0}\"", Path), out _);
-			}
-
-			// Create the intermediate local directory
-			IOSToolChain.RunExecutableAndWait("mkdir", String.Format("-p \"{0}\"", Path), out _);
-		}
-
-		public static void PostBuildSync(IOSPostBuildSyncTarget Target, ILogger Logger)
-		{
-			if (!AppleExports.CreatingAppOnWindows(Target.ProjectFile))
-			{
-				throw new BuildException("Xcode mode should not reach this, something has broken the PostBuild setup");
-			}
-
-			IOSProjectSettings ProjectSettings =
-				((IOSPlatform)UEBuildPlatform.GetBuildPlatform(Target.Platform)).ReadProjectSettings(Target.ProjectFile);
-			string? BundleID = ProjectSettings.BundleIdentifier;
-
-			bool bPerformFullAppCreation = true;
-			string PathToDsymZip = Target.OutputPath.FullName + ".dSYM.zip";
-			if (BuildHostPlatform.Current.Platform == UnrealTargetPlatform.Mac)
-			{
-				if (Target.bBuildAsFramework)
-				{
-					// make sure the framework has a plist
-					FileReference PlistSrcLocation = FileReference.Combine(Target.ProjectDirectory, "Build", "IOS", "Embedded.plist");
-					if (!FileReference.Exists(PlistSrcLocation))
-					{
-
-						WriteMissingEmbeddedPlist(PlistSrcLocation, Target.ProjectDirectory, Target.ProjectFile);
-						if (!FileReference.Exists(PlistSrcLocation))
-						{
-							throw new BuildException("Unable to find plist for output framework ({0})", PlistSrcLocation);
-						}
-					}
-
-					FileReference PlistDstLocation = FileReference.Combine(Target.OutputPath.Directory, "Info.plist");
-					if (FileReference.Exists(PlistDstLocation))
-					{
-						FileReference.SetAttributes(PlistDstLocation, FileAttributes.Normal);
-					}
-
-					FileReference.Copy(PlistSrcLocation, PlistDstLocation, true);
-
-					PathToDsymZip = Path.Combine(Target.OutputPath.Directory.ParentDirectory!.FullName,
-						Target.OutputPath.GetFileName() + ".dSYM.zip");
-
-					GenerateFrameworkWrapperIfNonexistent(Target);
-
-					// do not perform any of the .app creation below
-					bPerformFullAppCreation = false;
-				}
-			}
-
-			string AppName = Target.TargetName;
-
-			if (!Target.bSkipCrashlytics)
-			{
-				GenerateCrashlyticsData(PathToDsymZip, Logger);
-			}
-
-			// only make the app if needed
-			if (!bPerformFullAppCreation)
-			{
-				return;
-			}
-
-			// copy the executable
-			FileReference StagedExecutablePath = IOSToolChain.GetStagedExecutablePath(Target.OutputPath, Target.TargetName);
-			DirectoryReference.CreateDirectory(StagedExecutablePath.Directory);
-			FileReference.Copy(Target.OutputPath, StagedExecutablePath, true);
-			string RemoteShadowDirectoryMac = Target.OutputPath.Directory.FullName;
-
-			if (Target.bCreateStubIPA || Target.bBuildAsFramework)
-			{
-				string FrameworkDerivedDataDir = Path.GetFullPath(Target.ProjectDirectory + "/Intermediate/" +
-				                                                  (Target.Platform == UnrealTargetPlatform.IOS ? "IOS" : "TVOS") + "/DerivedData");
-				DirectoryReference FrameworkPayloadDirectory =
-					DirectoryReference.Combine(Target.OutputPath.Directory.ParentDirectory!.ParentDirectory!, "Payload");
-
-				// generate the dummy project so signing works
-				DirectoryReference? XcodeWorkspaceDir = null;
-				if (!Target.bBuildAsFramework)
-				{
-					if (AppName == "UnrealGame" || AppName == "UnrealClient" || Target.ProjectFile == null ||
-					    Target.ProjectFile.IsUnderDirectory(Unreal.EngineDirectory))
-					{
-						XcodeWorkspaceDir = DirectoryReference.Combine(Unreal.EngineDirectory, "Intermediate/ProjectFiles",
-							String.Format("UE5_{0}.xcworkspace", (Target.Platform == UnrealTargetPlatform.IOS ? "IOS" : "TVOS")));
-					}
-					else
-					{
-						XcodeWorkspaceDir = DirectoryReference.Combine(Target.ProjectDirectory, "Intermediate/ProjectFiles",
-							String.Format("{0}_{1}.xcworkspace", Target.ProjectFile.GetFileNameWithoutExtension(),
-								(Target.Platform == UnrealTargetPlatform.IOS ? "IOS" : "TVOS")));
-					}
-				}
-
-				// Path to the temporary keychain. When -ImportCertificate is specified, we will temporarily add this to the list of keychains to search, and remove it later.
-				FileReference TempKeychain = FileReference.Combine(Target.ProjectIntermediateDirectory!, "TempKeychain.keychain");
-
-				FileReference SignProjectScript = FileReference.Combine(Target.ProjectIntermediateDirectory!, "SignProject.sh");
-				using (StreamWriter Writer = new StreamWriter(SignProjectScript.FullName))
-				{
-					// Boilerplate
-					Writer.WriteLine("#!/bin/sh");
-					Writer.WriteLine("set -e");
-					Writer.WriteLine("set -x");
-
-					// Copy the mobile provision into the system store
-					if (Target.ImportProvision != null)
-					{
-						Writer.WriteLine("mkdir -p {0}", Utils.EscapeShellArgument(AppleExports.GetProvisionDirectory().FullName));
-						Writer.WriteLine("cp -f {0} {1}", Utils.EscapeShellArgument(Target.ImportProvision),
-							Utils.EscapeShellArgument(AppleExports.GetProvisionDirectory().FullName));
-					}
-
-					// Get the signing certificate to use
-					string? SigningCertificate;
-					if (Target.ImportCertificate == null)
-					{
-						// Take it from the standard settings
-						IOSProvisioningData ProvisioningData =
-							((IOSPlatform)UEBuildPlatform.GetBuildPlatform(Target.Platform)).ReadProvisioningData(Target.ProjectFile,
-								Target.bForDistribution);
-						SigningCertificate = ProvisioningData.SigningCertificate;
-
-						// Set the identity on the command line
-						if (!ProjectSettings.bAutomaticSigning)
-						{
-							Writer.WriteLine("CODE_SIGN_IDENTITY='{0}'",
-								String.IsNullOrEmpty(SigningCertificate) ? "IPhoneDeveloper" : SigningCertificate);
-						}
-					}
-					else
-					{
-						// Read the name from the certificate
-						X509Certificate2 Certificate;
-						try
-						{
-							Certificate = new X509Certificate2(Target.ImportCertificate, Target.ImportCertificatePassword ?? "");
-						}
-						catch (Exception Ex)
-						{
-							throw new BuildException(Ex, "Unable to read certificate '{0}': {1}", Target.ImportCertificate, Ex.Message);
-						}
-
-						SigningCertificate = Certificate.GetNameInfo(X509NameType.SimpleName, false);
-
-						// Install a certificate given on the command line to a temporary keychain
-						Writer.WriteLine("security delete-keychain \"{0}\" || true", TempKeychain);
-						Writer.WriteLine("security create-keychain -p \"A\" \"{0}\"", TempKeychain);
-						Writer.WriteLine("security list-keychains -s \"{0}\"", TempKeychain);
-						Writer.WriteLine("security list-keychains");
-						Writer.WriteLine("security set-keychain-settings -t 3600 -l  \"{0}\"", TempKeychain);
-						Writer.WriteLine("security -v unlock-keychain -p \"A\" \"{0}\"", TempKeychain);
-						Writer.WriteLine("security import {0} -P {1} -k \"{2}\" -T /usr/bin/codesign -T /usr/bin/security -t agg",
-							Utils.EscapeShellArgument(Target.ImportCertificate), Utils.EscapeShellArgument(Target.ImportCertificatePassword!),
-							TempKeychain);
-						Writer.WriteLine("security set-key-partition-list -S apple-tool:,apple:,codesign: -s -k \"A\" -D '{0}' -t private {1}",
-							SigningCertificate, TempKeychain);
-
-						// Set parameters to make sure it uses the correct identity and keychain
-						Writer.WriteLine("CERT_IDENTITY='{0}'", SigningCertificate);
-						Writer.WriteLine("CODE_SIGN_IDENTITY='{0}'", SigningCertificate);
-						Writer.WriteLine("CODE_SIGN_KEYCHAIN='{0}'", TempKeychain);
-					}
-
-					FileReference? MobileProvisionFile;
-					string? MobileProvisionUUID;
-					string? TeamUUID;
-					if (Target.ImportProvision == null)
-					{
-						IOSProvisioningData ProvisioningData =
-							((IOSPlatform)UEBuildPlatform.GetBuildPlatform(Target.Platform)).ReadProvisioningData(ProjectSettings,
-								Target.bForDistribution);
-						MobileProvisionFile = ProvisioningData.MobileProvisionFile;
-						MobileProvisionUUID = ProvisioningData.MobileProvisionUUID;
-						TeamUUID = ProvisioningData.TeamUUID;
-						if (ProvisioningData.BundleIdentifier == null || !ProvisioningData.BundleIdentifier.Contains('*', StringComparison.Ordinal))
-						{
-							// If the BundleIndentifer contains a wild card it will not be valid to use in the plist.
-							BundleID = ProvisioningData.BundleIdentifier;
-						}
-					}
-					else
-					{
-						MobileProvisionFile = new FileReference(Target.ImportProvision);
-
-						MobileProvisionContents MobileProvision = MobileProvisionContents.Read(MobileProvisionFile);
-						MobileProvisionUUID = MobileProvision.GetUniqueId();
-						MobileProvision.TryGetTeamUniqueId(out TeamUUID);
-						string BundleIdentifier = MobileProvision.GetBundleIdentifier();
-						if (!BundleIdentifier.Contains('*', StringComparison.Ordinal))
-						{
-							// If the BundleIndentifer contains a wild card it will not be valid to use in the plist.
-							BundleID = BundleIdentifier;
-						}
-					}
-
-					if (MobileProvisionFile == null)
-					{
-						throw new BuildException("Unable to find valid certificate/mobile provision pair.");
-					}
-
-					string ConfigName = Target.Configuration.ToString();
-					if (Target.TargetType != TargetType.Game && Target.TargetType != TargetType.Program)
-					{
-						ConfigName += " " + Target.TargetType.ToString();
-					}
-
-					string SchemeName = Target.TargetName;
-
-					Console.WriteLine("Provisioning: {0}, {1}, {2}, {3}", MobileProvisionFile, MobileProvisionFile.GetFileName(), MobileProvisionUUID,
-						BundleID);
-
-					string CmdLine;
-					if (Target.bBuildAsFramework)
-					{
-						DirectoryReference FrameworkProjectDirectory = GetActualProjectDirectory(Target.ProjectFile);
-						DirectoryReference WrapperDirectory = DirectoryReference.Combine(FrameworkProjectDirectory, "Wrapper");
-						string FrameworkName = Target.TargetName;
-						DirectoryReference WrapperProject = DirectoryReference.Combine(WrapperDirectory, FrameworkName, FrameworkName + ".xcodeproj");
-
-						// Delete everything in payload directory before building
-						Writer.WriteLine(String.Format("rm -rf \"{0}\"", FrameworkPayloadDirectory));
-
-						// Build the framework wrapper
-						CmdLine = FileReference.Combine(ApplePlatformSDK.DeveloperDir, "usr/bin/xcodebuild").FullName +
-						          " -project \"" + WrapperProject + "\"" +
-						          " -configuration \"" + ConfigName + "\"" +
-						          " -scheme '" + SchemeName + "'" +
-						          " -sdk " + IOSToolChain.GetCodesignPlatformName(Target.Platform) +
-						          $" -destination generic/platform=\"{AppleExports.GetDestinationPlatform(Target.Platform, Target.Architectures)}\"" +
-						          " -derivedDataPath \"" + FrameworkDerivedDataDir + "\"" +
-						          " CONFIGURATION_BUILD_DIR=\"" + FrameworkPayloadDirectory + "\"" +
-						          (!String.IsNullOrEmpty(TeamUUID) ? " DEVELOPMENT_TEAM=" + TeamUUID : "");
-					}
-					else
-					{
-						// code sign the project
-						CmdLine = FileReference.Combine(ApplePlatformSDK.DeveloperDir, "usr/bin/xcodebuild").FullName +
-						          " -workspace \"" + XcodeWorkspaceDir + "\"" +
-						          " -configuration \"" + ConfigName + "\"" +
-						          " -scheme '" + SchemeName + "'" +
-						          " -sdk " + IOSToolChain.GetCodesignPlatformName(Target.Platform) +
-						          $" -destination generic/platform=\"{AppleExports.GetDestinationPlatform(Target.Platform, Target.Architectures)}\"" +
-						          (!String.IsNullOrEmpty(TeamUUID) ? " DEVELOPMENT_TEAM=" + TeamUUID : "");
-					}
-
-					CmdLine += String.Format(" CODE_SIGN_IDENTITY='{0}'", SigningCertificate);
-
-					// <projectname>.entitlements file can get copied over during a RemoteMac build which confuses Xcode, so just tell it to ignore any 
-					// entitlement file changes.  See UE-181206.
-					CmdLine += String.Format(" CODE_SIGN_ALLOW_ENTITLEMENTS_MODIFICATION=YES");
-
-					if (!ProjectSettings.bAutomaticSigning)
-					{
-						CmdLine += (!String.IsNullOrEmpty(MobileProvisionUUID) ? (" PROVISIONING_PROFILE_SPECIFIER=" + MobileProvisionUUID) : "");
-					}
-
-					Writer.WriteLine("/usr/bin/xcrun {0}", CmdLine);
-				}
-
-				if (!Target.bBuildAsFramework)
-				{
-					Logger.LogInformation("Generating project...");
-					DirectoryReference? GeneratedProjectFile;
-					if (AppName == "UnrealGame" || AppName == "UnrealClient" || Target.ProjectFile == null ||
-					    Target.ProjectFile.IsUnderDirectory(Unreal.EngineDirectory))
-					{
-						AppleToolChain.GenerateProjectFiles(Target.ProjectFile,
-							new string[]
-							{
-								"-platforms=" + (Target.Platform == UnrealTargetPlatform.IOS ? "IOS" : "TVOS"), "-NoIntellIsense",
-								(Target.Platform == UnrealTargetPlatform.IOS ? "-iosdeployonly" : "-tvosdeployonly"), "-ignorejunk",
-								(Target.bForDistribution ? "-distribution" : "-development"), "-bundleID=" + BundleID, "-includetemptargets",
-								"-appname=" + AppName
-							}, null, Logger, out GeneratedProjectFile);
-					}
-					else
-					{
-						AppleToolChain.GenerateProjectFiles(Target.ProjectFile,
-							new string[]
-							{
-								"-platforms=" + (Target.Platform == UnrealTargetPlatform.IOS ? "IOS" : "TVOS"), "-NoIntellIsense",
-								(Target.Platform == UnrealTargetPlatform.IOS ? "-iosdeployonly" : "-tvosdeployonly"), "-ignorejunk",
-								(Target.bForDistribution ? "-distribution" : "-development"), String.Format("-project={0}", Target.ProjectFile),
-								"-game", "-bundleID=" + BundleID, "-includetemptargets"
-							}, null, Logger, out GeneratedProjectFile);
-					}
-
-					Logger.LogInformation("Generated {ProjectDirectory}!", GeneratedProjectFile?.FullName ?? "NONE");
-					// Make sure it exists
-					if (!DirectoryReference.Exists(XcodeWorkspaceDir!) || (GeneratedProjectFile != null && GeneratedProjectFile != XcodeWorkspaceDir))
-					{
-						throw new BuildException("Unable to create stub IPA; Xcode workspace not found at {0}", XcodeWorkspaceDir);
-					}
-				}
-
-				// ensure the plist, entitlements, and provision files are properly copied
-				UEDeployIOS DeployHandler = (Target.Platform == UnrealTargetPlatform.IOS ? new UEDeployIOS(Logger) : new UEDeployTVOS(Logger));
-				DeployHandler.ForDistribution = Target.bForDistribution;
-				DeployHandler.PrepTargetForDeployment(Target.ProjectFile, Target.TargetName, Target.OutputPath, Target.Configuration,
-					Target.UPLScripts,
-					Target.bCreateStubIPA, BundleID, Target.bBuildAsFramework);
-
-				Logger.LogInformation("Executing {Script}", SignProjectScript);
-
-				// write the entitlements file (building remotely)
-				WriteEntitlements(Target);
-
-				Process SignProcess = new Process();
-				SignProcess.StartInfo.WorkingDirectory = RemoteShadowDirectoryMac;
-				SignProcess.StartInfo.FileName = "/bin/sh";
-				string Arguments = String.Format("\"{0}\"", SignProjectScript.FullName);
-				SignProcess.StartInfo.Arguments = Arguments;
-
-				ProcessOutput Output = new ProcessOutput();
-
-				SignProcess.OutputDataReceived += new DataReceivedEventHandler((_, e) => Output.OutputReceivedDataEventHandler(e, Logger));
-				SignProcess.ErrorDataReceived += new DataReceivedEventHandler((_, e) => Output.OutputReceivedDataEventHandler(e, Logger));
-
-				Output.OutputReceivedDataEventHandlerEncounteredError = false;
-				Output.OutputReceivedDataEventHandlerEncounteredErrorMessage = "";
-				Utils.RunLocalProcess(SignProcess);
-
-				// cleanup
-				if (Target.ImportCertificate != null)
-				{
-					FileReference CleanProjectScript = FileReference.Combine(Target.ProjectIntermediateDirectory!, "CleanProject.sh");
-					using (StreamWriter CleanWriter = new StreamWriter(CleanProjectScript.FullName))
-					{
-						// Remove the temporary keychain from the search list
-						CleanWriter.WriteLine("security delete-keychain \"{0}\" || true", TempKeychain);
-						CleanWriter.WriteLine("security list-keychain -s login.keychain");
-					}
-
-					Logger.LogInformation("Executing {Script}", CleanProjectScript);
-
-					Process CleanProcess = new Process();
-					CleanProcess.StartInfo.WorkingDirectory = RemoteShadowDirectoryMac;
-					CleanProcess.StartInfo.FileName = "/bin/sh";
-					CleanProcess.StartInfo.Arguments = CleanProjectScript.FullName;
-
-					ProcessOutput CleanOutput = new ProcessOutput();
-
-					SignProcess.OutputDataReceived += new DataReceivedEventHandler((_, e) => ProcessOutput.OutputReceivedDataEventLogger(e, Logger));
-					SignProcess.ErrorDataReceived += new DataReceivedEventHandler((_, e) => ProcessOutput.OutputReceivedDataEventLogger(e, Logger));
-
-					Utils.RunLocalProcess(CleanProcess);
-				}
-
-				if (Target.bBuildAsFramework)
-				{
-					CleanIntermediateDirectory(FrameworkDerivedDataDir);
-				}
-
-				if (XcodeWorkspaceDir != null)
-				{
-					// delete the temp project
-					DirectoryReference.Delete(XcodeWorkspaceDir, true);
-
-					if (Output.OutputReceivedDataEventHandlerEncounteredError)
-					{
-						throw new Exception(Output.OutputReceivedDataEventHandlerEncounteredErrorMessage);
-					}
-				}
-
-				const bool bIsUsingCustomPackager = true;
-				// Package the stub
-				if (Target.bBuildAsFramework)
-				{
-					IOSToolChain.PackageStub(FrameworkPayloadDirectory.ParentDirectory!.ToString(), AppName,
-						Target.OutputPath.GetFileNameWithoutExtension(), bIsUsingCustomPackager);
-				}
-				else
-				{
-					IOSToolChain.PackageStub(RemoteShadowDirectoryMac, AppName, Target.OutputPath.GetFileNameWithoutExtension(),
-						bIsUsingCustomPackager);
-				}
-			}
-			else
-			{
-				// ensure the plist, entitlements, and provision files are properly copied
-				UEDeployIOS DeployHandler = (Target.Platform == UnrealTargetPlatform.IOS ? new UEDeployIOS(Logger) : new UEDeployTVOS(Logger));
-				DeployHandler.ForDistribution = Target.bForDistribution;
-				DeployHandler.PrepTargetForDeployment(Target.ProjectFile, Target.TargetName, Target.OutputPath, Target.Configuration,
-					Target.UPLScripts,
-					Target.bCreateStubIPA, BundleID, Target.bBuildAsFramework);
-
-				// write the entitlements file (building on Mac)
-				WriteEntitlements(Target);
-			}
-
-			{
-				// Copy bundled assets from additional frameworks to the intermediate assets directory (so they can get picked up during staging)
-				string LocalFrameworkAssets = Path.GetFullPath(Target.ProjectDirectory + "/Intermediate/" +
-				                                               (Target.Platform == UnrealTargetPlatform.IOS ? "IOS" : "TVOS") + "/FrameworkAssets");
-
-				// Clean the local dest directory if it exists
-				CleanIntermediateDirectory(LocalFrameworkAssets);
-
-				foreach (KeyValuePair<string, DirectoryReference> Pair in Target.FrameworkNameToSourceDir)
-				{
-					//string UnpackedZipPath = Pair.Value.FullName;
-
-					// For now, this is hard coded, but we need to loop over all modules, and copy bundled assets that need it
-					string LocalDest = LocalFrameworkAssets + "/" + Pair.Key;
-
-					Logger.LogInformation("Copying bundled asset... LocalSource: {Source}, LocalDest: {Target}", Pair.Value, LocalDest);
-
-					string ResultsText;
-					IOSToolChain.RunExecutableAndWait("cp", String.Format("-R -L \"{0}\" \"{1}\"", Pair.Value, LocalDest), out ResultsText);
-				}
-			}
-		}
-
-		public  static bool GenerateLegacyIOSPList(FileReference? ProjectFile, string ProjectDirectory, bool bIsUnrealGame, string GameName,
+		public static bool GenerateLegacyIOSPList(FileReference? ProjectFile, string ProjectDirectory, bool bIsUnrealGame, string GameName,
 			bool bIsClient, string ProjectName, string InEngineDir, string AppDirectory, UnrealPluginLanguage? UPL, string? BundleID,
 			bool bBuildAsFramework, ILogger Logger)
 		{
@@ -1022,7 +416,7 @@ namespace UnrealBuildTool
 			return true;
 		}
 
-		public static void SafeFileCopy(FileInfo SourceFile, string DestinationPath, bool bOverwrite)
+		private static void SafeFileCopy(FileInfo SourceFile, string DestinationPath, bool bOverwrite)
 		{
 			FileInfo DI = new FileInfo(DestinationPath);
 			if (DI.Exists && bOverwrite)
@@ -1497,156 +891,5 @@ namespace UnrealBuildTool
 
 			return bSkipDefaultPNGs;
 		}
-		
-		public static void GatherBuildProducts(ReadOnlyTargetRules Target, UEBuildBinary Binary, Dictionary<FileReference, BuildProductType> BuildProducts)
-		{
-			if (Target.Platform == UnrealTargetPlatform.TVOS)
-			{
-				FileReference AssetFile = FileReference.Combine(Binary.OutputFilePath.Directory, "AssetCatalog", "Assets.car");
-				BuildProducts.Add(AssetFile, BuildProductType.RequiredResource);
-			}
-			else if (Target.Platform == UnrealTargetPlatform.IOS)
-			{
-				int Index = Binary.OutputFilePath.GetFileNameWithoutExtension().IndexOf("-", StringComparison.Ordinal);
-				string OutputFile = Binary.OutputFilePath.GetFileNameWithoutExtension().Substring(0, Index > 0 ? Index : Binary.OutputFilePath.GetFileNameWithoutExtension().Length);
-				FileReference AssetFile = FileReference.Combine(Binary.OutputFilePath.Directory, "Payload", OutputFile + ".app", "Assets.car");
-				BuildProducts.Add(AssetFile, BuildProductType.RequiredResource);
-
-				// Default AppIcons need to be copied as is in the Payload's root, even when using an asset catalog.
-				AssetFile = FileReference.Combine(Binary.OutputFilePath.Directory, "Payload", OutputFile + ".app", "[REDACTED_EMAIL]");
-				BuildProducts.Add(AssetFile, BuildProductType.RequiredResource);
-				AssetFile = FileReference.Combine(Binary.OutputFilePath.Directory, "Payload", OutputFile + ".app", "AppIcon76x76@2x~ipad.png");
-				BuildProducts.Add(AssetFile, BuildProductType.RequiredResource);
-			}
-		}
-		
-		private static FileItem CopyBundleResource(UEBuildBundleResource Resource, FileItem Executable, DirectoryReference BundleDirectory, IActionGraphBuilder Graph)
-		{
-			Action CopyAction = Graph.CreateAction(ActionType.CreateAppBundle);
-			CopyAction.WorkingDirectory = Unreal.EngineSourceDirectory;
-			CopyAction.CommandPath = BuildHostPlatform.Current.Shell;
-			CopyAction.CommandDescription = "";
-
-			string BundlePath = BundleDirectory.FullName;
-			string SourcePath = Path.Combine(Path.GetFullPath("."), Resource.ResourcePath!);
-			string TargetPath;
-			if (Resource.BundleContentsSubdir == "Resources")
-			{
-				TargetPath = Path.Combine(BundlePath, Path.GetFileName(Resource.ResourcePath)!);
-			}
-			else
-			{
-				TargetPath = Path.Combine(BundlePath, Resource.BundleContentsSubdir!, Path.GetFileName(Resource.ResourcePath)!);
-			}
-
-			FileItem TargetItem = FileItem.GetItemByPath(TargetPath);
-
-			CopyAction.CommandArguments = String.Format("-c \"cp -f -R \\\"{0}\\\" \\\"{1}\\\"; touch -c \\\"{2}\\\"\"", SourcePath, Path.GetDirectoryName(TargetPath)!.Replace('\\', '/') + "/", TargetPath.Replace('\\', '/'));
-			CopyAction.PrerequisiteItems.Add(Executable);
-			CopyAction.ProducedItems.Add(TargetItem);
-			CopyAction.bShouldOutputStatusDescription = Resource.bShouldLog;
-			CopyAction.StatusDescription = String.Format("Copying {0} to app bundle", Path.GetFileName(Resource.ResourcePath));
-			CopyAction.bCanExecuteRemotely = false;
-
-			return TargetItem;
-		}
-		
-		public static void PostBuild(FileReference? ProjectFile, AppleToolChainSettings Settings, ILogger Logger, ReadOnlyTargetRules Target, FileItem Executable, LinkEnvironment BinaryLinkEnvironment, IActionGraphBuilder Graph, ref List<FileItem> OutputFiles)
-		{
-			if (!BinaryLinkEnvironment.bIsBuildingDLL)
-			{
-				// generate the asset catalog
-				bool bUserImagesExist = false;
-				DirectoryReference ResourcesDir = IOSToolChain.GenerateAssetCatalog(ProjectFile, BinaryLinkEnvironment.Platform, ref bUserImagesExist);
-
-				// Get the output location for the asset catalog
-				FileItem AssetCatalogFile = FileItem.GetItemByFileReference(IOSToolChain.GetAssetCatalogFile(BinaryLinkEnvironment.Platform, Executable.Location));
-
-				// Make the compile action
-				Action CompileAssetAction = Graph.CreateAction(ActionType.CreateAppBundle);
-				CompileAssetAction.WorkingDirectory = Unreal.EngineSourceDirectory;
-				CompileAssetAction.CommandPath = new FileReference("/usr/bin/xcrun");
-				CompileAssetAction.CommandArguments = IOSToolChain.GetAssetCatalogArgs(BinaryLinkEnvironment.Platform, ResourcesDir.FullName, Path.GetDirectoryName(AssetCatalogFile.AbsolutePath)!);
-				CompileAssetAction.PrerequisiteItems.Add(Executable);
-				CompileAssetAction.ProducedItems.Add(AssetCatalogFile);
-				CompileAssetAction.DeleteItems.Add(AssetCatalogFile);
-				CompileAssetAction.StatusDescription = CompileAssetAction.CommandArguments;// string.Format("Generating debug info for {0}", Path.GetFileName(Executable.AbsolutePath));
-				CompileAssetAction.bCanExecuteRemotely = false;
-
-				// Add it to the output files so it's always built
-				OutputFiles.Add(AssetCatalogFile);
-			}
-
-			// Generate the app bundle
-			if (!Target.bDisableLinking)
-			{
-				Logger.LogInformation("Adding PostBuildSync action");
-
-				List<string> UPLScripts = UEDeployIOS.CollectPluginDataPaths(BinaryLinkEnvironment.AdditionalProperties, Logger);
-				VersionNumber SdkVersion = VersionNumber.Parse(Settings.SDKVersion);
-
-				Dictionary<string, DirectoryReference> FrameworkNameToSourceDir = new Dictionary<string, DirectoryReference>();
-
-				FileReference StagedExecutablePath = IOSToolChain.GetStagedExecutablePath(Executable.Location, Target.Name);
-				DirectoryReference BundleDirectory = Target.bShouldCompileAsDLL ? Executable.Directory.Location : StagedExecutablePath.Directory;
-				foreach (UEBuildFramework Framework in BinaryLinkEnvironment.AdditionalFrameworks)
-				{
-					DirectoryReference? FrameworkDirectory = Framework.GetFrameworkDirectory(BinaryLinkEnvironment.Platform, BinaryLinkEnvironment.Architecture, Logger);
-					if (FrameworkDirectory != null)
-					{
-						if (!String.IsNullOrEmpty(Framework.CopyBundledAssets))
-						{
-							// For now, this is hard coded, but we need to loop over all modules, and copy bundled assets that need it
-							DirectoryReference LocalSource = DirectoryReference.Combine(FrameworkDirectory, Framework.CopyBundledAssets);
-							string BundleName = Framework.CopyBundledAssets.Substring(Framework.CopyBundledAssets.LastIndexOf('/') + 1);
-							FrameworkNameToSourceDir[BundleName] = LocalSource;
-						}
-
-						if (Framework.bCopyFramework)
-						{
-							string FrameworkDir = FrameworkDirectory.FullName;
-							if (FrameworkDir.EndsWith(".framework", StringComparison.Ordinal))
-							{
-								FrameworkDir = Path.GetDirectoryName(FrameworkDir)!;
-							}
-
-							OutputFiles.Add(CopyBundleResource(new UEBuildBundleResource(new ModuleRules.BundleResource(Path.Combine(FrameworkDir, Framework.Name + ".framework"), "Frameworks")), Executable, BundleDirectory, Graph));
-						}
-					}
-				}
-
-				foreach (UEBuildBundleResource Resource in BinaryLinkEnvironment.AdditionalBundleResources)
-				{
-					OutputFiles.Add(CopyBundleResource(Resource, Executable, StagedExecutablePath.Directory, Graph));
-				}
-
-				IOSPostBuildSyncTarget PostBuildSyncTarget = new IOSPostBuildSyncTarget(Target, Executable.Location, BinaryLinkEnvironment.IntermediateDirectory, UPLScripts, SdkVersion, FrameworkNameToSourceDir);
-				FileReference PostBuildSyncFile = FileReference.Combine(BinaryLinkEnvironment.IntermediateDirectory!, "PostBuildSync.json");
-				DirectoryReference.CreateDirectory(PostBuildSyncFile.Directory);
-				using (FileStream stream = new(PostBuildSyncFile.FullName, FileMode.Create, FileAccess.Write))
-				{
-					JsonSerializer.Serialize(stream, PostBuildSyncTarget, AppleToolChain.SerializerOptions);
-				}
-
-				string PostBuildSyncArguments = String.Format("-Input=\"{0}\" -XmlConfigCache=\"{1}\" -remoteini=\"{2}\" ", PostBuildSyncFile, XmlConfig.CacheFile, UnrealBuildTool.GetRemoteIniPath());
-
-				if (Log.OutputFile != null)
-				{
-					string LogFileName = Log.OutputFile.GetFileNameWithoutExtension() + "_PostBuildSync.txt";
-					LogFileName = FileReference.Combine(Log.OutputFile.Directory, LogFileName).ToString();
-					PostBuildSyncArguments += " -log=\"" + LogFileName + "\"";
-				}
-
-				Action PostBuildSyncAction = Graph.CreateRecursiveAction<IOSPostBuildSyncMode>(ActionType.CreateAppBundle, PostBuildSyncArguments);
-				PostBuildSyncAction.WorkingDirectory = Unreal.EngineSourceDirectory;
-				PostBuildSyncAction.PrerequisiteItems.Add(Executable);
-				PostBuildSyncAction.PrerequisiteItems.UnionWith(OutputFiles);
-				PostBuildSyncAction.ProducedItems.Add(FileItem.GetItemByFileReference(IOSToolChain.GetStagedExecutablePath(Executable.Location, Target.Name)));
-				PostBuildSyncAction.DeleteItems.UnionWith(PostBuildSyncAction.ProducedItems);
-				PostBuildSyncAction.StatusDescription = "Executing PostBuildSync";
-
-				OutputFiles.AddRange(PostBuildSyncAction.ProducedItems);
-			}
-		}
 	}
 }
\ No newline at end of file
diff --git a/Engine/Source/Programs/UnrealBuildTool/Platform/Apple/AppleToolChain.cs b/Engine/Source/Programs/UnrealBuildTool/Platform/Apple/AppleToolChain.cs
--- a/Engine/Source/Programs/UnrealBuildTool/Platform/Apple/AppleToolChain.cs
+++ b/Engine/Source/Programs/UnrealBuildTool/Platform/Apple/AppleToolChain.cs
@@ -77,15 +77,13 @@ namespace UnrealBuildTool
 		readonly bool bUseSwiftUIMain;
 
 		/// <summary>
-		/// Serializer options for serializing <see cref="IOSPostBuildSyncTarget"/> and <see cref="ApplePostBuildSyncTarget"/> instances.
+		/// Serializer options for serializing <see cref="ApplePostBuildSyncTarget"/> instances.
 		/// </summary>
 		public static readonly JsonSerializerOptions SerializerOptions = new()
 		{
 			IncludeFields = true,
 		};
 
-		protected bool bCreatingAppOnWindows => AppleExports.CreatingAppOnWindows(ProjectFile);
-
 		protected ApplePortableToolChainInfo? PortableToolchain;
 		private static bool CrossArchitecturePathsRegistered;
 
diff --git a/Engine/Source/Programs/UnrealBuildTool/Platform/IOS/IOSCodeSigning.cs b/Engine/Source/Programs/UnrealBuildTool/Platform/IOS/IOSCodeSigning.cs
--- a/Engine/Source/Programs/UnrealBuildTool/Platform/IOS/IOSCodeSigning.cs
+++ b/Engine/Source/Programs/UnrealBuildTool/Platform/IOS/IOSCodeSigning.cs
@@ -551,25 +551,6 @@ namespace UnrealBuildTool
 			return FindPrefixedFileOrDirectory(BasePath, FileName, false);
 		}
 
-		/// <summary>
-		/// Finds the first file in the directory with a given extension (if it exists).  This method is to provide some
-		/// idiot-proofing (for example, .mobileprovision files, which will be expected to be saved from the website as
-		/// GameName.mobileprovision or Distro_GameName.mobileprovision, but falling back to using the first one we find
-		/// if the correctly named one is not found helps the user if they mess up during the process.
-		/// </summary>
-		public static string FindAnyFileWithExtension(string BasePath, string FileExt)
-		{
-			string[] FileList = Directory.GetFiles(BasePath);
-			foreach (string Filename in FileList)
-			{
-				if (Path.GetExtension(Filename).Equals(FileExt, StringComparison.OrdinalIgnoreCase))
-				{
-					return Filename;
-				}
-			}
-			return "";
-		}
-
 		/// <summary>
 		/// Find a prefixable file or directory (checks first with the prefix and then without it)
 		/// </summary>
@@ -593,99 +574,6 @@ namespace UnrealBuildTool
 		}
 	}
 
-	internal class Utilities
-	{
-		static ILogger? _logger;
-
-		/// <summary>
-		/// The constructor
-		/// </summary>
-		public Utilities()
-		{
-			ILoggerFactory loggerFactory = LoggerFactory.Create(builder =>
-			{
-				builder.AddConsole();
-			});
-			_logger = loggerFactory.CreateLogger("CodeSigningUtilities");
-		}
-
-		/// <summary>
-		/// Get the string from the plist
-		/// </summary>
-		public static string GetStringFromPList(string KeyName)
-		{
-			// Open the .plist and read out the specified key
-			string PListAsString;
-			if (!Utilities.GetSourcePList(out PListAsString))
-			{
-				return "(unknown)";
-			}
-
-			PListHelper Helper = new PListHelper(PListAsString);
-
-			string Result;
-			if (Helper.GetString(KeyName, out Result))
-			{
-				return Result;
-			}
-			else
-			{
-				return "(unknown)";
-			}
-		}
-
-		/// <summary>
-		/// Get the precompile source plist filename
-		/// </summary>
-		public static string GetPrecompileSourcePListFilename()
-		{
-			// check for one in the project directory
-			string SourceName = FileOperations.FindPrefixedFile(CodeSigningConfig.IntermediateDirectory, CodeSigningConfig.Program_GameName + "-Info.plist");
-
-			if (!File.Exists(SourceName))
-			{
-				// Check for a premade one
-				SourceName = FileOperations.FindPrefixedFile(CodeSigningConfig.BuildDirectory, CodeSigningConfig.Program_GameName + "-Info.plist");
-
-				if (!File.Exists(SourceName))
-				{
-					// fallback to the shared one
-					SourceName = FileOperations.FindPrefixedFile(CodeSigningConfig.EngineBuildDirectory, "UnrealGame-Info.plist");
-
-					if (!File.Exists(SourceName))
-					{
-						_logger?.LogInformation("Failed to find {GameName}-Info.plist. Please create new .plist or copy a base .plist from a provided game sample.", CodeSigningConfig.Program_GameName);
-					}
-				}
-			}
-
-			return SourceName;
-		}
-
-		/**
-		 * Handle grabbing the initial plist
-		 */
-		public static bool GetSourcePList(out string PListSource)
-		{
-			// Check for a premade one
-			string SourceName = GetPrecompileSourcePListFilename();
-
-			if (File.Exists(SourceName))
-			{
-				_logger?.LogInformation(" ... reading source .plist: {PlistPath}", SourceName);
-				PListSource = File.ReadAllText(SourceName);
-				return true;
-			}
-			else
-			{
-				_logger?.LogError(" ... failed to locate the source .plist file");
-				//Program.ReturnCode = (int)ErrorCodes.Error_KeyNotFoundInPList;
-				PListSource = "";
-				return false;
-			}
-		}
-	}
-
 	/// <summary>
 	/// The mobile provision
 	/// </summary>
diff --git a/Engine/Source/Programs/UnrealBuildTool/Platform/IOS/IOSPostBuildSyncMode.cs b/Engine/Source/Programs/UnrealBuildTool/Platform/IOS/IOSPostBuildSyncMode.cs
deleted file mode 100644
diff --git a/Engine/Source/Programs/UnrealBuildTool/Platform/IOS/IOSToolChain.cs b/Engine/Source/Programs/UnrealBuildTool/Platform/IOS/IOSToolChain.cs
--- a/Engine/Source/Programs/UnrealBuildTool/Platform/IOS/IOSToolChain.cs
+++ b/Engine/Source/Programs/UnrealBuildTool/Platform/IOS/IOSToolChain.cs
@@ -2,7 +2,6 @@
 
 using System;
 using System.Collections.Generic;
-using System.Diagnostics;
 using System.IO;
 using System.IO.Compression;
 using System.Linq;
@@ -71,12 +70,6 @@ namespace UnrealBuildTool
 			{
 				FileReference StubFile = FileReference.Combine(Binary.OutputFilePath.Directory, Binary.OutputFilePath.GetFileNameWithoutExtension() + ".stub");
 				BuildProducts.Add(StubFile, BuildProductType.Package);
-
-				// These Payload assets and .app bundle are needed for the Apple On Windows mode
-				if (bCreatingAppOnWindows)
-				{
-					AppleOnWindowsAppMaker.GatherBuildProducts(Target, Binary, BuildProducts);
-				}
 			}
 			if (Target.IOSPlatform.bGeneratedSYM && ProjectSettings.bGenerateCrashReportSymbols && Binary.Type == UEBuildBinaryType.StaticLibrary)
 			{
@@ -957,66 +950,9 @@ namespace UnrealBuildTool
 				}
 			}
 
-			if (bCreatingAppOnWindows)
-			{
-				AppleOnWindowsAppMaker.PostBuild(ProjectFile, Settings, Logger, Target, Executable, BinaryLinkEnvironment, Graph, ref OutputFiles);
-				return OutputFiles;
-			}
-			
-			// If we are not running in the Apple On Windows mode, the base class AppleToolchain makes the shared postbuild sync object, so we are done
+			// The base class AppleToolchain makes the shared postbuild sync object, so we are done
 			return OutputFiles;
 		}
-		public static string GetCodesignPlatformName(UnrealTargetPlatform Platform)
-		{
-			if (Platform == UnrealTargetPlatform.TVOS)
-			{
-				return "appletvos";
-			}
-			if (Platform == UnrealTargetPlatform.IOS)
-			{
-				return "iphoneos";
-			}
-
-			throw new BuildException("Invalid platform for GetCodesignPlatformName()");
-		}
-
-		public static FileReference GetStagedExecutablePath(FileReference Executable, string TargetName)
-		{
-			return FileReference.Combine(Executable.Directory, "Payload", TargetName + ".app", TargetName);
-		}
-
-		public static void PostBuildSync(IOSPostBuildSyncTarget Target, ILogger Logger)
-		{
-			// if xcode is building this, it will also do the Run stuff anyway, so no need to do it here as well
-			if (Environment.GetEnvironmentVariable("UE_BUILD_FROM_XCODE") == "1")
-			{
-				return;
-			}
-
-			if (AppleExports.CreatingAppOnWindows(Target.ProjectFile))
-			{
-				AppleOnWindowsAppMaker.PostBuildSync(Target, Logger);
-			}
-			else
-			{
-				throw new Exception("IOSToolchain.PostBuildSync was called outside the apple on windows build mode. This is no longer supported.");
-			}
-		}
-
-		public static int RunExecutableAndWait(string ExeName, string ArgumentList, out string StdOutResults)
-		{
-			// Create the process
-			ProcessStartInfo PSI = new ProcessStartInfo(ExeName, ArgumentList);
-			PSI.RedirectStandardOutput = true;
-			PSI.UseShellExecute = false;
-			PSI.CreateNoWindow = true;
-			Process NewProcess = Process.Start(PSI)!;
-
-			// Wait for the process to exit and grab it's output
-			StdOutResults = NewProcess.StandardOutput.ReadToEnd();
-			NewProcess.WaitForExit();
-			return NewProcess.ExitCode;
-		}
 
 		public void StripSymbols(FileReference SourceFile, FileReference TargetFile)
 		{
diff --git a/Engine/Source/Programs/UnrealBuildTool/Platform/IOS/MobileProvisionContents.cs b/Engine/Source/Programs/UnrealBuildTool/Platform/IOS/MobileProvisionContents.cs
--- a/Engine/Source/Programs/UnrealBuildTool/Platform/IOS/MobileProvisionContents.cs
+++ b/Engine/Source/Programs/UnrealBuildTool/Platform/IOS/MobileProvisionContents.cs
@@ -70,71 +70,6 @@ namespace UnrealBuildTool
 			return UniqueIdElement.InnerText;
 		}
 
-		/// <summary>
-		/// Gets the bundle id for this mobileprovision
-		/// </summary>
-		/// <returns>Bundle Identifier for the provision</returns>
-		public string GetBundleIdentifier()
-		{
-			XmlElement? UniqueIdElement = null;
-			XmlElement? UniqueIdEntitlement;
-			if (!NameToValue.TryGetValue("Entitlements", out UniqueIdEntitlement) || UniqueIdEntitlement.Name != "dict")
-			{
-				throw new BuildException("Missing Entitlements in MobileProvision");
-			}
-
-			XmlNodeList? Nodes = UniqueIdEntitlement.SelectNodes("key");
-			if (Nodes != null)
-			{
-				foreach (XmlElement? KeyElement in Nodes)
-				{
-					if (KeyElement == null)
-					{
-						continue;
-					}
-
-					Console.WriteLine("Found entitlement node:" + KeyElement.InnerText);
-					if (!KeyElement.InnerText.Equals("application-identifier", StringComparison.Ordinal))
-					{
-						continue;
-					}
-					UniqueIdElement = KeyElement.NextSibling as XmlElement;
-					break;
-				}
-			}
-
-			if (UniqueIdElement == null)
-			{
-				throw new BuildException("Missing Bundle Identifier in MobileProvision");
-			}
-			return UniqueIdElement.InnerText.Substring(UniqueIdElement.InnerText.IndexOf('.', StringComparison.Ordinal) + 1);
-		}
-
-		/// <summary>
-		/// Gets the team unique id for this mobileprovision
-		/// </summary>
-		/// <param name="UniqueId">Receives the team unique id</param>
-		/// <returns>True if the team unique ID was found, false otherwise</returns>
-		public bool TryGetTeamUniqueId(out string? UniqueId)
-		{
-			XmlElement? UniqueIdElement;
-			if (!NameToValue.TryGetValue("TeamIdentifier", out UniqueIdElement) || UniqueIdElement.Name != "array")
-			{
-				UniqueId = null;
-				return false;
-			}
-
-			XmlElement? ValueElement = UniqueIdElement.SelectSingleNode("string") as XmlElement;
-			if (ValueElement == null)
-			{
-				UniqueId = null;
-				return false;
-			}
-
-			UniqueId = ValueElement.InnerText;
-			return true;
-		}
-
 		/// <summary>
 		/// Reads a mobileprovision from a file on disk
 		/// </summary>
diff --git a/Engine/Source/Programs/UnrealBuildTool/Platform/IOS/UEBuildIOS.cs b/Engine/Source/Programs/UnrealBuildTool/Platform/IOS/UEBuildIOS.cs
--- a/Engine/Source/Programs/UnrealBuildTool/Platform/IOS/UEBuildIOS.cs
+++ b/Engine/Source/Programs/UnrealBuildTool/Platform/IOS/UEBuildIOS.cs
@@ -1179,11 +1179,7 @@ namespace UnrealBuildTool
 		/// <inheritdoc/>
 		public override void Deploy(TargetReceipt Receipt)
 		{
-			if (Receipt.HasValueForAdditionalProperty("CompileAsDll", "true"))
-			{
-				// IOSToolchain.PostBuildSync handles the copy, nothing else to do here
-			}
-			else
+			if (!Receipt.HasValueForAdditionalProperty("CompileAsDll", "true"))
 			{
 				new UEDeployIOS(Logger).PrepTargetForDeployment(Receipt);
 			}
diff --git a/Engine/Source/Programs/UnrealBuildTool/ProjectFiles/Xcode/XcodeProjectLegacy.cs b/Engine/Source/Programs/UnrealBuildTool/ProjectFiles/Xcode/XcodeProjectLegacy.cs
--- a/Engine/Source/Programs/UnrealBuildTool/ProjectFiles/Xcode/XcodeProjectLegacy.cs
+++ b/Engine/Source/Programs/UnrealBuildTool/ProjectFiles/Xcode/XcodeProjectLegacy.cs
@@ -217,7 +217,7 @@ namespace UnrealBuildTool.XcodeProjectLegacy
 		/// Returns a project navigator group to which the file should belong based on its path.
 		/// Creates a group tree if it doesn't exist yet.
 		/// </summary>
-		public static XcodeFileGroup? FindGroupByAbsolutePath(ref Dictionary<string, XcodeFileGroup> Groups, string AbsolutePath)
+		private static XcodeFileGroup? FindGroupByAbsolutePath(ref Dictionary<string, XcodeFileGroup> Groups, string AbsolutePath)
 		{
 			string[] Parts = AbsolutePath.Split(Path.DirectorySeparatorChar);
 			string CurrentPath = "/";
@@ -1770,11 +1770,6 @@ namespace UnrealBuildTool.XcodeProjectLegacy
 			};
 		}
 
-		public bool ShouldIncludeProjectInWorkspace()
-		{
-			return CanBuildProjectLocally();
-		}
-
 		public bool CanBuildProjectLocally()
 		{
 			foreach (Project ProjectTarget in ProjectTargets)
