UnrealBuildTool: Changes made whilst looking at Rider project file generation. Make some fields readonly that are never reassigned and supply constructor parameters in the correct order.

#Synced-CL 50903846

diff --git a/Engine/Source/Programs/UnrealBuildTool/ProjectFiles/ProjectTarget.cs b/Engine/Source/Programs/UnrealBuildTool/ProjectFiles/ProjectTarget.cs
--- a/Engine/Source/Programs/UnrealBuildTool/ProjectFiles/ProjectTarget.cs
+++ b/Engine/Source/Programs/UnrealBuildTool/ProjectFiles/ProjectTarget.cs
@@ -12,16 +12,16 @@ namespace UnrealBuildTool
 	class ProjectTarget : Project
 	{
 		/// The target rules file path on disk, if we have one
-		public FileReference TargetFilePath;
+		public FileReference TargetFilePath { get; }
 
 		/// The project file path on disk
-		public FileReference ProjectFilePath;
+		public FileReference ProjectFilePath { get; }
 
 		/// Delegate for creating a rules instance for a given platform/configuration
-		public Func<UnrealTargetPlatform, UnrealTargetConfiguration, TargetRules> CreateRulesDelegate;
+		public Func<UnrealTargetPlatform, UnrealTargetConfiguration, TargetRules> CreateRulesDelegate { get; }
 
 		/// Delegate for creating a rules instance for a given platform/configuration/architecture
-		public Func<UnrealTargetPlatform, UnrealTargetConfiguration, UnrealArch, TargetRules> CreateRulesDelegateWithArch;
+		public Func<UnrealTargetPlatform, UnrealTargetConfiguration, UnrealArch, TargetRules> CreateRulesDelegateWithArch { get;  }
 
 		public string Name => TargetFilePath.GetFileNameWithoutAnyExtensions();
 
