UnrealBuildTool: Force a UHT run when running remote Mac builds, as otherwise builds of e.g. UnrealClient fail.
#Synced-CL 50903846

diff --git a/Engine/Source/Programs/UnrealBuildTool/Platform/Apple/RemoteMac.cs b/Engine/Source/Programs/UnrealBuildTool/Platform/Apple/RemoteMac.cs
--- a/Engine/Source/Programs/UnrealBuildTool/Platform/Apple/RemoteMac.cs
+++ b/Engine/Source/Programs/UnrealBuildTool/Platform/Apple/RemoteMac.cs
@@ -513,6 +513,8 @@ namespace UnrealBuildTool
 			RemoteArguments.Add(String.Format("-Log={0}", GetRemotePath(RemoteLogFile)));
 			RemoteArguments.Add(String.Format("-Manifest={0}", GetRemotePath(RemoteManifestFile)));
 			RemoteArguments.Add(String.Format("-SkipPreBuildTargets"));
+			// Compilation has been observed to fail without forcing a UHT run.
+			RemoteArguments.Add("-ForceHeaderGeneration");
 
 			// Handle any per-platform setup that is required
 			if (TargetDesc.Platform == UnrealTargetPlatform.IOS || TargetDesc.Platform == UnrealTargetPlatform.TVOS)
