SNDBS: Prevent exception when there is no explicit_input_files key.

#Synced-CL 0

diff --git a/Engine/Source/Programs/UnrealBuildTool/Executors/SNDBS.cs b/Engine/Source/Programs/UnrealBuildTool/Executors/SNDBS.cs
--- a/Engine/Source/Programs/UnrealBuildTool/Executors/SNDBS.cs
+++ b/Engine/Source/Programs/UnrealBuildTool/Executors/SNDBS.cs
@@ -286,7 +286,7 @@ namespace UnrealBuildTool
 					}
 
 					// Look for any prerequisite actions that produce .pch files and add any .cpp files they depend on
-					List<Dictionary<string, object>>? ExplicitInputFiles = Job["explicit_input_files"] as List<Dictionary<string, object>>;
+					List<Dictionary<string, object>>? ExplicitInputFiles = Job.ContainsKey("explicit_input_files") ? Job["explicit_input_files"] as List<Dictionary<string, object>> : null;
 
 					ExplicitInputFiles?.AddRange(a.PrerequisiteActions
 						.Where(Prereq => Prereq.ProducedItems.Any(Produced => Produced.AbsolutePath.EndsWith(".gch", System.StringComparison.Ordinal) || Produced.AbsolutePath.EndsWith(".pch", System.StringComparison.Ordinal)))
