[UBA]
* Added method to UbaHordeSession to be able to disable remote execution on any active agent

#Synced-CL 50919680

diff --git a/Engine/Source/Programs/UnrealBuildTool/Executors/UnrealBuildAccelerator/UBAHordeSession.cs b/Engine/Source/Programs/UnrealBuildTool/Executors/UnrealBuildAccelerator/UBAHordeSession.cs
--- a/Engine/Source/Programs/UnrealBuildTool/Executors/UnrealBuildAccelerator/UBAHordeSession.cs
+++ b/Engine/Source/Programs/UnrealBuildTool/Executors/UnrealBuildAccelerator/UBAHordeSession.cs
@@ -55,6 +55,7 @@ namespace UnrealBuildTool
 
 		public class Worker
 		{
+			public string? Name { get; init; }
 			public Stopwatch? StartTime { get; init; }
 			public string? Ip { get; init; }
 			public ConnectionMetadataPort? Port { get; init; }
@@ -127,6 +128,7 @@ namespace UnrealBuildTool
 			if (_targetWorkerPlatform == OSPlatform.Windows)
 			{
 				agentFiles.Add("UbaAgent.exe");
+				//agentFiles.Add("UbaAgent.pdb");
 				if (FileReference.Exists(FileReference.Combine(ubaDir, "UbaWine.dll.so")))
 				{
 					agentFiles.Add("UbaWine.dll.so"); // Not needed but used by linux/wine helpers for more detailed logging or optimizations
@@ -412,6 +414,7 @@ namespace UnrealBuildTool
 
 				Worker worker = new()
 				{
+					Name = agentName,
 					StartTime = stopwatch,
 					NumLogicalCores = numLogicalCores,
 					Ip = ip,
@@ -718,6 +721,24 @@ namespace UnrealBuildTool
 		string _additionalStatus = "";
 		string _lastStatus = "";
 
+		internal bool DisableRemoteExecutionOnAgent()
+		{
+			lock (_workers)
+			{
+				for (int i = _workers.Count - 1; i >= 0; i--)
+				{
+					Worker worker = _workers[i];
+					if (!worker.Active)
+					{
+						continue;
+					}
+					_scheduler?.DisableRemoteExecutionOnAgent(worker.Name!);
+					return true;
+				}
+			}
+			return false;
+		}
+
 		internal void UpdateHordeStatus(string? additionalStatus)
 		{
 			if (additionalStatus != null)
