Don't build MetalShaderFormat for all Win64 builds. The module should only be included for non-program builds with bBuildRequiresCookedData set to false.

#Synced-CL 50914396

diff --git a/Engine/Source/Programs/UnrealBuildTool/Platform/IOS/UEBuildIOS.cs b/Engine/Source/Programs/UnrealBuildTool/Platform/IOS/UEBuildIOS.cs
--- a/Engine/Source/Programs/UnrealBuildTool/Platform/IOS/UEBuildIOS.cs
+++ b/Engine/Source/Programs/UnrealBuildTool/Platform/IOS/UEBuildIOS.cs
@@ -950,10 +950,17 @@ namespace UnrealBuildTool
 
 			AddPlatformStaticShaders(Target, Rules, []);
 
-			if ((Target.Platform == UnrealTargetPlatform.Win64) || (Target.Platform == UnrealTargetPlatform.Mac))
+			if (Target.Platform == UnrealTargetPlatform.Win64 || Target.Platform == UnrealTargetPlatform.Mac)
 			{
-				// For Windows we use the MetalDeveloperTools and we will check if the toolchain is present at runtime
-				bool bBuildShaderFormats = Target.bForceBuildShaderFormats || Target.Platform == UnrealTargetPlatform.Win64;
+				bool bBuildShaderFormats = Target.bForceBuildShaderFormats;
+				if (!Target.bBuildRequiresCookedData && Target.Type != TargetType.Program)
+				{
+					if (ModuleName == "TargetPlatform")
+					{
+						bBuildShaderFormats = true;
+					}
+				}
+
 				if (!Target.bBuildRequiresCookedData && Target.bBuildDeveloperTools)
 				{
 					if (ModuleName == "Engine")
@@ -973,7 +980,6 @@ namespace UnrealBuildTool
 					{
 						if (bIsPlatformAvailableForTargetWithSDK)
 						{
-							bBuildShaderFormats = true;
 							Rules.DynamicallyLoadedModuleNames.Add("TextureFormatASTC");
 							Rules.DynamicallyLoadedModuleNames.Add("TextureFormatETC2");
 							if (Target.bCompileAgainstEngine)
