Automation: Remove references to Ionic.Zip.Reduced from places where it's not used.

#Synced-CL 50931528

diff --git a/Engine/Platforms/VisionOS/Source/Programs/AutomationTool/VisionOS.Automation.csproj b/Engine/Platforms/VisionOS/Source/Programs/AutomationTool/VisionOS.Automation.csproj
--- a/Engine/Platforms/VisionOS/Source/Programs/AutomationTool/VisionOS.Automation.csproj
+++ b/Engine/Platforms/VisionOS/Source/Programs/AutomationTool/VisionOS.Automation.csproj
@@ -11,12 +11,6 @@
     <OutputPath>..\..\..\Binaries\DotNET\AutomationTool\AutomationScripts\</OutputPath>
     <TreatWarningsAsErrors>true</TreatWarningsAsErrors>
   </PropertyGroup>
-  <ItemGroup>
-    <Reference Include="Ionic.Zip.Reduced">
-      <HintPath>..\..\..\..\..\Binaries\DotNET\Ionic.Zip.Reduced.dll</HintPath>
-      <Private>False</Private>
-    </Reference>
-  </ItemGroup>
 
   <ItemGroup>
     <Compile Include="..\..\..\..\..\Source\Programs\Shared\MetaData.cs">
diff --git a/Engine/Source/Programs/AutomationTool/AutomationTool.csproj b/Engine/Source/Programs/AutomationTool/AutomationTool.csproj
--- a/Engine/Source/Programs/AutomationTool/AutomationTool.csproj
+++ b/Engine/Source/Programs/AutomationTool/AutomationTool.csproj
@@ -52,12 +52,8 @@
     <GenerateDocumentationFile>True</GenerateDocumentationFile>
     <DocumentationFile></DocumentationFile>
   </PropertyGroup>
-  
+
   <ItemGroup>
-    <Reference Include="Ionic.Zip.Reduced">
-      <HintPath>..\..\..\Binaries\DotNET\Ionic.Zip.Reduced.dll</HintPath>
-      <Private>true</Private>
-    </Reference>
     <Reference Include="fastJSON">
       <HintPath>..\..\..\Binaries\ThirdParty\fastJSON\netstandard2.0\fastJSON.dll</HintPath>
       <Private>true</Private>
diff --git a/Engine/Source/Programs/AutomationTool/Gauntlet/Gauntlet.Automation.csproj b/Engine/Source/Programs/AutomationTool/Gauntlet/Gauntlet.Automation.csproj
--- a/Engine/Source/Programs/AutomationTool/Gauntlet/Gauntlet.Automation.csproj
+++ b/Engine/Source/Programs/AutomationTool/Gauntlet/Gauntlet.Automation.csproj
@@ -17,10 +17,6 @@
     <Reference Include="fastJSON">
       <HintPath>..\..\..\..\Binaries\ThirdParty\fastJSON\netstandard2.0\fastJSON.dll</HintPath>
     </Reference>
-    <Reference Include="Ionic.Zip.Reduced">
-      <HintPath>..\..\..\..\Binaries\DotNET\Ionic.Zip.Reduced.dll</HintPath>
-      <Private>False</Private>
-    </Reference>
   </ItemGroup>
 
     <!-- The following will glob all platform extensions' Gauntlet source
diff --git a/Engine/Source/Programs/AutomationTool/Gauntlet/Platform/IOS/Gauntlet.IOSBuildSource.cs b/Engine/Source/Programs/AutomationTool/Gauntlet/Platform/IOS/Gauntlet.IOSBuildSource.cs
--- a/Engine/Source/Programs/AutomationTool/Gauntlet/Platform/IOS/Gauntlet.IOSBuildSource.cs
+++ b/Engine/Source/Programs/AutomationTool/Gauntlet/Platform/IOS/Gauntlet.IOSBuildSource.cs
@@ -77,8 +77,6 @@ namespace Gauntlet
 			return Result;
 		}
 
-		// There are issues with IPA Zip64 files being created with Ionic.Zip possibly limited to when running on mono (see IOSPlatform.PackageIPA)
-		// This manifests as header overflow errors, etc in 7zip, Ionic.Zip, System.IO.Compression, and OSX system unzip
 		internal static bool ExecuteIPAZipCommand(String Arguments, out String Output, String ShouldExist = "")
 		{
 			using (new ScopedSuspendECErrorParsing())
diff --git a/Engine/Source/Programs/AutomationTool/TVOS/TVOS.Automation.csproj b/Engine/Source/Programs/AutomationTool/TVOS/TVOS.Automation.csproj
--- a/Engine/Source/Programs/AutomationTool/TVOS/TVOS.Automation.csproj
+++ b/Engine/Source/Programs/AutomationTool/TVOS/TVOS.Automation.csproj
@@ -11,12 +11,6 @@
     <OutputPath>..\..\..\..\Binaries\DotNET\AutomationTool\AutomationScripts\Platforms\TVOS\</OutputPath>
     <TreatWarningsAsErrors>true</TreatWarningsAsErrors>
   </PropertyGroup>
-  <ItemGroup>
-    <Reference Include="Ionic.Zip.Reduced">
-      <HintPath>..\..\..\..\Binaries\DotNET\Ionic.Zip.Reduced.dll</HintPath>
-      <Private>False</Private>
-    </Reference>
-  </ItemGroup>
 
   <ItemGroup>
     <Compile Include="..\..\Shared\MetaData.cs">
diff --git a/Engine/Source/Programs/AutomationTool/Win/Win.Automation.csproj b/Engine/Source/Programs/AutomationTool/Win/Win.Automation.csproj
--- a/Engine/Source/Programs/AutomationTool/Win/Win.Automation.csproj
+++ b/Engine/Source/Programs/AutomationTool/Win/Win.Automation.csproj
@@ -11,12 +11,6 @@
     <OutputPath>..\..\..\..\Binaries\DotNET\AutomationTool\AutomationScripts\Platforms\Windows</OutputPath>
     <TreatWarningsAsErrors>true</TreatWarningsAsErrors>
   </PropertyGroup>
-  <ItemGroup>
-    <Reference Include="Ionic.Zip.Reduced">
-      <HintPath>..\..\..\..\Binaries\DotNET\Ionic.Zip.Reduced.dll</HintPath>
-      <Private>False</Private>
-    </Reference>
-  </ItemGroup>
 
   <ItemGroup>
     <Compile Include="..\..\Shared\MetaData.cs">
diff --git a/Engine/Source/Programs/UnrealBuildTool/Platform/IOS/IOSCodeSigning.cs b/Engine/Source/Programs/UnrealBuildTool/Platform/IOS/IOSCodeSigning.cs
--- a/Engine/Source/Programs/UnrealBuildTool/Platform/IOS/IOSCodeSigning.cs
+++ b/Engine/Source/Programs/UnrealBuildTool/Platform/IOS/IOSCodeSigning.cs
@@ -1624,7 +1624,8 @@ namespace UnrealBuildTool
 		/// Whether to use None or Best for the compression setting when repackaging IPAs (int version of Ionic.Zlib.CompressionLevel)
 		/// By making this an int, we are able to delay load the Ionic assembly from a different path (which is required)
 		/// </summary>
-		public static int RecompressionSetting = 0; //Ionic.Zlib.CompressionLevel.None;
+		[Obsolete("Deprecated in 5.8. The Ionic compression library is no longer used.")]
+		public static int RecompressionSetting = 0;
 
 		/// <summary>
 		/// If the requirements blob is present in the existing executable when code signing, should it be carried forward
