UnrealBuildTool: Add information about third-party modules to Analyze mode.

#Synced-CL 50931528

diff --git a/Engine/Source/Programs/UnrealBuildTool/Configuration/Rules/ExternalModuleMetadata.cs b/Engine/Source/Programs/UnrealBuildTool/Configuration/Rules/ExternalModuleMetadata.cs
new file mode 100644
--- /dev/null
+++ b/Engine/Source/Programs/UnrealBuildTool/Configuration/Rules/ExternalModuleMetadata.cs
@@ -0,0 +1,30 @@
+// Copyright Epic Games, Inc. All Rights Reserved.
+
+using System;
+using EpicGames.Core;
+
+namespace UnrealBuildTool;
+
+/// <summary>
+/// Metadata describing an external (third-party) module.
+/// </summary>
+/// <param name="Uri">A URI pointing to the homepage of the external software.</param>
+/// <param name="Version">
+/// The referenced version of the external software.
+/// This is modelled as a string as not all external dependencies will follow a specific versioning convention such as SemVer.
+/// </param>
+/// <param name="TpsFilePath">
+/// The path to the `.tps` file for this external software.
+/// This can be relative to the module file, or use <c>$(EngineDir)</c> or <c>$(ProjectDir)</c>.
+/// </param>
+public record ExternalModuleMetadata(
+	Uri Uri,
+	string Version,
+	string TpsFilePath
+);
+
+internal record ProcessedExternalModuleMetadata(
+	Uri Uri,
+	string Version,
+	FileReference TpsFilePath
+);
diff --git a/Engine/Source/Programs/UnrealBuildTool/Configuration/Rules/ModuleRules.cs b/Engine/Source/Programs/UnrealBuildTool/Configuration/Rules/ModuleRules.cs
--- a/Engine/Source/Programs/UnrealBuildTool/Configuration/Rules/ModuleRules.cs
+++ b/Engine/Source/Programs/UnrealBuildTool/Configuration/Rules/ModuleRules.cs
@@ -1560,6 +1560,12 @@ namespace UnrealBuildTool
 		/// </summary>
 		public List<string> FPassPlugins = [];
 
+		/// <summary>
+		/// The metadata for external (third-party) modules.
+		/// Should be non-<see langword="null"/> if and only if the module is of type <see cref="ModuleType.External"/>.
+		/// </summary>
+		public ExternalModuleMetadata? ExternalMetadata { get; set; }
+
 #nullable disable
 		/// <summary>
 		/// Constructor. For backwards compatibility while the parameterless constructor is being phased out, initialization which would happen here is done by 
diff --git a/Engine/Source/Programs/UnrealBuildTool/Configuration/UEBuildTarget.cs b/Engine/Source/Programs/UnrealBuildTool/Configuration/UEBuildTarget.cs
--- a/Engine/Source/Programs/UnrealBuildTool/Configuration/UEBuildTarget.cs
+++ b/Engine/Source/Programs/UnrealBuildTool/Configuration/UEBuildTarget.cs
@@ -4721,7 +4721,7 @@ namespace UnrealBuildTool
 		/// </summary>
 		/// <param name="Plugin">The current plugin</param>
 		/// <returns>Map of variable names to values</returns>
-		private Dictionary<string, string> GetTargetVariables(UEBuildPlugin? Plugin)
+		internal Dictionary<string, string> GetTargetVariables(UEBuildPlugin? Plugin)
 		{
 			Dictionary<string, string> Variables = new Dictionary<string, string>();
 			Variables.Add("RootDir", Unreal.RootDirectory.FullName);
diff --git a/Engine/Source/Programs/UnrealBuildTool/Modes/AnalyzeMode.cs b/Engine/Source/Programs/UnrealBuildTool/Modes/AnalyzeMode.cs
--- a/Engine/Source/Programs/UnrealBuildTool/Modes/AnalyzeMode.cs
+++ b/Engine/Source/Programs/UnrealBuildTool/Modes/AnalyzeMode.cs
@@ -2,6 +2,7 @@
 
 using System;
 using System.Collections.Generic;
+using System.Collections.Immutable;
 using System.IO;
 using System.Linq;
 using System.Threading.Tasks;
@@ -56,6 +57,7 @@ namespace UnrealBuildTool.Modes
 			public long ObjSize = 0;
 			public List<FileReference> BinaryFiles = [];
 			public long BinSize = 0;
+			public ProcessedExternalModuleMetadata? ExternalMetadata = null;
 
 			public string Chain => String.Join(" -> ", IncludeChain.Where(x => !String.IsNullOrEmpty(x)));
 		}
@@ -247,6 +249,33 @@ namespace UnrealBuildTool.Modes
 				logger.LogInformation("Missing module '{MissingModuleName}'", missingModule.Name);
 			}
 
+			Dictionary<string, string> variables = target.GetTargetVariables(null);
+			ImmutableArray<ModuleInfo> externalModules = [.. moduleToInfo.Values
+				.Where(x => x.Module.Rules.Type == ModuleRules.ModuleType.External)
+				.OrderBy(x => x.Module.Name)
+			];
+			foreach (ModuleInfo moduleInfo in externalModules)
+			{
+				try
+				{
+					moduleInfo.ExternalMetadata = ProcessExternalModuleMetadata(moduleInfo, variables);
+				}
+				catch (Exception ex)
+				{
+					logger.LogInformation(ex, "Failed to process external module metadata for module '{ModuleName}'.", moduleInfo.Module.Name);
+				}
+				if (moduleInfo.ExternalMetadata is null)
+				{
+					logger.LogInformation("Module '{ModuleName}' is external but has no ExternalMetadata set.", moduleInfo.Module.Name);
+					continue;
+				}
+
+				if (!FileReference.Exists(moduleInfo.ExternalMetadata.TpsFilePath))
+				{
+					logger.LogInformation("Module '{ModuleName}' points to TPS file '{TpsFile}' that does not exist.", moduleInfo.Module.Name, moduleInfo.ExternalMetadata.TpsFilePath);
+				}
+			}
+
 			List<KeyValuePair<FileReference, BuildProductType>> analyzeProducts = [];
 
 			// Generate the dependency graph between modules
@@ -273,12 +302,30 @@ namespace UnrealBuildTool.Modes
 			WriteInformationCsv(moduleToInfo, csvFile);
 			analyzeProducts.Add(new(csvFile, BuildProductType.BuildResource));
 
+			// Write all the external module dependencies as a CSV file.
+			FileReference externalCsvFile = target.ReceiptFileName.ChangeExtension(".External.csv");
+			logger.LogInformation("Writing external dependencies information to {CsvFile}", externalCsvFile);
+			WriteExternalCsv(externalModules, externalCsvFile);
+			analyzeProducts.Add(new(externalCsvFile, BuildProductType.BuildResource));
+
 			foreach (FileReference manifestFileName in target.Rules.ManifestFileNames)
 			{
 				target.GenerateManifest(manifestFileName, analyzeProducts, logger);
 			}
 		}
 
+		private static ProcessedExternalModuleMetadata? ProcessExternalModuleMetadata(ModuleInfo moduleInfo, Dictionary<string, string> variables)
+		{
+			if (moduleInfo.Module.Rules.ExternalMetadata is not ExternalModuleMetadata metadata)
+			{
+				return null;
+			}
+
+			string tpsPath = Utils.ExpandVariables(metadata.TpsFilePath, variables, true);
+			FileReference tpsReference = FileReference.Combine(moduleInfo.Module.ModuleDirectory, tpsPath);
+			return new(metadata.Uri, metadata.Version, tpsReference);
+		}
+
 		private static void WriteDependencyList(TextWriter writer, string prefix, HashSet<UEBuildModule> modules)
 		{
 			if (modules.Count == 0)
@@ -424,5 +471,30 @@ namespace UnrealBuildTool.Modes
 				}
 			}
 		}
+
+		private static void WriteExternalCsv(ImmutableArray<ModuleInfo> externalModules, FileReference csvFile)
+		{
+			using (StreamWriter writer = new(csvFile.FullName))
+			{
+				List<string> columns =
+				[
+					"Module",
+					"Uri",
+					"Version",
+					"TpsFile",
+				];
+				writer.WriteLine(String.Join(",", columns));
+
+				foreach (ModuleInfo moduleInfo in externalModules)
+				{
+					columns.Clear();
+					columns.Add(moduleInfo.Module.Name);
+					columns.Add(moduleInfo.ExternalMetadata?.Uri.ToString() ?? "<unavailable>");
+					columns.Add(moduleInfo.ExternalMetadata?.Version ?? "<unavailable>");
+					columns.Add(moduleInfo.ExternalMetadata?.TpsFilePath.ToString() ?? "<unavailable>");
+					writer.WriteLine(String.Join(",", columns));
+				}
+			}
+		}
 	}
 }
diff --git a/Engine/Source/Programs/UnrealBuildTool/System/DynamicCompilation.cs b/Engine/Source/Programs/UnrealBuildTool/System/DynamicCompilation.cs
--- a/Engine/Source/Programs/UnrealBuildTool/System/DynamicCompilation.cs
+++ b/Engine/Source/Programs/UnrealBuildTool/System/DynamicCompilation.cs
@@ -245,6 +245,7 @@ namespace UnrealBuildTool
 				"System.Threading.Tasks",
 				"System.Threading.Tasks.Parallel",
 				"System.Collections.Concurrent",
+				"System.Private.Uri",
 				"System.Private.Xml",
 				"System.Private.Xml.Linq",
 				"System.Text.RegularExpressions",
