UnrealBuildTool: Enable and selectively ignore issue CA1707 (Identifiers should not contain underscores) for public facing APIs.

#Synced-CL 50931528

diff --git a/Engine/Source/Programs/UnrealBuildTool/Configuration/Rules/TargetRules.cs b/Engine/Source/Programs/UnrealBuildTool/Configuration/Rules/TargetRules.cs
--- a/Engine/Source/Programs/UnrealBuildTool/Configuration/Rules/TargetRules.cs
+++ b/Engine/Source/Programs/UnrealBuildTool/Configuration/Rules/TargetRules.cs
@@ -182,9 +182,10 @@ namespace UnrealBuildTool
 	/// Specific versions can be used by licensees to avoid compile errors when integrating new versions of Unreal.
 	/// Using Latest comes with a high risk of introducing compile errors in your code on newer Unreal versions.
 	/// </summary>
+#pragma warning disable CA1707 // Identifiers should not contain underscores - public facing API
 #pragma warning disable CA1027 // Mark enums with FlagsAttribute
 	public enum EngineIncludeOrderVersion
-#pragma warning restore  CA1027 // Mark enums with FlagsAttribute
+#pragma warning restore CA1027 // Mark enums with FlagsAttribute
 	{
 		/// <summary>
 		/// Include order used in Unreal 5.1
@@ -247,6 +248,7 @@ namespace UnrealBuildTool
 		Oldest = Unreal5_4,
 #pragma warning restore CS0618 // Type or member is obsolete
 	}
+#pragma warning restore CA1707
 
 	/// <summary>
 	/// Which static analyzer to use
diff --git a/Engine/Source/Programs/UnrealBuildTool/Platform/Windows/AppXManifestGeneratorBase.cs b/Engine/Source/Programs/UnrealBuildTool/Platform/Windows/AppXManifestGeneratorBase.cs
--- a/Engine/Source/Programs/UnrealBuildTool/Platform/Windows/AppXManifestGeneratorBase.cs
+++ b/Engine/Source/Programs/UnrealBuildTool/Platform/Windows/AppXManifestGeneratorBase.cs
@@ -18,6 +18,7 @@ namespace UnrealBuildTool
 	/// </summary>
 	public abstract class AppXManifestGeneratorBase
 	{
+#pragma warning disable CA1707 // Identifiers should not contain underscores - public facing API
 		/// config section for platform-specific target settings
 		protected virtual string IniSection_PlatformTargetSettings => String.Format("/Script/{0}PlatformEditor.{0}TargetSettings", Platform.ToString());
 
@@ -26,6 +27,7 @@ namespace UnrealBuildTool
 
 		/// config section for general target settings
 		protected virtual string IniSection_GeneralProjectSettings => "/Script/EngineSettings.GeneralProjectSettings";
+#pragma warning restore CA1707
 
 		/// default subdirectory for build resources
 		protected const string BuildResourceSubPath = "Resources";
diff --git a/Engine/Source/Programs/UnrealBuildTool/Platform/Windows/PVSToolChain.cs b/Engine/Source/Programs/UnrealBuildTool/Platform/Windows/PVSToolChain.cs
--- a/Engine/Source/Programs/UnrealBuildTool/Platform/Windows/PVSToolChain.cs
+++ b/Engine/Source/Programs/UnrealBuildTool/Platform/Windows/PVSToolChain.cs
@@ -58,6 +58,7 @@ namespace UnrealBuildTool
 		MISRA = 32,
 	}
 
+#pragma warning disable CA1707 // Identifiers should not contain underscores - public facing API
 	/// <summary>
 	/// Supported PVS analyzer timeout values
 	/// </summary>
@@ -80,6 +81,7 @@ namespace UnrealBuildTool
 		/// </summary>
 		After_60_minutes = 3600,
 	}
+#pragma warning restore CA1707
 
 	/// <summary>
 	/// Partial representation of PVS-Studio main settings file
diff --git a/Engine/Source/Programs/UnrealBuildTool/System/CppCompileEnvironment.cs b/Engine/Source/Programs/UnrealBuildTool/System/CppCompileEnvironment.cs
--- a/Engine/Source/Programs/UnrealBuildTool/System/CppCompileEnvironment.cs
+++ b/Engine/Source/Programs/UnrealBuildTool/System/CppCompileEnvironment.cs
@@ -119,6 +119,7 @@ namespace UnrealBuildTool
 	/// Note that by enabling this you are changing the minspec for the PC platform, and the resultant executable will crash on machines without AVX support.
 	/// For more details please see https://learn.microsoft.com/en-us/cpp/build/reference/arch-x64
 	/// </summary>
+#pragma warning disable CA1707 // Identifiers should not contain underscores - public facing API
 #pragma warning disable CA1027 // Mark enums with FlagsAttribute
 	public enum MinimumCpuArchitectureX64
 #pragma warning restore CA1027 // Mark enums with FlagsAttribute
@@ -158,6 +159,7 @@ namespace UnrealBuildTool
 		/// </summary>
 		Default = None,
 	}
+#pragma warning restore CA1707
 
 	/// <summary>
 	/// Target ARM64 CPU for codegen. This defines different ARM ISA and extension targets
diff --git a/Engine/Source/Programs/UnrealBuildTool/System/Plugins.cs b/Engine/Source/Programs/UnrealBuildTool/System/Plugins.cs
--- a/Engine/Source/Programs/UnrealBuildTool/System/Plugins.cs
+++ b/Engine/Source/Programs/UnrealBuildTool/System/Plugins.cs
@@ -243,6 +243,7 @@ namespace UnrealBuildTool
 		/// </summary>
 		static ConcurrentDictionary<DirectoryReference, Lazy<List<PluginInfo>>> PluginDirectoryCache = [];
 
+#pragma warning disable CA1707 // Identifiers should not contain underscores - public facing API
 		/// <summary>
 		/// Invalidate cached plugin data so that we can pickup new things
 		/// Warning: Will make subsequent plugin lookups and directory scans slow until the caches are repopulated
@@ -254,6 +255,7 @@ namespace UnrealBuildTool
 			PluginsBase.InvalidateCache_SLOW();
 			DirectoryItem.ResetAllCachedInfo_SLOW();
 		}
+#pragma warning restore CA1707
 
 		/// <summary>
 		/// Add a means to retrieve a plugin to determine if it is present.
