HWASAN debug fixes.
Remove any possible occurences of CRLFs from the asan.sh or hwasan.sh files when generating wrap.sh.
If the hwasan.sh is recognised we replace it with a debuggable script instead.

#tests local HWASAN testing and debugging.
#Synced-CL 50901530

diff --git a/Engine/Source/Programs/UnrealBuildTool/Platform/Android/UEDeployAndroid.cs b/Engine/Source/Programs/UnrealBuildTool/Platform/Android/UEDeployAndroid.cs
--- a/Engine/Source/Programs/UnrealBuildTool/Platform/Android/UEDeployAndroid.cs
+++ b/Engine/Source/Programs/UnrealBuildTool/Platform/Android/UEDeployAndroid.cs
@@ -1556,7 +1556,47 @@ namespace UnrealBuildTool
 
 				string WrapDestFilePath = Path.Combine(WrapDestDir, "wrap.sh");
 				Logger.LogInformation("Copying wrap.sh from {WrapShFilePath} to {WrapDestFilePath}", WrapShFilePath, WrapDestFilePath);
-				File.Copy(WrapShFilePath, WrapDestFilePath, true);
+
+				string WrapShContent = File.ReadAllText(WrapShFilePath);
+				// replace the hwasan.sh with a debuggable script.
+				// see https://developer.android.com/ndk/guides/wrap-script#debugging_when_using_wrapsh
+				// first, check that this is the hwasan.sh we're expecting:
+				if (WrapShFilePath.Contains("hwasan.sh", StringComparison.Ordinal))
+				{
+					if (WrapShContent.Replace("\n", "", StringComparison.Ordinal).Replace("\r", "", StringComparison.Ordinal) == @"#!/system/bin/shLD_HWASAN=1 exec ""$@""")
+					{
+						Logger.LogInformation("Substituting {WrapShFilePath} with debuggable version.", WrapShFilePath);
+						WrapShContent =
+@"#!/system/bin/sh
+HERE=$(cd ""$(dirname ""$0"")"" && pwd)
+
+cmd=$1
+shift
+
+# This must be called *before* `LD_PRELOAD` is set. Otherwise, if this is a 32-
+# bit app running on a 64-bit device, the 64-bit getprop will fail to load
+# because it will preload a 32-bit ASan runtime.
+# [REDACTED_PROFILE_URL]/ndk/issues/1744
+os_version=$(getprop ro.build.version.sdk)
+
+if [ ""$os_version"" -eq ""27"" ]; then
+	cmd=""$cmd -Xrunjdwp:transport=dt_android_adb,suspend=n,server=y -Xcompiler-option --debuggable $@""
+elif [ ""$os_version"" -eq ""28"" ]; then
+	cmd=""$cmd -XjdwpProvider:adbconnection -XjdwpOptions:suspend=n,server=y -Xcompiler-option --debuggable $@""
+else
+	cmd=""$cmd -XjdwpProvider:adbconnection -XjdwpOptions:suspend=n,server=y $@""
+fi
+
+LD_HWASAN=1 exec $cmd";
+					}
+					else
+					{
+						Logger.LogWarning("Warning: The generated APK may not be debuggable. Unrecognized {WrapShFilePath} was not replaced with a debuggable script.", WrapShFilePath);
+					}
+				}
+				// handle any CRLF that may have crept in from windows text file handling.
+				WrapShContent = WrapShContent.Replace("\r", "", StringComparison.Ordinal);
+				File.WriteAllText(WrapDestFilePath, WrapShContent);
 
 				FileAttributes Attributes = File.GetAttributes(WrapDestFilePath);
 				Attributes &= ~FileAttributes.ReadOnly;
