UBT: Add C++ library paths to the system include paths on Linux

#rnx

#Synced-CL 50875397
#Synced-CL 50875397

diff --git a/Engine/Source/Programs/UnrealBuildTool/Platform/Linux/LinuxProjectGenerator.cs b/Engine/Source/Programs/UnrealBuildTool/Platform/Linux/LinuxProjectGenerator.cs
--- a/Engine/Source/Programs/UnrealBuildTool/Platform/Linux/LinuxProjectGenerator.cs
+++ b/Engine/Source/Programs/UnrealBuildTool/Platform/Linux/LinuxProjectGenerator.cs
@@ -70,6 +70,7 @@ namespace UnrealBuildTool
 				string ClangIncludeDirectory = Path.Combine(InternalSdkPath, "lib/clang/" + Version + "/include/");
 
 				Result.Add(Path.Combine(InternalSdkPath, "include"));
+				Result.Add(Path.Combine(InternalSdkPath, "include/c++/v1"));
 				Result.Add(Path.Combine(InternalSdkPath, "usr/include"));
 				Result.Add(ClangIncludeDirectory);
 				if (!Directory.Exists(ClangIncludeDirectory))
