[macOS/iOS] Adding additional compiler arguments to the portable toolchain path, to better match the default toolchain output.

#Synced-CL 50951058

diff --git a/Engine/Source/Programs/UnrealBuildTool/Platform/Apple/AppleToolChain.cs b/Engine/Source/Programs/UnrealBuildTool/Platform/Apple/AppleToolChain.cs
--- a/Engine/Source/Programs/UnrealBuildTool/Platform/Apple/AppleToolChain.cs
+++ b/Engine/Source/Programs/UnrealBuildTool/Platform/Apple/AppleToolChain.cs
@@ -626,8 +626,6 @@ namespace UnrealBuildTool
 				return;
 			}
 
-			Arguments.Add("-Xclang -fno-validate-pch"); // Needed since pch is actually portable but there is a check in clang that prevents it
-			
 			Arguments.Add("-nostdinc");
 
 			// FreeBSD math.h compatibility
@@ -636,7 +634,11 @@ namespace UnrealBuildTool
 			Arguments.Add("-D__double_t=double");
 			
 			VersionNumber AppleClangVersion = ((AppleToolChainInfo)Info).AppleClangVersion;
-			Arguments.Add($"-fclang-abi-compat={AppleClangVersion.Components[0]}");
+
+			// Apple clang generates a binary with ABI where the move constructors are not considered when determining if a type is trivially reallocable
+			// To keep the portable tool chain and default toolchain binaries compatible we need to force the ABI compat version to 4
+			// @note: A longer term fix would be patch our version of clang
+			Arguments.Add("-fclang-abi-compat=4");
 
 			Arguments.Add($"-D{XnuPlatformDefine}");
 
@@ -659,6 +661,33 @@ namespace UnrealBuildTool
 			{
 				Logger.LogError("Clang builtin headers not found");
 			}
+			
+			if (CompileEnvironment.Architectures.Contains(UnrealArch.Arm64))
+			{
+				// Portable clang defaults to ARMv8.4-A without BTI among other differences.
+				// Therefore, we need to set a few flags to get a better matching binary
+				
+				// Enable ARMv8.3 RCPC instructions (ldapr) instead of ARMv8.4 (ldapur)
+				Arguments.Add("-march=armv8.5-a+rcpc");
+				
+				// Disable ARMv8.4 RCPC "immediate offset" variants (ldapurb, stlurb) that Apple clang doesn't use
+				Arguments.Add("-Xclang -target-feature -Xclang -rcpc-immo");
+				
+				// Disable PAC/BTI instructions to match Apple clang
+				Arguments.Add("-mbranch-protection=none");
+				
+				// Vanilla clang inlines everything by default, so try to match apple clang behavior
+				if (CompileEnvironment.bPGOProfile)
+				{
+					Arguments.Add("-fbasic-block-sections=all");
+				}
+				else
+				{
+					Arguments.Add("-mllvm -hot-cold-split");
+				}
+				
+				Arguments.Add("-mllvm -inline-threshold=225");
+			}
 
 			// Apple OSS headers
 			Arguments.Add($"-isystem \"{NormalizeCommandLinePath(OSSHeadersDir,CompileEnvironment.RootPaths)}\"");
