Fixing Static Analysis warning in UBT.

My previous CL changed a compiler argument that depended on a local variable, but that local variable was not used anywhere else, therefore stayed as unused, which generates a warning.

I removed the unused local varibale.

#Synced-CL 50951058

diff --git a/Engine/Source/Programs/UnrealBuildTool/Platform/Apple/AppleToolChain.cs b/Engine/Source/Programs/UnrealBuildTool/Platform/Apple/AppleToolChain.cs
--- a/Engine/Source/Programs/UnrealBuildTool/Platform/Apple/AppleToolChain.cs
+++ b/Engine/Source/Programs/UnrealBuildTool/Platform/Apple/AppleToolChain.cs
@@ -632,8 +632,6 @@ namespace UnrealBuildTool
 			Arguments.Add("-D__ISO_C_VISIBLE=1999");
 			Arguments.Add("-D__float_t=float");
 			Arguments.Add("-D__double_t=double");
-			
-			VersionNumber AppleClangVersion = ((AppleToolChainInfo)Info).AppleClangVersion;
 
 			// Apple clang generates a binary with ABI where the move constructors are not considered when determining if a type is trivially reallocable
 			// To keep the portable tool chain and default toolchain binaries compatible we need to force the ABI compat version to 4
