Fix native CookedEditor build on Linux.

#Synced-CL 50902315

diff --git a/Engine/Source/Developer/CookedEditor/CookedEditor.Build.cs b/Engine/Source/Developer/CookedEditor/CookedEditor.Build.cs
--- a/Engine/Source/Developer/CookedEditor/CookedEditor.Build.cs
+++ b/Engine/Source/Developer/CookedEditor/CookedEditor.Build.cs
@@ -29,7 +29,7 @@ public class CookedEditor : ModuleRules
 				"TargetPlatform",
 			});
 
-		if (IsPlatformAvailable(UnrealTargetPlatform.Win64))
+		if (IsPlatformAvailable(UnrealTargetPlatform.Win64, bIgnoreSDKCheck: true))
 		{
 			PublicDefinitions.Add("COOKEDEDITOR_WITH_WINDOWSTARGETPLATFORM=1");
 			PublicIncludePathModuleNames.Add("WindowsTargetPlatformSettings");
diff --git a/Engine/Source/Programs/UnrealBuildTool/Platform/Linux/LinuxToolChain.cs b/Engine/Source/Programs/UnrealBuildTool/Platform/Linux/LinuxToolChain.cs
--- a/Engine/Source/Programs/UnrealBuildTool/Platform/Linux/LinuxToolChain.cs
+++ b/Engine/Source/Programs/UnrealBuildTool/Platform/Linux/LinuxToolChain.cs
@@ -709,6 +709,7 @@ namespace UnrealBuildTool
 			Arguments.Add("-Wl,-rpath=${ORIGIN}");
 			Arguments.Add("-Wl,-rpath-link=${ORIGIN}");
 			Arguments.Add("-Wl,-rpath=${ORIGIN}/..");   // for modules that are in sub-folders of the main Engine/Binary/Linux folder
+			Arguments.Add("-Wl,-rpath=${ORIGIN}/NotForLicensees");   // a common subfolder
 			if (LinkEnvironment.Architecture == UnrealArch.X64)
 			{
 				Arguments.Add("-Wl,-rpath=${ORIGIN}/../../../Engine/Binaries/ThirdParty/Qualcomm/Linux");
