[IOS] PGO build & data

Gathered locally from a 50759635, using MB3 allocator.

The actualy automated gathering will be submitted separately to a later stream to keep these changes minimal.

#rnx

diff --git a/Engine/Source/Programs/UnrealBuildTool/Platform/IOS/UEBuildIOS.cs b/Engine/Source/Programs/UnrealBuildTool/Platform/IOS/UEBuildIOS.cs
--- a/Engine/Source/Programs/UnrealBuildTool/Platform/IOS/UEBuildIOS.cs
+++ b/Engine/Source/Programs/UnrealBuildTool/Platform/IOS/UEBuildIOS.cs
@@ -1124,6 +1124,20 @@ namespace UnrealBuildTool
 				CompileEnvironment.Definitions.Add("ENABLE_ADVERTISING_IDENTIFIER=1");
 			}
 
+			if (Target.bPGOOptimize || Target.bPGOProfile)
+			{
+				Logger.LogInformation("PGO {PgoType} build", Target.bPGOOptimize ? "optimize" : "profile");
+
+				CompileEnvironment.PGODirectory = DirectoryReference.Combine(Target.ProjectFile?.Directory ?? Unreal.WritableEngineDirectory, "Platforms", "IOS", "Build", "PGO").FullName;
+				CompileEnvironment.PGOFilenamePrefix = String.Format("{0}-IOS.profdata", Target.Name);
+
+				LinkEnvironment.PGODirectory = CompileEnvironment.PGODirectory;
+				LinkEnvironment.PGOFilenamePrefix = CompileEnvironment.PGOFilenamePrefix;
+
+				Logger.LogInformation("PGO Dir: {PgoDir}", CompileEnvironment.PGODirectory);
+				Logger.LogInformation("PGO Prefix: {PgoPrefix}", CompileEnvironment.PGOFilenamePrefix);
+			}
+
 			// if the project has an Oodle compression Dll, enable the decompressor on IOS
 			if (Target.ProjectFile != null)
 			{
