Removing more GDK references from the engine code
- remove code that removed "GDK" from various restricted folder lists
- GDK module binaries no longer going to a GDK subfolder
- Windows DDPI.ini no longer needs to list GDK

#rnx

#Synced-CL 51007224

diff --git a/Engine/Config/Windows/DataDrivenPlatformInfo.ini b/Engine/Config/Windows/DataDrivenPlatformInfo.ini
--- a/Engine/Config/Windows/DataDrivenPlatformInfo.ini
+++ b/Engine/Config/Windows/DataDrivenPlatformInfo.ini
@@ -2,7 +2,6 @@
 bIsConfidential=false
 IniParent=VulkanPC
 PlatformNameAliases=Win64
-AdditionalRestrictedFolders=GDK
 TargetSettingsIniSectionName=/Script/WindowsTargetPlatform.WindowsTargetSettings
 Freezing_b32Bit=false
 Freezing_MaxFieldAlignment32=4
diff --git a/Engine/Plugins/Runtime/Microsoft/GDKPackageChunkInstall/Source/GDKPackageChunkInstall/GDKPackageChunkInstall.Build.cs b/Engine/Plugins/Runtime/Microsoft/GDKPackageChunkInstall/Source/GDKPackageChunkInstall/GDKPackageChunkInstall.Build.cs
--- a/Engine/Plugins/Runtime/Microsoft/GDKPackageChunkInstall/Source/GDKPackageChunkInstall/GDKPackageChunkInstall.Build.cs
+++ b/Engine/Plugins/Runtime/Microsoft/GDKPackageChunkInstall/Source/GDKPackageChunkInstall/GDKPackageChunkInstall.Build.cs
@@ -6,11 +6,6 @@ public class GDKPackageChunkInstall : ModuleRules
 {
     public GDKPackageChunkInstall(ReadOnlyTargetRules Target) : base(Target)
     {
-		if (Target.Platform == UnrealTargetPlatform.Win64)
-		{
-			BinariesSubFolder = "GDK";
-		}
-
         PrivateDependencyModuleNames.AddRange(
             new string[] 
             {
diff --git a/Engine/Source/Developer/Microsoft/GDKPlatformEditor/GDKPlatformEditor.Build.cs b/Engine/Source/Developer/Microsoft/GDKPlatformEditor/GDKPlatformEditor.Build.cs
--- a/Engine/Source/Developer/Microsoft/GDKPlatformEditor/GDKPlatformEditor.Build.cs
+++ b/Engine/Source/Developer/Microsoft/GDKPlatformEditor/GDKPlatformEditor.Build.cs
@@ -6,8 +6,6 @@ public class GDKPlatformEditor : ModuleRules
 {
 	public GDKPlatformEditor(ReadOnlyTargetRules Target) : base(Target)
 	{
-		BinariesSubFolder = "GDK";
-
 		PrivateDependencyModuleNames.AddRange(
 			new string[] {
 				"Core",
diff --git a/Engine/Source/Programs/AutomationTool/Win/MSGaming/MSGameStoreCustomDeploymentHandler.cs b/Engine/Source/Programs/AutomationTool/Win/MSGaming/MSGameStoreCustomDeploymentHandler.cs
--- a/Engine/Source/Programs/AutomationTool/Win/MSGaming/MSGameStoreCustomDeploymentHandler.cs
+++ b/Engine/Source/Programs/AutomationTool/Win/MSGaming/MSGameStoreCustomDeploymentHandler.cs
@@ -268,8 +268,6 @@ public class MSGameStoreCustomDeploymentHandler : CustomDeploymentHandler
 
 	public override bool PreGetFilesToDeployOrStage(ProjectParams Params, DeploymentContext SC)
 	{
-		SC.RestrictedFolderNames.RemoveWhere(x => x == "GDK");
-
 		// copy all custom install actions to the correct folder if we are not using package precaching
 		MSGameStorePackageGenerator PackageGenerator = new(SC.RawProjectPath, SC.StageTargetPlatform.PlatformType, IniSection_TargetSettings, SC.CustomConfig, SC.StageTargetPlatform.CrashReportPlatform);
 		if (!PackageGenerator.PrecachePackageLayout)
diff --git a/Engine/Source/Programs/AutomationTool/Win/MSGaming/MSGameStoreCustomStagingHandler.cs b/Engine/Source/Programs/AutomationTool/Win/MSGaming/MSGameStoreCustomStagingHandler.cs
--- a/Engine/Source/Programs/AutomationTool/Win/MSGaming/MSGameStoreCustomStagingHandler.cs
+++ b/Engine/Source/Programs/AutomationTool/Win/MSGaming/MSGameStoreCustomStagingHandler.cs
@@ -65,9 +65,6 @@ public class MSGameStoreCustomStagingHandler : CustomStagingHandler
 			return false;
 		}
 
-		// don't restrict items from the GDK folder - necessary to allow PC GDK plugins to be included in the project while PC GDK is still under the Platforms folder
-		SC.RestrictedFolderNames.RemoveWhere(x => x == "GDK");
-
 		return true;
 	}
 
diff --git a/Engine/Source/Programs/AutomationTool/Win/MSGaming/MSGamingRuntimeCustomStagingHandler.cs b/Engine/Source/Programs/AutomationTool/Win/MSGaming/MSGamingRuntimeCustomStagingHandler.cs
--- a/Engine/Source/Programs/AutomationTool/Win/MSGaming/MSGamingRuntimeCustomStagingHandler.cs
+++ b/Engine/Source/Programs/AutomationTool/Win/MSGaming/MSGamingRuntimeCustomStagingHandler.cs
@@ -60,9 +60,6 @@ public class MSGamingRuntimeCustomStagingHandler : CustomStagingHandler
 			return false;
 		}
 
-		// don't restrict items from the GDK folder - necessary to allow PC GDK plugins to be included in the project while PC GDK is still under the Platforms folder
-		SC.RestrictedFolderNames.RemoveWhere(x => x == "GDK");
-
 		return true;
 	}
 
diff --git a/Engine/Source/Programs/UnrealBuildTool/Configuration/UEBuildModule.cs b/Engine/Source/Programs/UnrealBuildTool/Configuration/UEBuildModule.cs
--- a/Engine/Source/Programs/UnrealBuildTool/Configuration/UEBuildModule.cs
+++ b/Engine/Source/Programs/UnrealBuildTool/Configuration/UEBuildModule.cs
@@ -545,15 +545,6 @@ namespace UnrealBuildTool
 				ReferencedDirs.ExceptWith(PublicDependencyModules!.SelectMany(x => x.RestrictedFoldersAllowList));
 				ReferencedDirs.ExceptWith(PrivateDependencyModules!.SelectMany(x => x.RestrictedFoldersAllowList));
 
-				// ... temporary workaround until all GDK code has moved to the public engine. should be gone before 5.8 is released #jira UE-250237
-				if (Rules.Target.Platform == UnrealTargetPlatform.Win64)
-				{
-					// allow anything that lives in a folder named 'GDK'
-					ReferencedDirs.RemoveWhere( x => x.FullName.Split([System.IO.Path.DirectorySeparatorChar, System.IO.Path.AltDirectorySeparatorChar])
-						.Any( f => f.Equals("GDK", StringComparison.OrdinalIgnoreCase) ));
-				}
-				// ....
-
 				// Add flags for each of them
 				foreach (DirectoryReference ReferencedDir in ReferencedDirs)
 				{
diff --git a/Engine/Source/Runtime/Microsoft/GDKPackageManifest/GDKPackageManifest.Build.cs b/Engine/Source/Runtime/Microsoft/GDKPackageManifest/GDKPackageManifest.Build.cs
--- a/Engine/Source/Runtime/Microsoft/GDKPackageManifest/GDKPackageManifest.Build.cs
+++ b/Engine/Source/Runtime/Microsoft/GDKPackageManifest/GDKPackageManifest.Build.cs
@@ -6,11 +6,6 @@ public class GDKPackageManifest : ModuleRules
 {
 	public GDKPackageManifest( ReadOnlyTargetRules Target) : base(Target)
 	{
-		if (Target.Platform == UnrealTargetPlatform.Win64)
-		{
-			BinariesSubFolder = "GDK";
-		}
-
 		PrivateDependencyModuleNames.AddRange(
 			new string[]
 			{
diff --git a/Engine/Source/Runtime/Microsoft/GDKRuntime/GDKRuntime.Build.cs b/Engine/Source/Runtime/Microsoft/GDKRuntime/GDKRuntime.Build.cs
--- a/Engine/Source/Runtime/Microsoft/GDKRuntime/GDKRuntime.Build.cs
+++ b/Engine/Source/Runtime/Microsoft/GDKRuntime/GDKRuntime.Build.cs
@@ -1,14 +1,11 @@
 // Copyright Epic Games, Inc. All Rights Reserved.
 
 using UnrealBuildTool;
-using System.IO;
 
 public class GDKRuntime : ModuleRules
 {
 	public GDKRuntime(ReadOnlyTargetRules Target) : base(Target)
 	{
-		BinariesSubFolder = "GDK";
-
 		PrivateDependencyModuleNames.Add("Core");
 		PrivateDependencyModuleNames.Add("GRDK");
 	}
diff --git a/Engine/Source/Runtime/Microsoft/GDKSaveGameSystem/GDKSaveGameSystem.Build.cs b/Engine/Source/Runtime/Microsoft/GDKSaveGameSystem/GDKSaveGameSystem.Build.cs
--- a/Engine/Source/Runtime/Microsoft/GDKSaveGameSystem/GDKSaveGameSystem.Build.cs
+++ b/Engine/Source/Runtime/Microsoft/GDKSaveGameSystem/GDKSaveGameSystem.Build.cs
@@ -6,11 +6,6 @@ public class GDKSaveGameSystem : ModuleRules
 {
 	public GDKSaveGameSystem( ReadOnlyTargetRules Target) : base(Target)
 	{
-		if (Target.Platform == UnrealTargetPlatform.Win64)
-		{
-			BinariesSubFolder = "GDK";
-		}
-
 		PrivateDependencyModuleNames.AddRange(new string[]
 			{
 				"Core",
diff --git a/Engine/Source/Runtime/Windows/MSGaming/MSGameStorePlatformFeatures/MSGameStorePlatformFeatures.Build.cs b/Engine/Source/Runtime/Windows/MSGaming/MSGameStorePlatformFeatures/MSGameStorePlatformFeatures.Build.cs
--- a/Engine/Source/Runtime/Windows/MSGaming/MSGameStorePlatformFeatures/MSGameStorePlatformFeatures.Build.cs
+++ b/Engine/Source/Runtime/Windows/MSGaming/MSGameStorePlatformFeatures/MSGameStorePlatformFeatures.Build.cs
@@ -7,8 +7,6 @@ public class MSGameStorePlatformFeatures : ModuleRules
 {
 	public MSGameStorePlatformFeatures( ReadOnlyTargetRules Target) : base(Target)
 	{
-		BinariesSubFolder = "GDK";
-
 		PrivateDependencyModuleNames.AddRange(new string[]
 			{
 				"Core",
