UnrealBuildTool: Fix a concurrency bug in project file generation.

#Synced-CL 51090839

diff --git a/Engine/Source/Programs/UnrealBuildTool/ProjectFiles/ProjectFileGenerator.cs b/Engine/Source/Programs/UnrealBuildTool/ProjectFiles/ProjectFileGenerator.cs
--- a/Engine/Source/Programs/UnrealBuildTool/ProjectFiles/ProjectFileGenerator.cs
+++ b/Engine/Source/Programs/UnrealBuildTool/ProjectFiles/ProjectFileGenerator.cs
@@ -776,7 +776,12 @@ namespace UnrealBuildTool
 										string[] Directories = ProjectRelativePath.Split(Path.DirectorySeparatorChar);
 										foreach (string IntermediateDir in Directories)
 										{
-											ProjectFolderToAddTo = ProjectFolderToAddTo.AddSubFolder(IntermediateDir);
+											// To be clear, we are intentionally locking on the original value.
+											PrimaryProjectFolder ProjectFolderToLock = ProjectFolderToAddTo;
+											lock (ProjectFolderToLock)
+											{
+												ProjectFolderToAddTo = ProjectFolderToLock.AddSubFolder(IntermediateDir);
+											}
 										}
 									}
 								}
