Make AppleCertificateVerification fail the interfering check only if a matching codesigining identity is found

#Synced-CL 50908424

diff --git a/Engine/Source/Programs/UnrealBuildTool/Platform/Apple/AppleCertificateVerification.cs b/Engine/Source/Programs/UnrealBuildTool/Platform/Apple/AppleCertificateVerification.cs
--- a/Engine/Source/Programs/UnrealBuildTool/Platform/Apple/AppleCertificateVerification.cs
+++ b/Engine/Source/Programs/UnrealBuildTool/Platform/Apple/AppleCertificateVerification.cs
@@ -50,6 +50,7 @@ namespace UnrealBuildTool
 	    {
 	    	Regex CertOrgRegex = new Regex(@"OU=([^,]+)", RegexOptions.IgnoreCase);
 			Regex CertCommonNameRegex = new Regex(@"CN=([^,]+)", RegexOptions.IgnoreCase);
+			HashSet<string> IdentityHashes = AppleCodeSign.GetCodeSigningIdentityHashes();
 
 			foreach (X509Certificate2 FoundCert in CertCollection)
 			{
@@ -67,9 +68,14 @@ namespace UnrealBuildTool
 				if (CertCommonNameVal.StartsWith("Apple Development: ", StringComparison.Ordinal) && CodeSigningTeamID == CertOrgVal &&
 				    !CertOfTruthHashes.Contains(FoundCert.GetCertHashString()))
 				{
-					Logger.LogError("Found interfering cert identity: {FoundCert.Subject}!", FoundCert.Subject);
-				
-					return true;
+					// we found a possible interfering cert, check that that cert has a valid identity it can sign with
+					// otherwise it won't interfere
+					if(IdentityHashes.Contains(FoundCert.GetCertHashString()))
+					{
+						Logger.LogError("Found interfering cert identity: {FoundCert.Subject}!", FoundCert.Subject);
+
+						return true;
+					}
 				}
 			}
 
diff --git a/Engine/Source/Programs/UnrealBuildTool/Platform/IOS/IOSCodeSigning.cs b/Engine/Source/Programs/UnrealBuildTool/Platform/IOS/IOSCodeSigning.cs
--- a/Engine/Source/Programs/UnrealBuildTool/Platform/IOS/IOSCodeSigning.cs
+++ b/Engine/Source/Programs/UnrealBuildTool/Platform/IOS/IOSCodeSigning.cs
@@ -8,6 +8,7 @@ using System.IO;
 using System.Linq;
 using System.Security.Cryptography.X509Certificates;
 using System.Text;
+using System.Text.RegularExpressions;
 using System.Xml;
 using EpicGames.Core;
 using Microsoft.Extensions.Logging;
@@ -84,6 +85,38 @@ namespace UnrealBuildTool
 			}
 		}
 
+		/// <summary>
+		/// Finds all identites that could be used for apple code signing and gets their certificate hashes
+		/// </summary>
+		public static HashSet<string> GetCodeSigningIdentityHashes()
+		{
+			CertToolData = "";
+			HashSet<string> Identities = new HashSet<string>();
+			Process CertTool = new Process();
+			CertTool.StartInfo.FileName = "/usr/bin/security";
+			CertTool.StartInfo.UseShellExecute = false;
+			CertTool.StartInfo.Arguments = String.Format("find-identity -v -p codesigning");
+			CertTool.StartInfo.RedirectStandardOutput = true;
+			CertTool.StartInfo.StandardOutputEncoding = Encoding.UTF8;
+			CertTool.OutputDataReceived += new DataReceivedEventHandler(OutputReceivedCertToolProcessCall);
+			CertTool.Start();
+			CertTool.BeginOutputReadLine();
+			CertTool.WaitForExit();
+
+			if (CertTool.ExitCode == 0)
+			{
+				const string IdentityHashPattern = @"^\s*\d+\)\s+([A-F0-9]{40})\s+";
+				Identities = 
+					Regex.Matches(CertToolData, IdentityHashPattern, RegexOptions.Multiline).Cast<Match>()
+					.Select(Match => Match.Groups[1].Value)
+					.ToHashSet();
+			}
+
+			CertTool.Dispose();
+
+			return Identities;
+		}
+
 		/// <summary>
 		/// Finds all installed certificates with known Apple development/distro prefixes as a X509Certificate2Collection
 		/// </summary>
