Add ability to enable /Zs from buildconfig or commandline

#Synced-CL 51101942

diff --git a/Engine/Source/Programs/UnrealBuildTool/Platform/Windows/UEBuildWindows.cs b/Engine/Source/Programs/UnrealBuildTool/Platform/Windows/UEBuildWindows.cs
--- a/Engine/Source/Programs/UnrealBuildTool/Platform/Windows/UEBuildWindows.cs
+++ b/Engine/Source/Programs/UnrealBuildTool/Platform/Windows/UEBuildWindows.cs
@@ -802,6 +802,13 @@ namespace UnrealBuildTool
 		/// </summary>
 		public bool bOptimizeGlobalData = true;
 
+		/// <summary>
+		/// Tells the compiler to check only the syntax of the source files on the command line. (/Zs)
+		/// </summary>
+		[XmlConfigFile(Category = "WindowsPlatform")]
+		[CommandLine("-SyntaxCheckOnly")]
+		public bool bSyntaxCheckOnly = false;
+
 		/// <summary>
 		/// Whether to reduce optimizations for huge functions over an instruction threshold to improve compile time
 		/// https://devblogs.microsoft.com/cppblog/msvc-backend-updates-in-visual-studio-2019-versions-16-3-and-16-4/
@@ -1101,6 +1108,8 @@ namespace UnrealBuildTool
 
 		public bool bOptimizeGlobalData => Inner.bOptimizeGlobalData;
 
+		public bool bSyntaxCheckOnly => Inner.bSyntaxCheckOnly;
+
 		public bool bReducedOptimizeHugeFunctions => Inner.bReducedOptimizeHugeFunctions;
 
 		public int ReducedOptimizeHugeFunctionsThreshold => Inner.ReducedOptimizeHugeFunctionsThreshold;
diff --git a/Engine/Source/Programs/UnrealBuildTool/Platform/Windows/VCToolChain.cs b/Engine/Source/Programs/UnrealBuildTool/Platform/Windows/VCToolChain.cs
--- a/Engine/Source/Programs/UnrealBuildTool/Platform/Windows/VCToolChain.cs
+++ b/Engine/Source/Programs/UnrealBuildTool/Platform/Windows/VCToolChain.cs
@@ -679,6 +679,12 @@ namespace UnrealBuildTool
 				Arguments.Add("/Gw");
 			}
 
+			// Enable syntax check only mode
+			if (Target.WindowsPlatform.bSyntaxCheckOnly)
+			{
+				Arguments.Add("/Zs");
+			}
+
 			// Reduce optimizations for huge functions, may improve compile time a the expense of speed for functions over the threshold
 			if (Target.WindowsPlatform.Compiler.IsMSVC() && Target.WindowsPlatform.bReducedOptimizeHugeFunctions)
 			{
