Fix UBT generating vcxproj with the wrong NMakeOutput for targets with overridden names.
#Synced-CL 51155616

#Synced-CL 51155616

diff --git a/Engine/Source/Programs/UnrealBuildTool/ProjectFiles/VisualStudio/VCProject.cs b/Engine/Source/Programs/UnrealBuildTool/ProjectFiles/VisualStudio/VCProject.cs
--- a/Engine/Source/Programs/UnrealBuildTool/ProjectFiles/VisualStudio/VCProject.cs
+++ b/Engine/Source/Programs/UnrealBuildTool/ProjectFiles/VisualStudio/VCProject.cs
@@ -2410,7 +2410,6 @@ namespace UnrealBuildTool
 						TargetRulesObject = Combination.ProjectTarget!.TargetRules!;
 					}
 					FileReference TargetFilePath = Combination.ProjectTarget.TargetFilePath;
-					string TargetName = TargetFilePath.GetFileNameWithoutAnyExtensions();
 					string UBTPlatformName = Platform.ToString();
 					string UBTConfigurationName = Configuration.ToString();
 					VSSettings VSSettings = new(Platform, Configuration, ProjectFileFormat, null);
@@ -2439,44 +2438,46 @@ namespace UnrealBuildTool
 						RootOutputDirectory = UEBuildTarget.GetOutputDirectoryForExecutable(Unreal.EngineDirectory, TargetRulesObject.File!);
 					}
 
-					// Get the output directory
-					DirectoryReference OutputDirectory = DirectoryReference.Combine(RootOutputDirectory, "Binaries", UBTPlatformName);
-
-					if (!String.IsNullOrEmpty(TargetRulesObject.ExeBinariesSubFolder))
+					if (TargetRulesObject.OutputFile != null)
 					{
-						OutputDirectory = DirectoryReference.Combine(OutputDirectory, TargetRulesObject.ExeBinariesSubFolder);
+						NMakePath = FileReference.Combine(RootOutputDirectory, TargetRulesObject.OutputFile);
 					}
-
-					// Get the executable name (minus any platform or config suffixes)
-					string BaseExeName = TargetName;
-					if (!bShouldCompileMonolithic && TargetRulesObject.Type != TargetType.Program && TargetRulesObject.BuildEnvironment != TargetBuildEnvironment.Unique)
+					else
 					{
-						BaseExeName = "Unreal" + TargetRulesObject.Type.ToString();
-					}
+						// Get the output directory
+						DirectoryReference OutputDirectory = DirectoryReference.Combine(RootOutputDirectory, "Binaries", UBTPlatformName);
 
-					// Make the output file path
-					NMakePath = FileReference.Combine(OutputDirectory, BaseExeName);
-					if (Configuration != TargetRulesObject.UndecoratedConfiguration)
-					{
-						NMakePath += TargetRulesObject.DecoratedSeparator + UBTPlatformName + TargetRulesObject.DecoratedSeparator + UBTConfigurationName;
-					}
-					if (UnrealArchitectureConfig.ForPlatform(Platform).RequiresArchitectureFilenames(TargetRulesObject.Architectures))
-					{
-						NMakePath += TargetRulesObject.Architecture.ToString();
-					}
-					else if (Combination.Architecture != null) // support the case where the project/platform combination explicitly sets an architecture(e.g. Win64)
-					{
-						UnrealArch Architecture = (UnrealArch)Combination.Architecture;
-						if (UnrealArchitectureConfig.ForPlatform(Platform).RequiresArchitectureFilenames(new UnrealArchitectures(Architecture)))
+						if (!String.IsNullOrEmpty(TargetRulesObject.ExeBinariesSubFolder))
 						{
-							NMakePath += Architecture.ToString();
+							OutputDirectory = DirectoryReference.Combine(OutputDirectory, TargetRulesObject.ExeBinariesSubFolder);
 						}
-					}
-					NMakePath += BuildPlatform.GetBinaryExtension(UEBuildBinaryType.Executable);
 
-					if (TargetRulesObject.OutputFile != null)
-					{
-						NMakePath = FileReference.Combine(RootOutputDirectory, TargetRulesObject.OutputFile);
+						// Get the executable name (minus any platform or config suffixes)
+						string BaseExeName = TargetRulesObject.Name;
+						if (!bShouldCompileMonolithic && TargetRulesObject.Type != TargetType.Program && TargetRulesObject.BuildEnvironment != TargetBuildEnvironment.Unique)
+						{
+							BaseExeName = "Unreal" + TargetRulesObject.Type.ToString();
+						}
+
+						// Make the output file path
+						NMakePath = FileReference.Combine(OutputDirectory, BaseExeName);
+						if (Configuration != TargetRulesObject.UndecoratedConfiguration)
+						{
+							NMakePath += TargetRulesObject.DecoratedSeparator + UBTPlatformName + TargetRulesObject.DecoratedSeparator + UBTConfigurationName;
+						}
+						if (UnrealArchitectureConfig.ForPlatform(Platform).RequiresArchitectureFilenames(TargetRulesObject.Architectures))
+						{
+							NMakePath += TargetRulesObject.Architecture.ToString();
+						}
+						else if (Combination.Architecture != null) // support the case where the project/platform combination explicitly sets an architecture(e.g. Win64)
+						{
+							UnrealArch Architecture = (UnrealArch)Combination.Architecture;
+							if (UnrealArchitectureConfig.ForPlatform(Platform).RequiresArchitectureFilenames(new UnrealArchitectures(Architecture)))
+							{
+								NMakePath += Architecture.ToString();
+							}
+						}
+						NMakePath += BuildPlatform.GetBinaryExtension(UEBuildBinaryType.Executable);
 					}
 
 					VCProjectFileContent.AppendLine("  <PropertyGroup {0}>", ConditionString);
