UnrealBuildTool: Add constants for tps files to indicate when they're missing or not required. Check for them in the Analyze mode.

#Synced-CL 51258774

diff --git a/Engine/Source/Programs/UnrealBuildTool/Configuration/Rules/ExternalModuleMetadata.cs b/Engine/Source/Programs/UnrealBuildTool/Configuration/Rules/ExternalModuleMetadata.cs
--- a/Engine/Source/Programs/UnrealBuildTool/Configuration/Rules/ExternalModuleMetadata.cs
+++ b/Engine/Source/Programs/UnrealBuildTool/Configuration/Rules/ExternalModuleMetadata.cs
@@ -21,7 +21,18 @@ public record ExternalModuleMetadata(
 	Uri Uri,
 	string Version,
 	string TpsFilePath
-);
+)
+{
+	/// <summary>
+	/// Used to denote a missing .tps file for third party software when one should exist.
+	/// </summary>
+	public const string TpsMissing = "<missing>";
+
+	/// <summary>
+	/// Used to denote a .tps file not being required - for instance when a first-party component is being linked externally.
+	/// </summary>
+	public const string TpsNotRequired = "<notrequired>";
+}
 
 internal record ProcessedExternalModuleMetadata(
 	Uri Uri,
diff --git a/Engine/Source/Programs/UnrealBuildTool/Modes/AnalyzeMode.cs b/Engine/Source/Programs/UnrealBuildTool/Modes/AnalyzeMode.cs
--- a/Engine/Source/Programs/UnrealBuildTool/Modes/AnalyzeMode.cs
+++ b/Engine/Source/Programs/UnrealBuildTool/Modes/AnalyzeMode.cs
@@ -270,7 +270,15 @@ namespace UnrealBuildTool.Modes
 					continue;
 				}
 
-				if (!FileReference.Exists(moduleInfo.ExternalMetadata.TpsFilePath))
+				if (String.Equals(moduleInfo.ExternalMetadata.TpsFilePath.GetFileName(), ExternalModuleMetadata.TpsMissing, StringComparison.OrdinalIgnoreCase))
+				{
+					logger.LogInformation("Module '{ModuleName}' is missing a TPS file.", moduleInfo.Module.Name);
+				}
+				else if (String.Equals(moduleInfo.ExternalMetadata.TpsFilePath.GetFileName(), ExternalModuleMetadata.TpsNotRequired, StringComparison.OrdinalIgnoreCase))
+				{
+					// Nothing to log here, it's declared as not required and that's fine.
+				}
+				else if (!FileReference.Exists(moduleInfo.ExternalMetadata.TpsFilePath))
 				{
 					logger.LogInformation("Module '{ModuleName}' points to TPS file '{TpsFile}' that does not exist.", moduleInfo.Module.Name, moduleInfo.ExternalMetadata.TpsFilePath);
 				}
