UnrealBuildTool: Filter out textually duplicate target arguments.
#Synced-CL 51258774

diff --git a/Engine/Source/Programs/UnrealBuildTool/Configuration/Descriptors/TargetDescriptor.cs b/Engine/Source/Programs/UnrealBuildTool/Configuration/Descriptors/TargetDescriptor.cs
--- a/Engine/Source/Programs/UnrealBuildTool/Configuration/Descriptors/TargetDescriptor.cs
+++ b/Engine/Source/Programs/UnrealBuildTool/Configuration/Descriptors/TargetDescriptor.cs
@@ -353,6 +353,11 @@ namespace UnrealBuildTool
 			List<string> Targets;
 			Arguments = Arguments.Remove("-Target=", out Targets);
 
+			// We don't have defined equality semantics to compare target descriptors, so catch any textual duplicates passed to us.
+			// This is a safeguard, as duplicate targets are not handled 100% correctly by the rest of UBT.
+			TargetLists = [.. TargetLists.Distinct(StringComparer.OrdinalIgnoreCase)];
+			Targets = [.. Targets.Distinct(StringComparer.OrdinalIgnoreCase)];
+
 			if (TargetLists.Count > 0 || Targets.Count > 0)
 			{
 				// Try to parse multiple arguments from a single command line
