UnrealBuildTool: Write the source mapping for debugging in VSCode project files as both `sourceFileMap` and `sourceMap` in the `launch.json` file.

#Synced-CL 51291646

#Synced-CL 51291646

diff --git a/Engine/Source/Programs/UnrealBuildTool/ProjectFiles/VisualStudioCode/VSCodeProjectFileGenerator.cs b/Engine/Source/Programs/UnrealBuildTool/ProjectFiles/VisualStudioCode/VSCodeProjectFileGenerator.cs
--- a/Engine/Source/Programs/UnrealBuildTool/ProjectFiles/VisualStudioCode/VSCodeProjectFileGenerator.cs
+++ b/Engine/Source/Programs/UnrealBuildTool/ProjectFiles/VisualStudioCode/VSCodeProjectFileGenerator.cs
@@ -1427,18 +1427,14 @@ namespace UnrealBuildTool
 					if (Target.TargetRules is not null)
 					{ 
 						Dictionary<DirectoryReference, DirectoryReference> SourceFileMap = GenerateSourceMap(Target.TargetRules, BuildProduct.UProjectFile);
-						if (SourceFileMap.Any())
+						if (SourceFileMap.Count != 0)
 						{
-							OutFile.BeginObject("sourceFileMap");
-							{
-								foreach (KeyValuePair<DirectoryReference, DirectoryReference> entry in SourceFileMap)
-								{
-									OutFile.AddField(
-										MakeUnquotedPathString(entry.Key, EPathType.Absolute),
-										MakeUnquotedPathString(entry.Value, EPathType.Absolute));
-								}
-							}
-							OutFile.EndObject();
+							// [REDACTED_PROFILE_URL]/vscode-lldb/blob/3c2eb665730175c3dacf624d760c22e8837f1d7b/README.md?plain=1#L236
+							// https://code.visualstudio.com/docs/cpp/launch-json-reference#_sourcefilemap
+							// The documentation disagrees on whether this should be `sourceMap` or `sourceFileMap`.
+							// From testing, lldb uses `sourceMap`, gdb uses 'sourceFileMap', so write both.
+							WriteSourceMap(OutFile, SourceFileMap, "sourceFileMap");
+							WriteSourceMap(OutFile, SourceFileMap, "sourceMap");
 						}
 					}
 				}
@@ -1446,6 +1442,20 @@ namespace UnrealBuildTool
 			}
 		}
 
+		private void WriteSourceMap(JsonFile OutFile, Dictionary<DirectoryReference, DirectoryReference> SourceFileMap, string FieldName)
+		{
+			OutFile.BeginObject(FieldName);
+
+			foreach (KeyValuePair<DirectoryReference, DirectoryReference> entry in SourceFileMap)
+			{
+				OutFile.AddField(
+					MakeUnquotedPathString(entry.Key, EPathType.Absolute),
+					MakeUnquotedPathString(entry.Value, EPathType.Absolute));
+			}
+
+			OutFile.EndObject();
+		}
+
 		private void WriteNativeLaunchConfig(ProjectData.Project InProject, JsonFile OutFile)
 		{
 			foreach (ProjectData.Target Target in InProject.Targets)
