UnrealBuildTool: Create a root OpenTelemetry trace for the run so that child traces will be parented correctly. Without this, the traces are incorrectly parented and their finish times are truncated as Horde incorrectly guesses their parentage.

#Synced-CL 51291646

diff --git a/Engine/Source/Programs/UnrealBuildTool/UnrealBuildTool.cs b/Engine/Source/Programs/UnrealBuildTool/UnrealBuildTool.cs
--- a/Engine/Source/Programs/UnrealBuildTool/UnrealBuildTool.cs
+++ b/Engine/Source/Programs/UnrealBuildTool/UnrealBuildTool.cs
@@ -12,6 +12,7 @@ using System.Threading.Tasks;
 using System.Xml.Linq;
 using EpicGames.Core;
 using Microsoft.Extensions.Logging;
+using OpenTracing;
 using OpenTracing.Util;
 using UnrealBuildBase;
 
@@ -181,7 +182,6 @@ namespace UnrealBuildTool
 			FileReference? RunFile = null;
 			DirectoryReference? TempDirectory = null;
 			GlobalSingleInstanceMutex? Mutex = null;
-			JsonTracer? Tracer = null;
 
 			ILogger Logger = Log.Logger;
 
@@ -193,6 +193,11 @@ namespace UnrealBuildTool
 				Console.CancelKeyPress += CancelKeyPressAsync;
 			}
 
+			JsonTracer? Tracer = JsonTracer.TryRegisterAsGlobalTracer();
+			// Start the UBT OTel scope as the outermost span, as early as possible.
+			// We'll have to manually dispose this as it doesn't fit nicely into the try/catch block below.
+			IScope ubtScope = GlobalTracer.Instance.BuildSpan("UnrealBuildTool").StartActive();
+
 			// If command line contains -Session it is a recursive action and we don't want to report this trace to the default trace channel
 			string? traceChannel = null;
 			string? traceName = null;
@@ -227,12 +232,13 @@ namespace UnrealBuildTool
 			using TimelineEventScope ubaTimelineEventScope = new((name) => ubaTrace.TaskBegin(name, ""), (id) => ubaTrace.TaskEnd(id));
 
 			bool bShowTimeline = false;
+
 			try
 			{
 				// Start capturing performance info
 				Timeline.Start();
 				using ITimelineEvent MainScope = Timeline.ScopeEvent("Main");
-				Tracer = JsonTracer.TryRegisterAsGlobalTracer();
+
 				ArgumentsArray = ArgumentsArray.Concat(GetExtraArgsFromEnvVar()).ToArray();
 
 				// Parse the command line arguments
@@ -546,6 +552,7 @@ namespace UnrealBuildTool
 				Trace.Close();
 
 				// Write any trace logs
+				ubtScope.Dispose();
 				Tracer?.Flush();
 
 				// Delete the ubt run file
