UnrealBuildTool: Load package list for a visual studio install

#rnx
#Synced-CL 51237423

diff --git a/Engine/Source/Programs/UnrealBuildTool/Platform/Windows/MicrosoftPlatformSDK.cs b/Engine/Source/Programs/UnrealBuildTool/Platform/Windows/MicrosoftPlatformSDK.cs
--- a/Engine/Source/Programs/UnrealBuildTool/Platform/Windows/MicrosoftPlatformSDK.cs
+++ b/Engine/Source/Programs/UnrealBuildTool/Platform/Windows/MicrosoftPlatformSDK.cs
@@ -1021,8 +1021,17 @@ namespace UnrealBuildTool
 							string ProductId = Instance.GetProduct().GetId();
 							bool bCommunity = ProductId.Equals("Microsoft.VisualStudio.Product.Community", StringComparison.Ordinal);
 
+							Dictionary<string, VisualStudioPackage> Packages = [];
+							foreach (ISetupPackageReference Package in Instance.GetPackages())
+							{
+								if (System.Version.TryParse(Package.GetVersion(), out Version? version))
+								{
+									Packages[Package.GetId()] = new(Package.GetId(), version, Package.GetIsExtension());
+								}
+							}
+
 							DirectoryReference BaseDir = new DirectoryReference(Instance.GetInstallationPath());
-							Installations.Add(new VisualStudioInstallation(Compiler, Version, BaseDir, bCommunity, ReleaseChannel));
+							Installations.Add(new VisualStudioInstallation(Compiler, Version, BaseDir, bCommunity, ReleaseChannel, Packages));
 
 							Logger.LogDebug("Found {VisualStudio} installation: {BaseDir} (Product={ProductId}, Version={Version})", Compiler, BaseDir, ProductId, Version);
 						}
diff --git a/Engine/Source/Programs/UnrealBuildTool/Platform/Windows/UEBuildWindows.cs b/Engine/Source/Programs/UnrealBuildTool/Platform/Windows/UEBuildWindows.cs
--- a/Engine/Source/Programs/UnrealBuildTool/Platform/Windows/UEBuildWindows.cs
+++ b/Engine/Source/Programs/UnrealBuildTool/Platform/Windows/UEBuildWindows.cs
@@ -1165,48 +1165,63 @@ namespace UnrealBuildTool
 		#endregion
 	}
 
+	/// <summary>
+	/// Information about a particular Visual Studio package
+	/// </summary>
+	[DebuggerDisplay("{Id},{Version}")]
+	record VisualStudioPackage(string id, Version version, bool isExtension)
+	{
+		/// <summary>
+		/// General package identifier
+		/// </summary>
+		public string Id { get; } = id;
+
+		/// <summary>
+		/// Version of the package
+		/// </summary>
+		public Version Version { get; } = version;
+
+		/// <summary>
+		/// Whether the package refers to an external extension
+		/// </summary>
+		public bool IsExtension { get; } = isExtension;
+	}
+
 	/// <summary>
 	/// Information about a particular Visual Studio installation
 	/// </summary>
 	[DebuggerDisplay("{BaseDir}")]
-	class VisualStudioInstallation
+	record VisualStudioInstallation(WindowsCompiler Compiler, VersionNumber Version, DirectoryReference BaseDir, bool bCommunity, WindowsCompilerChannel ReleaseChannel, IDictionary<string, VisualStudioPackage> Packages)
 	{
 		/// <summary>
 		/// Compiler type
 		/// </summary>
-		public WindowsCompiler Compiler { get; }
+		public WindowsCompiler Compiler { get; } = Compiler;
 
 		/// <summary>
 		/// Version number for this installation
 		/// </summary>
-		public VersionNumber Version { get; }
+		public VersionNumber Version { get; } = Version;
 
 		/// <summary>
 		/// Base directory for the installation
 		/// </summary>
-		public DirectoryReference BaseDir { get; }
+		public DirectoryReference BaseDir { get; } = BaseDir;
 
 		/// <summary>
 		/// Whether it's a community edition of Visual Studio.
 		/// </summary>
-		public bool bCommunity { get; }
+		public bool bCommunity { get; } = bCommunity;
 
 		/// <summary>
 		/// The release channel of this installation
 		/// </summary>
-		public WindowsCompilerChannel ReleaseChannel { get; }
+		public WindowsCompilerChannel ReleaseChannel { get; } = ReleaseChannel;
 
 		/// <summary>
-		/// Constructor
+		/// The packages added to this installation
 		/// </summary>
-		public VisualStudioInstallation(WindowsCompiler Compiler, VersionNumber Version, DirectoryReference BaseDir, bool bCommunity, WindowsCompilerChannel ReleaseChannel)
-		{
-			this.Compiler = Compiler;
-			this.Version = Version;
-			this.BaseDir = BaseDir;
-			this.bCommunity = bCommunity;
-			this.ReleaseChannel = ReleaseChannel;
-		}
+		public IDictionary<string, VisualStudioPackage> Packages { get; } = Packages;
 	}
 
 	class WindowsArchitectureConfig : UnrealArchitectureConfig
