Enable building universal apps on MacOS. DefaultArchitecture 'all' is now 'universal'.

#Synced-CL 51099106

diff --git a/Engine/Source/Programs/UnrealBuildTool/Platform/Mac/UEBuildMac.cs b/Engine/Source/Programs/UnrealBuildTool/Platform/Mac/UEBuildMac.cs
--- a/Engine/Source/Programs/UnrealBuildTool/Platform/Mac/UEBuildMac.cs
+++ b/Engine/Source/Programs/UnrealBuildTool/Platform/Mac/UEBuildMac.cs
@@ -213,8 +213,8 @@ namespace UnrealBuildTool
 			// choose a supported architecture(s) based on desired type
 			List<UnrealArch> Architectures = new();
 
-			// return all supported if getting supported, compiling for distribution, or we want active, and "all" is selected
-			if (bGetAllSupported || bIsDistributionMode || DefaultArchitecture.Equals("all", StringComparison.OrdinalIgnoreCase) ||
+			// return all supported if getting supported, compiling for distribution, or we want active, and "universal" is selected
+			if (bGetAllSupported || bIsDistributionMode || DefaultArchitecture.Contains("universal", StringComparison.OrdinalIgnoreCase) ||
 				(bIsBuildMachine && bBuildAllSupportedOnBuildMachine))
 			{
 				if (bSupportsArm64)
