UnrealBuildTool: Auto downgrade visual studio solution if certain packages are installed

#rnx

#Synced-CL 51357799

diff --git a/Engine/Source/Programs/UnrealBuildTool/Platform/Android/UEBuildAndroid.cs b/Engine/Source/Programs/UnrealBuildTool/Platform/Android/UEBuildAndroid.cs
--- a/Engine/Source/Programs/UnrealBuildTool/Platform/Android/UEBuildAndroid.cs
+++ b/Engine/Source/Programs/UnrealBuildTool/Platform/Android/UEBuildAndroid.cs
@@ -3,6 +3,7 @@
 using System;
 using System.Collections.Generic;
 using System.IO;
+using System.Linq;
 using EpicGames.Core;
 using Microsoft.Extensions.Logging;
 using UnrealBuildBase;
@@ -243,6 +244,25 @@ namespace UnrealBuildTool
 			return BuildHostPlatform.Current.Platform != UnrealTargetPlatform.Linux;
 		}
 
+		/// <inheritdoc/>
+		public override VCProjectFileFormat GetRequiredVisualStudioVersion()
+		{
+			IEnumerable<VisualStudioInstallation> Installations = MicrosoftPlatformSDK.FindVisualStudioInstallations(Logger);
+
+			bool vsiInstalled = Installations
+				.Where(x => x.Compiler == WindowsCompiler.VisualStudio2022)
+				.SelectMany(x => x.Packages.Where(x => x.Key.StartsWith("Component.GoogleAndroid", StringComparison.Ordinal)).Select(x => x.Value.Version))
+				.Any();
+
+			if (vsiInstalled)
+			{
+				// Currently unknown when 2026 will be supported and at which version
+				return VCProjectFileFormat.VisualStudio2022;
+			}
+
+			return VCProjectFileFormat.Default;
+		}
+
 		public override bool IsBuildProduct(string FileName, string[] NamePrefixes, string[] NameSuffixes)
 		{
 			return IsBuildProductWithArch(FileName, NamePrefixes, NameSuffixes, ".so")
diff --git a/Engine/Source/Programs/UnrealBuildTool/ProjectFiles/VisualStudio/VCProjectFileGenerator.cs b/Engine/Source/Programs/UnrealBuildTool/ProjectFiles/VisualStudio/VCProjectFileGenerator.cs
--- a/Engine/Source/Programs/UnrealBuildTool/ProjectFiles/VisualStudio/VCProjectFileGenerator.cs
+++ b/Engine/Source/Programs/UnrealBuildTool/ProjectFiles/VisualStudio/VCProjectFileGenerator.cs
@@ -380,7 +380,7 @@ namespace UnrealBuildTool
 				}
 				Settings.ProjectFileFormat = Format;
 
-				bool DowngradeAvailable = false; // ex: Installations.Any(x => x.Compiler == WindowsCompiler.VisualStudio2022);
+				bool DowngradeAvailable = Installations.Any(x => x.Compiler == WindowsCompiler.VisualStudio2022);
 
 				// Allow the SDKs to override
 				foreach (UnrealTargetPlatform SupportedPlatform in SupportedPlatforms)
@@ -398,11 +398,11 @@ namespace UnrealBuildTool
 								// Reduce the Visual Studio version to the max supported by each platform we plan to include.
 								if (Settings.ProjectFileFormat == VCProjectFileFormat.Default || ProposedFormat < Settings.ProjectFileFormat)
 								{
-									Logger.LogInformation("Available {SupportedPlatform} SDK does not support Visual Studio 2022.", SupportedPlatform);
-									Version Version = BuildPlatform.GetVersionRequiredForVisualStudio(VCProjectFileFormat.VisualStudio2022);
+									Logger.LogInformation("Available {SupportedPlatform} SDK does not support Visual Studio 2026.", SupportedPlatform);
+									Version Version = BuildPlatform.GetVersionRequiredForVisualStudio(VCProjectFileFormat.VisualStudio2026);
 									if (Version > new Version())
 									{
-										Logger.LogInformation("Please update {SupportedPlatform} SDK to {Version} if Visual Studio 2022 support is desired.", SupportedPlatform, Version);
+										Logger.LogInformation("Please update {SupportedPlatform} SDK to {Version} if Visual Studio 2026 support is desired.", SupportedPlatform, Version);
 									}
 									if (!DowngradeAvailable)
 									{
