Add support for arm64 to the game config generator, introduced in the October GDK
- base appx generator api updated

#rnx

#Synced-CL 51338024

diff --git a/Engine/Source/Programs/AutomationTool/Microsoft/GDKAutomationUtils.cs b/Engine/Source/Programs/AutomationTool/Microsoft/GDKAutomationUtils.cs
--- a/Engine/Source/Programs/AutomationTool/Microsoft/GDKAutomationUtils.cs
+++ b/Engine/Source/Programs/AutomationTool/Microsoft/GDKAutomationUtils.cs
@@ -20,8 +20,8 @@ namespace AutomationUtils.GDK
 			DirectoryReference ManifestTargetDirectory = DirectoryReference.Combine(SC.ProjectRoot, "Saved", SC.StageTargetPlatform.PlatformType.ToString(), "Manifest");
 			InternalUtils.SafeDeleteDirectory(ManifestTargetDirectory.FullName);
 
-			Dictionary<UnrealTargetConfiguration, string> ExecutablePairs = GetGameConfigExecutables(Params, SC, GameConfigPurpose.ForStaging, IniSection_TargetSettings);
-			GDKGameConfigGenerator.Create(SC.StageTargetPlatform.PlatformType, Logger, ManifestTargetDirectory, SC.ShortProjectName, Params.RawProjectPath, ExecutablePairs, DLCFile: Params.DLCFile, CustomConfig: SC.CustomConfig, NoBootstrapExe: Params.NoBootstrapExe);
+			List<AppXManifestExecutable> Executables = GetGameConfigExecutables(Params, SC, GameConfigPurpose.ForStaging, IniSection_TargetSettings);
+			GDKGameConfigGenerator.Create(SC.StageTargetPlatform.PlatformType, Logger, ManifestTargetDirectory, SC.ShortProjectName, Params.RawProjectPath, Executables, DLCFile: Params.DLCFile, CustomConfig: SC.CustomConfig, NoBootstrapExe: Params.NoBootstrapExe);
 
 			if (Params.ApplyIoStoreOnDemand)
 			{
@@ -36,7 +36,7 @@ namespace AutomationUtils.GDK
 			ForStaging,
 			ForPackaging,
 		}
-		public static Dictionary<UnrealTargetConfiguration, string> GetGameConfigExecutables(ProjectParams Params, DeploymentContext SC, GameConfigPurpose Purpose, string IniSection_TargetSettings, UnrealTargetConfiguration? OnlyConfig = null)
+		public static List<AppXManifestExecutable> GetGameConfigExecutables(ProjectParams Params, DeploymentContext SC, GameConfigPurpose Purpose, string IniSection_TargetSettings, UnrealTargetConfiguration? OnlyConfig = null)
 		{
 			bool bForStaging = (Purpose == GameConfigPurpose.ForStaging);
 
@@ -47,17 +47,14 @@ namespace AutomationUtils.GDK
 				throw new AutomationException($"Could not find receipt for config {OnlyConfig!.Value} in currently stage targets - has the game been built & staged for this configuration?");
 			}
 
-			Dictionary<UnrealTargetConfiguration, string> ExecutablePairs = FilteredStageTargets  // note: only takes the first staged file per configuration at the moment
-				.GroupBy( X => X.Receipt.Configuration )
-				.ToDictionary(X => X.Key, X => SC.GetStagedFileLocation(X.First().Receipt.Launch!).ToString());
-
+			// special case for the Windows bootstrapper - this is in the game root and will run the most appropriate x64/arm64 executable automatically
+			// there is one bootstrapper per configuration (per build target)
 			if (!Params.NoBootstrapExe && SC.StageTargetPlatform.PlatformType.IsInGroup(UnrealPlatformGroup.Windows))
 			{
-				// bootstrap naming logic from WinPlatform.Automation.cs -> GetFilesToDeployOrStage
 				Dictionary<UnrealTargetConfiguration, List<string>> PerConfigExecutables = GetPerConfigBootstrapExecutables(Params, SC);
 					
 				bool bAllExist = PerConfigExecutables.Any();
-				Dictionary<UnrealTargetConfiguration, string> BootstrapExecutablePairs = [];
+				List<AppXManifestExecutable> BootstrapExecutables = [];
 				foreach (KeyValuePair<UnrealTargetConfiguration, List<string>> ConfigExecutables in PerConfigExecutables)
 				{
 					if (OnlyConfig != null && OnlyConfig!.Value != ConfigExecutables.Key)
@@ -65,10 +62,9 @@ namespace AutomationUtils.GDK
 						continue;
 					}
 
-					if (ConfigExecutables.Value.Any())
+					foreach (string Executable in ConfigExecutables.Value)
 					{
-						string Executable = ConfigExecutables.Value.First();
-						BootstrapExecutablePairs.Add( ConfigExecutables.Key, Executable );
+						BootstrapExecutables.Add( new(ConfigExecutables.Key, Executable));
 
 						bAllExist &= bForStaging ? SC.FilesToStage.NonUFSSystemFiles.ContainsKey(new StagedFileReference(Executable)) : FileReference.Exists(FileReference.Combine(SC.StageDirectory, Executable));
 					}
@@ -77,24 +73,30 @@ namespace AutomationUtils.GDK
 				// only use the bootstrap executables if they exist (to allow for development iteration scenarios). Bootstrap must exist for distribution
 				if (Params.Distribution || bAllExist)
 				{
-					ExecutablePairs = BootstrapExecutablePairs;
+					return BootstrapExecutables;
 				}
 			}
-			else if (Params.NoBootstrapExe)
+
+
+			// build the list of executables, applying config remapping as necessary
+			List<AppXManifestExecutable> Executables = [];
+			ConfigHierarchy EngineIni = ConfigCache.ReadHierarchy(ConfigHierarchyType.Engine, DirectoryReference.FromFile(SC.RawProjectPath), SC.StageTargetPlatform.PlatformType, SC.CustomConfig);
+			foreach ( StageTarget Target in FilteredStageTargets)
 			{
-				ConfigHierarchy EngineIni = ConfigCache.ReadHierarchy(ConfigHierarchyType.Engine, DirectoryReference.FromFile(SC.RawProjectPath), SC.StageTargetPlatform.PlatformType, SC.CustomConfig);
+				UnrealTargetConfiguration Config = Target.Receipt.Configuration;
+				UnrealArch Arch = Target.Receipt.Architectures.SingleArchitecture;
+				string Executable = SC.GetStagedFileLocation(Target.Receipt.Launch!).ToString();
 
-				foreach (UnrealTargetConfiguration Config in ExecutablePairs.Keys)
+				string ConfigEntry = $"OverrideLaunchExe_{Config}_{Arch}";
+				if (EngineIni.GetString(IniSection_TargetSettings, ConfigEntry, out string OverrideExectuablePath) && !string.IsNullOrEmpty(OverrideExectuablePath))
 				{
-					string ConfigEntry = $"OverrideLaunchExe_{Config}_x64";
-					if (EngineIni.GetString(IniSection_TargetSettings, ConfigEntry, out string OverrideExectuablePath) && !string.IsNullOrEmpty(OverrideExectuablePath))
-					{
-						ExecutablePairs[Config] = OverrideExectuablePath;
-					}
+					Executable = OverrideExectuablePath;
 				}
+
+				Executables.Add( new(Config, Executable, Arch));
 			}
 
-			return ExecutablePairs;
+			return Executables;
 		}
 
 		public static Dictionary<UnrealTargetConfiguration, List<string>> GetPerConfigBootstrapExecutables(ProjectParams Params, DeploymentContext SC)
diff --git a/Engine/Source/Programs/AutomationTool/Microsoft/GDKPackageGenerator.cs b/Engine/Source/Programs/AutomationTool/Microsoft/GDKPackageGenerator.cs
--- a/Engine/Source/Programs/AutomationTool/Microsoft/GDKPackageGenerator.cs
+++ b/Engine/Source/Programs/AutomationTool/Microsoft/GDKPackageGenerator.cs
@@ -551,8 +551,8 @@ namespace AutomationUtils.GDK
 			FilesToInclude = FilesToInclude.Except(ManifestFiles);
 
 			// generate a GDK manifest that just contains the selected configurations
-			Dictionary<UnrealTargetConfiguration, string> ExecutablePairs = GDKAutomationUtils.GetGameConfigExecutables(Params, SC, GDKAutomationUtils.GameConfigPurpose.ForPackaging, IniSection_TargetSettings, OnlyConfig);
-			GDKGameConfigGenerator.Create(PlatformType, Logger, IntermediatePath, SC.ShortProjectName, Params.RawProjectPath, ExecutablePairs, bDeleteIntermediates: true, DLCFile: Params.DLCFile, CustomConfig: SC.CustomConfig, NoBootstrapExe: Params.NoBootstrapExe);
+			List<AppXManifestExecutable> Executables = GDKAutomationUtils.GetGameConfigExecutables(Params, SC, GDKAutomationUtils.GameConfigPurpose.ForPackaging, IniSection_TargetSettings, OnlyConfig);
+			GDKGameConfigGenerator.Create(PlatformType, Logger, IntermediatePath, SC.ShortProjectName, Params.RawProjectPath, Executables, bDeleteIntermediates: true, DLCFile: Params.DLCFile, CustomConfig: SC.CustomConfig, NoBootstrapExe: Params.NoBootstrapExe);
 			IEnumerable<StagedFileReference> GeneratedFilesToInclude =  GetGDKManifestFilesInDirectory(IntermediatePath)
 				.Select( File => StagedFileReference.Combine(File.MakeRelativeTo(IntermediatePath)));
 
@@ -1449,8 +1449,8 @@ namespace AutomationUtils.GDK
 				FilesToInclude = FilesToInclude.Except(ExistingStagedManifestFiles);
 
 				// generate the manifest and gather the files
-				Dictionary<UnrealTargetConfiguration, string> ExecutablePairs = GDKAutomationUtils.GetGameConfigExecutables(Params, SC, Purpose: GDKAutomationUtils.GameConfigPurpose.ForStaging, IniSection_TargetSettings, OnlyConfig);
-				GDKGameConfigGenerator.Create(PlatformType, Logger, PackageMetadataDir, SC.ShortProjectName, Params.RawProjectPath, ExecutablePairs, bDeleteIntermediates: true, DLCFile: Params.DLCFile, CustomConfig: SC.CustomConfig);
+				List<AppXManifestExecutable> Executables = GDKAutomationUtils.GetGameConfigExecutables(Params, SC, Purpose: GDKAutomationUtils.GameConfigPurpose.ForStaging, IniSection_TargetSettings, OnlyConfig);
+				GDKGameConfigGenerator.Create(PlatformType, Logger, PackageMetadataDir, SC.ShortProjectName, Params.RawProjectPath, Executables, bDeleteIntermediates: true, DLCFile: Params.DLCFile, CustomConfig: SC.CustomConfig);
 				GeneratedFilesToInclude = GeneratedFilesToInclude.Union( GetGDKManifestFilesInDirectory(PackageMetadataDir)
 					.Select(File => StagedFileReference.Combine(File.MakeRelativeTo(PackageMetadataDir)))
 				);
diff --git a/Engine/Source/Programs/AutomationTool/Win/MSGaming/MSGameStoreCustomDeploymentHandler.cs b/Engine/Source/Programs/AutomationTool/Win/MSGaming/MSGameStoreCustomDeploymentHandler.cs
--- a/Engine/Source/Programs/AutomationTool/Win/MSGaming/MSGameStoreCustomDeploymentHandler.cs
+++ b/Engine/Source/Programs/AutomationTool/Win/MSGaming/MSGameStoreCustomDeploymentHandler.cs
@@ -33,6 +33,11 @@ public class MSGameStoreCustomDeploymentHandler : CustomDeploymentHandler
 		{
 			throw new AutomationException("GDK is not installed");
 		}
+
+		if (GDKEdition < 251000 && Params.NoBootstrapExe && SC.StageTargets.Any(X => !X.Receipt.Architectures.SingleArchitecture.bIsX64))
+		{
+			throw new AutomationException("October 2025 GDK or higher is required to use the GDK with ARM64 unless you use the bootstrapper (omit -NoBootstrapExe)");
+		}
 	}
 
 	public override bool SupportsPlatform(UnrealTargetPlatform Platform)
@@ -215,8 +220,10 @@ public class MSGameStoreCustomDeploymentHandler : CustomDeploymentHandler
 		}
 
 		// find the AUMID associated with the given ClientApp
-		UnrealTargetConfiguration Configuration = SC.StageTargets.First(X => X.Receipt.Launch!.GetFileName() == Path.GetFileName(ClientApp)).Receipt.Configuration;
-		string AUMIDName = GDKGameConfigGenerator.GetAUMIDFromProject(Platform.PlatformType, Logger, Params.RawProjectPath.Directory, Configuration, Params.ShortProjectName, SC.CustomConfig);
+		TargetReceipt Receipt = SC.StageTargets.First(X => X.Receipt.Launch!.GetFileName() == Path.GetFileName(ClientApp)).Receipt;
+		UnrealTargetConfiguration Configuration = Receipt.Configuration;
+		UnrealArch? Architecture = Params.NoBootstrapExe ? Receipt.Architectures.SingleArchitecture : null;
+		string AUMIDName = GDKGameConfigGenerator.GetAUMIDFromProject(Platform.PlatformType, Logger, Params.RawProjectPath.Directory, Configuration, Architecture, Params.ShortProjectName, SC.CustomConfig);
 		string PackageFamilyName = AUMIDName.Substring(0, AUMIDName.IndexOf('!'));
 
 		// see if the application has been deployed (either a staged or packaged build - it doesn't matter at this stage)
@@ -258,8 +265,8 @@ public class MSGameStoreCustomDeploymentHandler : CustomDeploymentHandler
 
 				// create a placeholder MicrosoftGame.config
 				Logger.LogInformation("Game has not been staged or deployed. Creating placeholder MicrosoftGame.config in {BaseLaunchDir}", BaseLaunchDir);
-				Dictionary<UnrealTargetConfiguration, string> ExecutablePairs = new() { { Configuration, Path.GetFileName(ClientApp)! } };
-				GDKGameConfigGenerator.Create(SC.StageTargetPlatform.PlatformType, Logger, DirectoryReference.FromString(BaseLaunchDir), SC.ShortProjectName, Params.RawProjectPath, ExecutablePairs, CustomConfig: SC.CustomConfig);
+				List<AppXManifestExecutable> Executables = [ new( Configuration, Path.GetFileName(ClientApp)! ) ];
+				GDKGameConfigGenerator.Create(SC.StageTargetPlatform.PlatformType, Logger, DirectoryReference.FromString(BaseLaunchDir), SC.ShortProjectName, Params.RawProjectPath, Executables, CustomConfig: SC.CustomConfig);
 			}
 		}
 
diff --git a/Engine/Source/Programs/AutomationTool/Win/MSGaming/MSGameStoreCustomStagingHandler.cs b/Engine/Source/Programs/AutomationTool/Win/MSGaming/MSGameStoreCustomStagingHandler.cs
--- a/Engine/Source/Programs/AutomationTool/Win/MSGaming/MSGameStoreCustomStagingHandler.cs
+++ b/Engine/Source/Programs/AutomationTool/Win/MSGaming/MSGameStoreCustomStagingHandler.cs
@@ -65,6 +65,12 @@ public class MSGameStoreCustomStagingHandler : CustomStagingHandler
 			return false;
 		}
 
+		// sanity check
+		if (GDKEdition < 251000 && Params.NoBootstrapExe && SC.StageTargets.Any(X => X.Receipt != null && !X.Receipt.Architectures.SingleArchitecture.bIsX64))
+		{
+			throw new AutomationException("October 2025 GDK or higher is required to use the GDK with ARM64 unless you use the bootstrapper (omit -NoBootstrapExe)");
+		}
+
 		return true;
 	}
 
@@ -242,8 +248,10 @@ public class MSGameStoreCustomStagingHandler : CustomStagingHandler
 				VerifyWindowsEnvironment();
 
 				// find the AUMID associated with the given ClientApp
-				UnrealTargetConfiguration Configuration = SC.StageTargets.First(X => X.Receipt.Launch!.GetFileName() == Path.GetFileName(ClientApp)).Receipt.Configuration;
-				string AUMIDName = GDKGameConfigGenerator.GetAUMIDFromProject(SC.StageTargetPlatform.PlatformType, CommandUtils.Logger, Params.RawProjectPath.Directory, Configuration, Params.ShortProjectName, SC.CustomConfig);
+				TargetReceipt Receipt = SC.StageTargets.First(X => X.Receipt.Launch!.GetFileName() == Path.GetFileName(ClientApp)).Receipt;
+				UnrealTargetConfiguration Configuration = Receipt.Configuration;
+				UnrealArch? Architecture = Params.NoBootstrapExe ? Receipt.Architectures.SingleArchitecture : null;
+				string AUMIDName = GDKGameConfigGenerator.GetAUMIDFromProject(SC.StageTargetPlatform.PlatformType, CommandUtils.Logger, Params.RawProjectPath.Directory, Configuration, Architecture, Params.ShortProjectName, SC.CustomConfig);
 				string PackageFamilyName = AUMIDName.Substring(0, AUMIDName.IndexOf('!'));
 
 				// see if the application has been deployed (either a staged or packaged build - it doesn't matter at this stage)
diff --git a/Engine/Source/Programs/AutomationTool/Win/MSGaming/MSGamingRuntimeCustomStagingHandler.cs b/Engine/Source/Programs/AutomationTool/Win/MSGaming/MSGamingRuntimeCustomStagingHandler.cs
--- a/Engine/Source/Programs/AutomationTool/Win/MSGaming/MSGamingRuntimeCustomStagingHandler.cs
+++ b/Engine/Source/Programs/AutomationTool/Win/MSGaming/MSGamingRuntimeCustomStagingHandler.cs
@@ -60,6 +60,12 @@ public class MSGamingRuntimeCustomStagingHandler : CustomStagingHandler
 			return false;
 		}
 
+		// sanity check (unlikely to get this far because GRDK will throw an exception during the build, but keeping it here in case logic changes elsewhere)
+		if (GDKEdition < 251000 && Params.NoBootstrapExe && SC.StageTargets.Any(X => !X.Receipt.Architectures.SingleArchitecture.bIsX64))
+		{
+			throw new AutomationException("October 2025 GDK or higher is required to use the GDK with ARM64 unless you use the bootstrapper (omit -NoBootstrapExe)");
+		}
+
 		return true;
 	}
 
diff --git a/Engine/Source/Programs/UnrealBuildTool/Platform/Microsoft/GDKGameConfigGenerator.cs b/Engine/Source/Programs/UnrealBuildTool/Platform/Microsoft/GDKGameConfigGenerator.cs
--- a/Engine/Source/Programs/UnrealBuildTool/Platform/Microsoft/GDKGameConfigGenerator.cs
+++ b/Engine/Source/Programs/UnrealBuildTool/Platform/Microsoft/GDKGameConfigGenerator.cs
@@ -50,14 +50,14 @@ namespace UnrealBuildTool.GDK
 		/// <summary>
 		/// Create a manifest for the given platform and return the list of modified files
 		/// </summary>
-		public static List<string>? Create(UnrealTargetPlatform Platform, ILogger Logger, DirectoryReference InOutputPath, string? InTargetName, FileReference? InProjectFile, Dictionary<UnrealTargetConfiguration, string> InExecutablePairs, bool bDeleteIntermediates = false, FileReference? DLCFile = null, string CustomConfig = "", bool NoBootstrapExe = true)
+		public static List<string>? Create(UnrealTargetPlatform Platform, ILogger Logger, DirectoryReference InOutputPath, string? InTargetName, FileReference? InProjectFile, IEnumerable<AppXManifestExecutable> InExecutables, bool bDeleteIntermediates = false, FileReference? DLCFile = null, string CustomConfig = "", bool NoBootstrapExe = true)
 		{
 			GDKGameConfigGenerator Generator = GDKGameConfigGeneratorFactory.CreateGenerator(Platform, Logger);
 			Generator.DLCFile = DLCFile;
 			Generator.CustomConfig = CustomConfig;
 			Generator.NoBootstrapExe = NoBootstrapExe;
 
-			List<string>? ModifiedFiles = Generator.CreateManifest("MicrosoftGame.config", InOutputPath, InTargetName, InProjectFile, InExecutablePairs);
+			List<string>? ModifiedFiles = Generator.CreateManifest("MicrosoftGame.config", InOutputPath, InTargetName, InProjectFile, InExecutables);
 
 			return ModifiedFiles;
 		}
@@ -65,7 +65,7 @@ namespace UnrealBuildTool.GDK
 		/// <summary>
 		/// Get the Game element
 		/// </summary>
-		protected virtual XElement GetGame(Dictionary<UnrealTargetConfiguration, string> InExecutablePairs, out string IdentityName)
+		protected virtual XElement GetGame(IEnumerable<AppXManifestExecutable> InExecutables, out string IdentityName)
 		{
 			int ManifestVersion = GetManifestVersion();
 
@@ -79,8 +79,8 @@ namespace UnrealBuildTool.GDK
 			if (!IsDLC)
 			{
 				Game.Add(
-					GetExecutableList(InExecutablePairs),
-					GetProtocolList(InExecutablePairs)
+					GetExecutableList(InExecutables),
+					GetProtocolList(InExecutables)
 				);
 
 				string? TitleId = GetConfigString("TitleId", null, null);
@@ -134,27 +134,27 @@ namespace UnrealBuildTool.GDK
 		/// <summary>
 		/// Get the executable list element
 		/// </summary>
-		protected XElement GetExecutableList(Dictionary<UnrealTargetConfiguration, string> InExecutablePairs)
+		protected XElement GetExecutableList(IEnumerable<AppXManifestExecutable> InExecutables)
 		{
 			XElement ExecutableList = new(XName.Get("ExecutableList"));
 
-			if (InExecutablePairs.Count < 1)
+			if (!InExecutables.Any())
 			{
 				return ExecutableList;
 			}
 
 			// Need exactly one exe not marked as dev only in manifest.  Most optimized exe gets non-dev.  Dependent on UnrealTargetConfiguration being ordered by optimization level.
-			UnrealTargetConfiguration MaxTargetConfiguration = InExecutablePairs.Keys.Max();
-			foreach (KeyValuePair<UnrealTargetConfiguration, string> ExecutablePair in InExecutablePairs)
+			UnrealTargetConfiguration MaxTargetConfiguration = InExecutables.MaxBy( X => X.Configuration ).Configuration;
+			foreach (AppXManifestExecutable Executable in InExecutables)
 			{
-				XElement Executable = GetExecutable(ExecutablePair.Key, ExecutablePair.Value, ExecutablePair.Key != MaxTargetConfiguration);
-				ExecutableList.Add(Executable);
+				XElement ExecutableElement = GetExecutable(Executable, Executable.Configuration != MaxTargetConfiguration);
+				ExecutableList.Add(ExecutableElement);
 			}
 
 			return ExecutableList;
 		}
 
-		private string GetAppId(UnrealTargetConfiguration InTargetConfig)
+		private string GetAppId(UnrealTargetConfiguration InTargetConfig, UnrealArch? InArchitecture)
 		{
 			bool bIsDevBuild = (InTargetConfig != UnrealTargetConfiguration.Shipping);
 			bool bSkipAppIdSuffixForShippingConfig = GetConfigBool("bSkipAppIdSuffixForShippingConfig", null, false);
@@ -162,6 +162,10 @@ namespace UnrealBuildTool.GDK
 
 			string? PackageBaseName = ValidateProjectBaseName(ReadIniString("ProjectName", IniSection_GeneralProjectSettings, "UEGame"));
 			string AppId = "App" + PackageBaseName + (bIncludeConfigSuffix ? InTargetConfig.ToString() : String.Empty);
+			if (InArchitecture != null && !InArchitecture!.Value.bIsX64)
+			{
+				AppId += "ARM64";
+			}
 
 			return AppId;
 		}
@@ -169,56 +173,70 @@ namespace UnrealBuildTool.GDK
 		/// <summary>
 		/// Get the executable element for the given configuration
 		/// </summary>
-		protected XElement GetExecutable(UnrealTargetConfiguration InTargetConfig, string InExecutablePath, bool IsDevOnly)
+		protected XElement GetExecutable(AppXManifestExecutable Executable, bool IsDevOnly)
 		{
-			bool bIsDevBuild = (InTargetConfig != UnrealTargetConfiguration.Shipping);
+			bool bIsDevBuild = (Executable.Configuration != UnrealTargetConfiguration.Shipping);
 
 			string RelativeExePath;
-			if (Path.IsPathRooted(InExecutablePath))
+			if (Path.IsPathRooted(Executable.ExecutablePath))
 			{
-				RelativeExePath = Utils.MakePathRelativeTo(InExecutablePath, Path.Combine(InExecutablePath, "../../../.."));
+				RelativeExePath = Utils.MakePathRelativeTo(Executable.ExecutablePath, Path.Combine(Executable.ExecutablePath, "../../../.."));
 			}
 			else
 			{
-				RelativeExePath = InExecutablePath;
+				RelativeExePath = Executable.ExecutablePath;
 			}
 
-			string AppId = GetAppId(InTargetConfig);
+			string AppId = GetAppId(Executable.Configuration, Executable.Architecture);
 
-			XElement Executable = new(XName.Get("Executable"),
+			XElement ExecutableElement = new(XName.Get("Executable"),
 				new XAttribute("Name", RelativeExePath),
 				new XAttribute("Id", AppId));
 
 			string DisplayNameDefault = ReadIniString(IniSection_GeneralProjectSettings, "ProjectName", GetDefaultDisplayName());
+			string ResourceEntrySuffix = (Executable.Architecture != null && Executable.Architecture != UnrealArch.X64) ? Executable.Architecture!.Value.ToString() : "";
 			if (bIsDevBuild)
 			{
-				string DisplayNameSuffix = " (" + InTargetConfig.ToString() + ")";
-				string DisplayNameResource = AppXResources!.AddResourceString("ExeDisplayName" + InTargetConfig.ToString(), "ApplicationDisplayName", DisplayNameDefault, DisplayNameSuffix);
-				Executable.Add(new XAttribute("OverrideDisplayName", DisplayNameResource));
+				string ExecutableInfo = Executable.Configuration.ToString();
+				if (Executable.Architecture != null && Executable.Architecture != UnrealArch.X64)
+				{
+					ExecutableInfo += " " + Executable.Architecture.ToString();
+				}
+
+				string DisplayNameSuffix = " (" + ExecutableInfo + ")";
+				string ResourceEntryName = "ExeDisplayName" + Executable.Configuration.ToString() + ResourceEntrySuffix;
+				string DisplayNameResource = AppXResources!.AddResourceString(ResourceEntryName, "ApplicationDisplayName", DisplayNameDefault, DisplayNameSuffix);
+				ExecutableElement.Add(new XAttribute("OverrideDisplayName", DisplayNameResource));
 			}
 			else
 			{
-				string DisplayNameResource = AppXResources!.AddResourceString("AppDisplayName", "ApplicationDisplayName", DisplayNameDefault);
-				Executable.Add(new XAttribute("OverrideDisplayName", DisplayNameResource));
+				string ResourceEntryName = "AppDisplayName" + ResourceEntrySuffix;
+				string DisplayNameResource = AppXResources!.AddResourceString(ResourceEntryName, "ApplicationDisplayName", DisplayNameDefault);
+				ExecutableElement.Add(new XAttribute("OverrideDisplayName", DisplayNameResource));
 			}
 
 			if (TargetDeviceFamily != null)
 			{
-				Executable.Add(new XAttribute("TargetDeviceFamily", TargetDeviceFamily));
+				ExecutableElement.Add(new XAttribute("TargetDeviceFamily", TargetDeviceFamily));
+			}
+
+			if (Executable.Architecture != null)
+			{
+				ExecutableElement.Add(new XAttribute("Architecture", Executable.Architecture!.Value.bIsX64 ? "x64" : "ARM64"));
 			}
 
 			if (IsDevOnly)
 			{
-				Executable.Add(new XAttribute("IsDevOnly", "true"));
+				ExecutableElement.Add(new XAttribute("IsDevOnly", "true"));
 			}
 
-			return Executable;
+			return ExecutableElement;
 		}
 
 		/// <summary>
 		/// Get the protocol list element
 		/// </summary>
-		protected XElement? GetProtocolList(Dictionary<UnrealTargetConfiguration, string> InExecutablePairs)
+		protected XElement? GetProtocolList(IEnumerable<AppXManifestExecutable> InExecutables)
 		{
 			XElement? ProtocolList = null;
 
@@ -226,11 +244,12 @@ namespace UnrealBuildTool.GDK
 			if (GDKEdition >= 230300)
 			{
 				// if there's multiple executables in the build, need to specify which one to associate with the protocol. it must use back-slashes
+				// documentation for newer GDKs indicate that this element is optional and it will use the default executable in ExecutableList if it is not specified
 				string? DefaultExecutable = null;
-				if (InExecutablePairs.Count > 1)
+				if (InExecutables.Count() > 1 && GDKEdition < 251000)
 				{
-					UnrealTargetConfiguration MaxTargetConfiguration = InExecutablePairs.Keys.Max(); // select the most optimized exe
-					DefaultExecutable = InExecutablePairs[MaxTargetConfiguration].Replace('/', '\\');
+					AppXManifestExecutable MaxExecutable = InExecutables.MaxBy( X => X.Configuration ); // select the most optimized exe
+					DefaultExecutable = MaxExecutable.ExecutablePath.Replace('/', '\\');
 				}
 
 				List<string> ProtocolNames = new();
@@ -612,19 +631,19 @@ namespace UnrealBuildTool.GDK
 		/// <summary>
 		/// Return the entire manifest element
 		/// </summary>
-		protected override XElement GetManifest(Dictionary<UnrealTargetConfiguration, string> InExecutablePairs, out string IdentityName)
+		protected override XElement GetManifest(IEnumerable<AppXManifestExecutable> InExecutables, out string IdentityName)
 		{
-			return GetGame(InExecutablePairs, out IdentityName);
+			return GetGame(InExecutables, out IdentityName);
 		}
 
 		/// <summary>
 		/// 
 		/// </summary>
-		protected override void ProcessManifest(Dictionary<UnrealTargetConfiguration, string> InExecutablePairs, string ManifestName, string ManifestTargetPath, string ManifestIntermediatePath)
+		protected override void ProcessManifest(IEnumerable<AppXManifestExecutable> InExecutables, string ManifestName, string ManifestTargetPath, string ManifestIntermediatePath)
 		{
 			AddResourcesForPackaging();
 
-			base.ProcessManifest(InExecutablePairs, ManifestName, ManifestTargetPath, ManifestIntermediatePath);
+			base.ProcessManifest(InExecutables, ManifestName, ManifestTargetPath, ManifestIntermediatePath);
 		}
 
 		/// <summary>
@@ -765,14 +784,14 @@ namespace UnrealBuildTool.GDK
 		/// <summary>
 		/// Returns the AUMID for the given project and configuration
 		/// </summary>
-		public static string GetAUMIDFromProject(UnrealTargetPlatform InPlatform, ILogger Logger, DirectoryReference InProjectDir, UnrealTargetConfiguration InTargetConfig, string? InTargetName, string InCustomConfig)
+		public static string GetAUMIDFromProject(UnrealTargetPlatform InPlatform, ILogger Logger, DirectoryReference InProjectDir, UnrealTargetConfiguration InTargetConfig, UnrealArch? InArchitecture, string? InTargetName, string InCustomConfig)
 		{
 			GDKGameConfigGenerator ManifestGenerator = GDKGameConfigGeneratorFactory.CreateGenerator(InPlatform, Logger);
 			ManifestGenerator.CustomConfig = InCustomConfig;
 
 			ManifestGenerator.GetManifestElementsForProject(InProjectDir, InTargetName, out string Name, out _, out _, out _, out string Publisher);
 
-			string AppId = ManifestGenerator.GetAppId(InTargetConfig);
+			string AppId = ManifestGenerator.GetAppId(InTargetConfig, InArchitecture);
 
 			return $"{Name}_{GDKGameConfigInfo.EncodePublisherNameForPackageName(Publisher)}!{AppId}";
 		}
@@ -828,6 +847,9 @@ namespace UnrealBuildTool.GDK
 		[CommandLine("-Configuration", Required = false)]
 		UnrealTargetConfiguration? Configuration = null;
 
+		[CommandLine("-Architecture", Required = false)]
+		UnrealArch? Architecture = null;
+
 		[CommandLine("-Project", Required = false)]
 		FileReference? ProjectFile = null;
 
@@ -850,17 +872,17 @@ namespace UnrealBuildTool.GDK
 			Arguments.ApplyTo(this);
 			Arguments.CheckAllArgumentsUsed();
 
-			Dictionary<UnrealTargetConfiguration, string> ExecutablePairs;
-			if (Executable == null && Configuration == null)
+			List<AppXManifestExecutable> Executables;
+			if (Executable == null || Configuration == null)
 			{
-				ExecutablePairs = [];
+				Executables = [];
 			}
 			else
 			{
-				ExecutablePairs = new() { { (UnrealTargetConfiguration)Configuration!, (string)Executable! } };
+				Executables = [ new (Configuration!.Value, Executable!, Architecture ) ];
 			}
 
-			GDKGameConfigGenerator.Create((UnrealTargetPlatform)Platform!, Logger, OutputDirectory!, TargetName, ProjectFile, ExecutablePairs, CustomConfig: CustomConfig);
+			GDKGameConfigGenerator.Create((UnrealTargetPlatform)Platform!, Logger, OutputDirectory!, TargetName, ProjectFile, Executables, CustomConfig: CustomConfig);
 
 			return Task.FromResult(0);
 		}
@@ -868,7 +890,7 @@ namespace UnrealBuildTool.GDK
 		/// <summary>
 		/// Helper function to create an action that will run this ToolMode
 		/// </summary>
-		public static Action CreateAction(TargetMakefileBuilder MakefileBuilder, UnrealTargetPlatform Platform, DirectoryReference OutputDirectory, string TargetName, FileReference? ProjectFile, UnrealTargetConfiguration? Configuration = null, string? ExecutableFileName = null, string CustomConfig = "")
+		public static Action CreateAction(TargetMakefileBuilder MakefileBuilder, UnrealTargetPlatform Platform, DirectoryReference OutputDirectory, string TargetName, FileReference? ProjectFile, UnrealTargetConfiguration? Configuration = null, string? ExecutableFileName = null, UnrealArch? Architecture = null, string CustomConfig = "")
 		{
 			Action GenAction = MakefileBuilder.CreateAction(ActionType.PostBuildStep);
 			GenAction.CommandPath = Unreal.DotnetPath;
@@ -884,6 +906,11 @@ namespace UnrealBuildTool.GDK
 				GenAction.CommandArguments += $" -Executable =\"{ExecutableFileName}\" -Configuration={Configuration}";
 			}
 
+			if (Architecture != null)
+			{
+				GenAction.CommandArguments += $" -Architecture={Architecture}";
+			}
+
 			if (!String.IsNullOrEmpty(CustomConfig))
 			{
 				GenAction.CommandArguments += $" -CustomConfig={CustomConfig}";
diff --git a/Engine/Source/Programs/UnrealBuildTool/Platform/Windows/AppXManifestGeneratorBase.cs b/Engine/Source/Programs/UnrealBuildTool/Platform/Windows/AppXManifestGeneratorBase.cs
--- a/Engine/Source/Programs/UnrealBuildTool/Platform/Windows/AppXManifestGeneratorBase.cs
+++ b/Engine/Source/Programs/UnrealBuildTool/Platform/Windows/AppXManifestGeneratorBase.cs
@@ -13,6 +13,29 @@ using UnrealBuildBase;
 
 namespace UnrealBuildTool
 {
+	/// <summary>
+	/// Information about an executable in an AppX manifest file
+	/// </summary>
+	public struct AppXManifestExecutable
+	{
+		/// Constructor
+		public AppXManifestExecutable(UnrealTargetConfiguration InConfiguration, string InExecutablePath, UnrealArch? InArchitecture = null)
+		{
+			Configuration = InConfiguration;
+			ExecutablePath = InExecutablePath;
+			Architecture = InArchitecture;
+		}
+
+		/// The build configuration
+		public readonly UnrealTargetConfiguration Configuration;
+
+		/// The path to the executable - relative or absolute
+		public readonly string ExecutablePath;
+
+		/// The executable architecture (optional)
+		public readonly UnrealArch? Architecture;
+	};
+
 	/// <summary>
 	/// Base class for VC appx manifest generation
 	/// </summary>
@@ -461,12 +484,12 @@ namespace UnrealBuildTool
 		/// <summary>
 		/// Return the entire manifest element
 		/// </summary>
-		protected abstract XElement GetManifest(Dictionary<UnrealTargetConfiguration, string> InExecutablePairs, out string IdentityName);
+		protected abstract XElement GetManifest(IEnumerable<AppXManifestExecutable> InExecutables, out string IdentityName);
 
 		/// <summary>
 		/// Perform any platform-specific processing on the manifest before it is saved
 		/// </summary>
-		protected virtual void ProcessManifest(Dictionary<UnrealTargetConfiguration, string> InExecutablePairs, string ManifestName, string ManifestTargetPath, string ManifestIntermediatePath)
+		protected virtual void ProcessManifest(IEnumerable<AppXManifestExecutable> InExecutables, string ManifestName, string ManifestTargetPath, string ManifestIntermediatePath)
 		{
 		}
 
@@ -480,7 +503,7 @@ namespace UnrealBuildTool
 		/// <summary>
 		/// Create a manifest and return the list of modified files
 		/// </summary>
-		public List<string>? CreateManifest(string InManifestName, DirectoryReference OutputDirectory, string? InTargetName, FileReference? InProjectFile, Dictionary<UnrealTargetConfiguration, string> InExecutablePairs)
+		public List<string>? CreateManifest(string InManifestName, DirectoryReference OutputDirectory, string? InTargetName, FileReference? InProjectFile, IEnumerable<AppXManifestExecutable> InExecutables)
 		{
 			FileUtils.CreateDirectoryTree(OutputDirectory);
 
@@ -547,14 +570,14 @@ namespace UnrealBuildTool
 
 			// Create the manifest document
 			string? IdentityName = null;
-			XDocument ManifestXmlDocument = new XDocument(GetManifest(InExecutablePairs, out IdentityName));
+			XDocument ManifestXmlDocument = new XDocument(GetManifest(InExecutables, out IdentityName));
 
 			// Export manifest to the intermediate directory and add it to the manifest resources files for copying
 			FileReference ManifestIntermediateFile = FileReference.Combine(IntermediateDirectory, InManifestName);
 			FileReference ManifestTargetFile = FileReference.Combine(OutputDirectory, InManifestName);
 			ManifestXmlDocument.Save(ManifestIntermediateFile.FullName);
 			AppXResources.AddFileReference(ManifestIntermediateFile, InManifestName);
-			ProcessManifest(InExecutablePairs, InManifestName, ManifestTargetFile.FullName, ManifestIntermediateFile.FullName);
+			ProcessManifest(InExecutables, InManifestName, ManifestTargetFile.FullName, ManifestIntermediateFile.FullName);
 
 			// Generate the package resource index and copy all resource files to the output
 			FileReference ManifestTargetPath = FileReference.Combine(OutputDirectory, InManifestName);
diff --git a/Engine/Source/Programs/UnrealBuildTool/Platform/Windows/MSGaming/MSGamingGameConfigGenerator.cs b/Engine/Source/Programs/UnrealBuildTool/Platform/Windows/MSGaming/MSGamingGameConfigGenerator.cs
--- a/Engine/Source/Programs/UnrealBuildTool/Platform/Windows/MSGaming/MSGamingGameConfigGenerator.cs
+++ b/Engine/Source/Programs/UnrealBuildTool/Platform/Windows/MSGaming/MSGamingGameConfigGenerator.cs
@@ -76,13 +76,13 @@ namespace UnrealBuildTool.GDK
 			}
 			AppXResources!.EngineFallbackBinaryResourceDirectories.Add(MSGamingSupportResourcePath);
 
-			// fallback path because Program targets do not load plugins & this is still under the platforms folder for now
+			// fallback path because Program targets do not load plugins
 			AppXResources!.EngineFallbackBinaryResourceDirectories.Add(DirectoryReference.Combine(Unreal.EngineDirectory, @"Plugins\Runtime\Windows\MSGamingSupport\Build\DefaultImages"));
 		}
 
-		protected override XElement GetGame(Dictionary<UnrealTargetConfiguration, string> InExecutablePairs, out string IdentityName)
+		protected override XElement GetGame(IEnumerable<AppXManifestExecutable> InExecutables, out string IdentityName)
 		{
-			XElement Game = base.GetGame(InExecutablePairs, out IdentityName);
+			XElement Game = base.GetGame(InExecutables, out IdentityName);
 
 			Game.Add(GetDesktopRegistration());
 
