UnrealBuildTool: Change to checking for a generated UHT file that should always be present when deciding if we need to run UHT, rather than just checking for existence of the UHT generated code directory, as that directory now always gets created when pre-populating the filesystem cache.

Remove some cache invalidation code that was never used.

#Synced-CL 51291646

diff --git a/Engine/Source/Programs/Shared/EpicGames.Build/System/DirectoryItem.cs b/Engine/Source/Programs/Shared/EpicGames.Build/System/DirectoryItem.cs
--- a/Engine/Source/Programs/Shared/EpicGames.Build/System/DirectoryItem.cs
+++ b/Engine/Source/Programs/Shared/EpicGames.Build/System/DirectoryItem.cs
@@ -119,8 +119,8 @@ namespace UnrealBuildBase
 		}
 
 		/// <summary>
-		/// This function can be used to force the cache to be populated with files that not yet exists.
-		/// Useful when wanting to run things in parallel and one thread write files while the other use the list of files in the directory
+		/// This function can be used to force the cache to be populated with files that do not yet exist.
+		/// Useful when wanting to run things in parallel and one thread write files while the other uses the list of files in the directory.
 		/// </summary>
 		/// <param name="files"></param>
 		public void PrepopulateCache(List<string> files)
diff --git a/Engine/Source/Programs/UnrealBuildTool/Configuration/UEBuildTarget.cs b/Engine/Source/Programs/UnrealBuildTool/Configuration/UEBuildTarget.cs
--- a/Engine/Source/Programs/UnrealBuildTool/Configuration/UEBuildTarget.cs
+++ b/Engine/Source/Programs/UnrealBuildTool/Configuration/UEBuildTarget.cs
@@ -3021,8 +3021,10 @@ namespace UnrealBuildTool
 
 				if (makefile.UObjectModules.Count > 0)
 				{
+					// Because UHT execution runs alongside building the rest of the makefile, I suspect we need to pre-populate the cache
+					// for the makefile generation to correctly pick up source files.
 					UHTExecution.PrepopulateDirectoryCache(makefile.UObjectModules);
-					(Task? uhtRun, Task? uhtTimestamp) = UHTExecution.ExecuteHeaderToolIfNecessary(BuildConfiguration, descriptor.ProjectFile, makefile, context.TargetName!, Logger, false);
+					(Task? uhtRun, Task? uhtTimestamp) = UHTExecution.ExecuteHeaderToolIfNecessary(BuildConfiguration, descriptor.ProjectFile, makefile, context.TargetName!, Logger);
 					context.uhtRun = uhtRun;
 					context.uhtTimestamp = uhtTimestamp;
 				}
diff --git a/Engine/Source/Programs/UnrealBuildTool/System/UHTExecution.cs b/Engine/Source/Programs/UnrealBuildTool/System/UHTExecution.cs
--- a/Engine/Source/Programs/UnrealBuildTool/System/UHTExecution.cs
+++ b/Engine/Source/Programs/UnrealBuildTool/System/UHTExecution.cs
@@ -834,6 +834,9 @@ namespace UnrealBuildTool
 			}
 		}
 
+		private static FileInfo GetInitGenCppFileInfo(string ModuleName, string ModuleGeneratedCodeDirectory) =>
+			new(Path.Combine(ModuleGeneratedCodeDirectory, ModuleName + ".init.gen.cpp"));
+
 		/// <summary>
 		/// Gets the timestamp of CoreUObject.gen.cpp file.
 		/// </summary>
@@ -847,7 +850,7 @@ namespace UnrealBuildTool
 			}
 
 			// Otherwise look for CoreUObject.init.gen.cpp
-			FileInfo CoreGeneratedFileInfo = new(Path.Combine(ModuleGeneratedCodeDirectory, ModuleName + ".init.gen.cpp"));
+			FileInfo CoreGeneratedFileInfo = GetInitGenCppFileInfo(ModuleName, ModuleGeneratedCodeDirectory);
 			if (CoreGeneratedFileInfo.Exists)
 			{
 				return CoreGeneratedFileInfo.LastWriteTimeUtc;
@@ -899,7 +902,7 @@ namespace UnrealBuildTool
 				{
 					if (Module.ModuleName.Equals("CoreUObject", StringComparison.OrdinalIgnoreCase))
 					{
-						CoreGeneratedTimestampUtc = GetCoreGeneratedTimestampUtc(Module.ModuleName, Path.GetDirectoryName(Module.GeneratedCPPFilenameBase)!);
+						CoreGeneratedTimestampUtc = GetCoreGeneratedTimestampUtc(Module.ModuleName, Module.GeneratedCodeDirectory.FullName);
 						break;
 					}
 				}
@@ -931,12 +934,10 @@ namespace UnrealBuildTool
 					return;
 				}
 
-				// Make sure we have an existing folder for generated code.  If not, then we definitely need to generate code!
-				//string GeneratedCodeDirectory = Path.GetDirectoryName(Module.GeneratedCPPFilenameBase)!;
-				DirectoryItem GeneratedCodeDirectory = Module.GeneratedCodeDirectory;
-
-				//FileSystemInfo TestDirectory = (FileSystemInfo)new DirectoryInfo(GeneratedCodeDirectory);
-				if (!GeneratedCodeDirectory.Exists)
+				// Make sure we have an existing init.gen.cpp for generated code. If not, then we definitely need to generate code.
+				// We can't check the generated code directory for existence as this is now pre-created in PrepopulateDirectoryCache.
+				FileInfo GenCppFileInfo = GetInitGenCppFileInfo(Module.ModuleName, Module.GeneratedCodeDirectory.FullName);
+				if (!GenCppFileInfo.Exists)
 				{
 					// Generated code directory is missing entirely!
 					if (Interlocked.Increment(ref invalidCounter) == 1)
@@ -945,7 +946,6 @@ namespace UnrealBuildTool
 					}
 					return;
 				}
-				GeneratedCodeDirectory.CacheFiles();
 
 				if (!timestampsFile.GetModuleInfo(Module.ModuleName, out UInt128 moduleOldHash, out long moduleLastUpdateUtc))
 				{
@@ -1202,7 +1202,7 @@ namespace UnrealBuildTool
 
 		/// <summary>
 		/// Will prepopulate DirectoryInfo cache with files that will be written.
-		/// This can be combined with ExecuteHeaderToolIfNecessaryAsync(resetDirectoryCaches = false) and WaitForExecution
+		/// This can be combined with ExecuteHeaderToolIfNecessary and WaitForExecution.
 		/// </summary>
 		/// <param name="moduleInfos"></param>
 		public static void PrepopulateDirectoryCache(List<UHTModuleInfo> moduleInfos)
@@ -1220,7 +1220,6 @@ namespace UnrealBuildTool
 					outputFiles.Add(fileName + ".gen.cpp");
 				}
 
-				DirectoryReference.CreateDirectory(mi.GeneratedCodeDirectory.Location);
 				mi.GeneratedCodeDirectory.PrepopulateCache(outputFiles);
 			});
 		}
@@ -1326,7 +1325,7 @@ namespace UnrealBuildTool
 		/// Builds and runs the header tool and touches the header directories.
 		/// Performs any early outs if headers need no changes, given the UObject modules, tool path, game name, and configuration
 		/// </summary>
-		public static (Task run, Task timestamp) ExecuteHeaderToolIfNecessary(BuildConfiguration BuildConfiguration, FileReference? ProjectFile, TargetMakefile Makefile, string TargetName, ILogger Logger, bool resetDirectoryCaches = true)
+		public static (Task run, Task timestamp) ExecuteHeaderToolIfNecessary(BuildConfiguration BuildConfiguration, FileReference? ProjectFile, TargetMakefile Makefile, string TargetName, ILogger Logger)
 		{
 			bool runTimestamp = false;
 
@@ -1371,7 +1370,7 @@ namespace UnrealBuildTool
 								{
 									bUHTNeedsToRun = true;
 								}
-							} 
+							}
 							catch (Exception)
 							{
 								bUHTNeedsToRun = true;
@@ -1422,9 +1421,9 @@ namespace UnrealBuildTool
 						// Generate the command line
 						List<string> CmdArgs = [
 							(ProjectFile != null) ? ProjectFile.FullName : TargetName,
-						ModuleInfoFileName.FullName,
-						"-WarningsAsErrors",
-					];
+							ModuleInfoFileName.FullName,
+							"-WarningsAsErrors",
+						];
 
 						if (Unreal.IsEngineInstalled())
 						{
@@ -1503,16 +1502,6 @@ namespace UnrealBuildTool
 						{
 							Makefile.UHTTargetSettings.Write(writer);
 						}
-
-						if (resetDirectoryCaches)
-						{
-							// Now that UHT has successfully finished generating code, we need to update all cached FileItems in case their last write time has changed.
-							// Otherwise UBT might not detect changes UHT made.
-							using (GlobalTracer.Instance.BuildSpan("UHTExecution.ResetCachedHeaderInfo()").StartActive())
-							{
-								ResetCachedHeaderInfo(Makefile.UObjectModules);
-							}
-						}
 					}
 					else
 					{
