UnrealBuildTool: Run QueryTargetsMode in parallel

#rnx

#Synced-CL 51388341

diff --git a/Engine/Source/Programs/UnrealBuildTool/Modes/QueryTargetsMode.cs b/Engine/Source/Programs/UnrealBuildTool/Modes/QueryTargetsMode.cs
--- a/Engine/Source/Programs/UnrealBuildTool/Modes/QueryTargetsMode.cs
+++ b/Engine/Source/Programs/UnrealBuildTool/Modes/QueryTargetsMode.cs
@@ -1,6 +1,7 @@
 // Copyright Epic Games, Inc. All Rights Reserved.
 
 using System;
+using System.Collections.Concurrent;
 using System.Collections.Generic;
 using System.Linq;
 using System.Threading.Tasks;
@@ -44,7 +45,7 @@ namespace UnrealBuildTool
 		/// <summary>
 		/// Write all targets from parent assemblies (useful for EngineRules which can be chained)
 		/// </summary>
-		[CommandLine("-IncludeParentAssembly", Value ="true")]
+		[CommandLine("-IncludeParentAssembly", Value = "true")]
 		[CommandLine("-DontIncludeParentAssembly", Value = "false")]
 		bool bIncludeParentAssembly = true;
 #pragma warning restore IDE0044
@@ -52,38 +53,34 @@ namespace UnrealBuildTool
 		/// <summary>
 		/// Execute the mode
 		/// </summary>
-		/// <param name="Arguments">Command line arguments</param>
+		/// <param name="arguments">Command line arguments</param>
 		/// <returns></returns>
-		/// <param name="Logger"></param>
-		public override Task<int> ExecuteAsync(CommandLineArguments Arguments, ILogger Logger)
+		/// <param name="logger"></param>
+		public override Task<int> ExecuteAsync(CommandLineArguments arguments, ILogger logger)
 		{
-			Arguments.ApplyTo(this);
+			arguments.ApplyTo(this);
 
 			// Create the build configuration object, and read the settings
-			BuildConfiguration BuildConfiguration = new BuildConfiguration();
-			XmlConfig.ApplyTo(BuildConfiguration);
-			Arguments.ApplyTo(BuildConfiguration);
+			BuildConfiguration buildConfiguration = new BuildConfiguration();
+			XmlConfig.ApplyTo(buildConfiguration);
+			arguments.ApplyTo(buildConfiguration);
 
 			// Ensure the path to the output file is valid
-			if (OutputFile == null)
-			{
-				OutputFile = GetDefaultOutputFile(ProjectFile);
-			}
+			OutputFile ??= GetDefaultOutputFile(ProjectFile);
 
 			// Create the rules assembly
-			RulesAssembly Assembly;
-			if (ProjectFile != null)
+			RulesAssembly assembly;
+			string mutexName = GlobalSingleInstanceMutex.GetUniqueMutexForPath("UnrealBuildTool_QueryMode_UEBuildTarget-Create", Unreal.RootDirectory);
+			using (new GlobalSingleInstanceMutex(mutexName, true))
 			{
-				Assembly = RulesCompiler.CreateProjectRulesAssembly(ProjectFile, BuildConfiguration.bUsePrecompiled, BuildConfiguration.bSkipRulesCompile, BuildConfiguration.bForceRulesCompile, Logger);
-			}
-			else
-			{
-				Assembly = RulesCompiler.CreateEngineRulesAssembly(BuildConfiguration.bUsePrecompiled, BuildConfiguration.bSkipRulesCompile, BuildConfiguration.bForceRulesCompile, Logger);
+				assembly = ProjectFile != null
+					? RulesCompiler.CreateProjectRulesAssembly(ProjectFile, buildConfiguration.bUsePrecompiled, buildConfiguration.bSkipRulesCompile, buildConfiguration.bForceRulesCompile, logger)
+					: RulesCompiler.CreateEngineRulesAssembly(buildConfiguration.bUsePrecompiled, buildConfiguration.bSkipRulesCompile, buildConfiguration.bForceRulesCompile, logger);
 			}
 
 			// Write information about these targets
-			WriteTargetInfo(ProjectFile, Assembly, OutputFile, Arguments, Logger, bIncludeAllTargets, bIncludeParentAssembly, bIncludeProgramTargets);
-			Logger.LogInformation("Written {OutputFile}", OutputFile);
+			WriteTargetInfo(ProjectFile, assembly, OutputFile, arguments, logger, bIncludeAllTargets, bIncludeParentAssembly, bIncludeProgramTargets);
+			logger.LogInformation("Written {OutputFile}", OutputFile);
 			return Task.FromResult(0);
 		}
 
@@ -94,121 +91,122 @@ namespace UnrealBuildTool
 		/// <returns>Path to the output file</returns>
 		public static FileReference GetDefaultOutputFile(FileReference? ProjectFile)
 		{
-			if (ProjectFile == null)
-			{
-				return FileReference.Combine(Unreal.EngineDirectory, "Intermediate", "TargetInfo.json");
-			}
-			else
-			{
-				return FileReference.Combine(ProjectFile.Directory, "Intermediate", "TargetInfo.json");
-			}
+			return ProjectFile == null
+				? FileReference.Combine(Unreal.EngineDirectory, "Intermediate", "TargetInfo.json")
+				: FileReference.Combine(ProjectFile.Directory, "Intermediate", "TargetInfo.json");
 		}
 
 		/// <summary>
 		/// Writes information about the targets in an assembly to a file
 		/// </summary>
-		/// <param name="ProjectFile">The project file for the targets being built</param>
-		/// <param name="Assembly">The rules assembly for this target</param>
-		/// <param name="OutputFile">Output file to write to</param>
-		/// <param name="Arguments"></param>
-		/// <param name="Logger">Logger for output</param>
+		/// <param name="projectFile">The project file for the targets being built</param>
+		/// <param name="assembly">The rules assembly for this target</param>
+		/// <param name="outputFile">Output file to write to</param>
+		/// <param name="arguments"></param>
+		/// <param name="logger">Logger for output</param>
 		/// <param name="bIncludeAllTargets">Include all targets even if a default target is specified for a given target type.</param>
 		/// <param name="bIncludeParentAssembly">Include all targets from parent assemblies (useful for EngineRules which can be chained)</param>
 		/// <param name="bIncludeProgramTargets">Include or exclude all program targets, has presedence over bIncludeAllTargets.</param>
-		public static void WriteTargetInfo(FileReference? ProjectFile, RulesAssembly Assembly, FileReference OutputFile, CommandLineArguments Arguments, ILogger Logger,
+		public static void WriteTargetInfo(FileReference? projectFile, RulesAssembly assembly, FileReference outputFile, CommandLineArguments arguments, ILogger logger,
 			bool bIncludeAllTargets = true, bool bIncludeParentAssembly = false, bool bIncludeProgramTargets = true)
 		{
 			// Construct all the targets in this assembly
-			List<string> TargetNames = new List<string>();
-			Assembly.GetAllTargetNames(TargetNames, bIncludeParentAssembly);
+			List<string> targetNames = new List<string>();
+			assembly.GetAllTargetNames(targetNames, bIncludeParentAssembly);
 
-			// Generate json
-			using System.IO.StringWriter StringWriter = new();
-			using (JsonWriter Writer = new(StringWriter))
+			ConcurrentBag<(string name, TargetType type, FileReference? path, bool? isDefaultTarget)> entries = [];
+			Parallel2.ForEach(targetNames, (targetName) =>
 			{
-				Writer.WriteObjectStart();
-				Writer.WriteArrayStart("Targets");
-				foreach (string TargetName in TargetNames)
+				// skip target rules that are platform extension or platform group specializations
+				string[] targetPathSplit = targetName.Split(new char[] { '_' }, StringSplitOptions.RemoveEmptyEntries);
+				if (targetPathSplit.Length > 1 && (UnrealTargetPlatform.IsValidName(targetPathSplit.Last()) || UnrealPlatformGroup.IsValidName(targetPathSplit.Last())))
 				{
-					// skip target rules that are platform extension or platform group specializations
-					string[] TargetPathSplit = TargetName.Split(new char[] { '_' }, StringSplitOptions.RemoveEmptyEntries);
-					if (TargetPathSplit.Length > 1 && (UnrealTargetPlatform.IsValidName(TargetPathSplit.Last()) || UnrealPlatformGroup.IsValidName(TargetPathSplit.Last())))
-					{
-						continue;
-					}
+					return;
+				}
 
-					// Construct the rules object
-					TargetRules TargetRules;
-					try
-					{
-						UnrealArchitectures Architectures = UnrealArchitectureConfig.ForPlatform(BuildHostPlatform.Current.Platform).ActiveArchitectures(ProjectFile, TargetName);
-						TargetRules = Assembly.CreateTargetRules(TargetName, BuildHostPlatform.Current.Platform, UnrealTargetConfiguration.Development, Architectures, ProjectFile, Arguments, Logger, ValidationOptions: TargetRulesValidationOptions.ValidateNothing);
-					}
-					catch (Exception Ex)
-					{
-						Logger.LogWarning("Unable to construct target rules for {TargetName}", TargetName);
-						Logger.LogDebug("{Ex}", ExceptionUtils.FormatException(Ex));
-						continue;
-					}
+				// Construct the rules object
+				TargetRules targetRules;
+				try
+				{
+					UnrealArchitectures architectures = UnrealArchitectureConfig.ForPlatform(BuildHostPlatform.Current.Platform).ActiveArchitectures(projectFile, targetName);
+					targetRules = assembly.CreateTargetRules(targetName, BuildHostPlatform.Current.Platform, UnrealTargetConfiguration.Development, architectures, projectFile, arguments, logger, ValidationOptions: TargetRulesValidationOptions.ValidateNothing);
+				}
+				catch (Exception ex)
+				{
+					logger.LogWarning("Unable to construct target rules for {TargetName}", targetName);
+					logger.LogDebug(ex, "{Ex}", ExceptionUtils.FormatException(ex));
+					return;
+				}
 
-					bool? bIsDefaultTarget = null;
+				bool? bIsDefaultTarget = null;
 
-					// Program targets are controlled by bIncludeProgramTargets which has presedence over bIncludeAllTargets
-					if (TargetRules.Type == TargetType.Program)
+				// Program targets are controlled by bIncludeProgramTargets which has presedence over bIncludeAllTargets
+				if (targetRules.Type == TargetType.Program)
+				{
+					if (!bIncludeProgramTargets)
 					{
-						if (!bIncludeProgramTargets)
-						{
-							continue;
-						}
+						return;
 					}
-					// Check if default target for Client, Editor, Game and Server targets
-					else
+				}
+				// Check if default target for Client, Editor, Game and Server targets
+				else
+				{
+					if (projectFile != null)
 					{
-						if (ProjectFile != null)
-						{
-							string? DefaultTargetName = ProjectFileGenerator.GetProjectDefaultTargetNameForType(ProjectFile.Directory, TargetRules.Type);
-
-							// GetProjectDefaultTargetNameForType returns
-							if (DefaultTargetName != null)
-							{
-								bIsDefaultTarget = DefaultTargetName == TargetName;
-							}
-						}
+						string? defaultTargetName = ProjectFileGenerator.GetProjectDefaultTargetNameForType(projectFile.Directory, targetRules.Type);
 
-						// If we don't want all targets, skip over non-defaults.
-						if (!bIncludeAllTargets && bIsDefaultTarget.HasValue && !bIsDefaultTarget.Value)
+						// GetProjectDefaultTargetNameForType returns
+						if (defaultTargetName != null)
 						{
-							continue;
+							bIsDefaultTarget = defaultTargetName == targetName;
 						}
 					}
 
-					// Get the path to the target
-					FileReference? path = Assembly.GetTargetFileName(TargetName);
+					// If we don't want all targets, skip over non-defaults.
+					if (!bIncludeAllTargets && bIsDefaultTarget.HasValue && !bIsDefaultTarget.Value)
+					{
+						return;
+					}
+				}
+
+				// Get the path to the target
+				FileReference? path = assembly.GetTargetFileName(targetName);
 
+				entries.Add((targetName, targetRules.Type, path, bIsDefaultTarget));
+			});
+
+			// Generate json
+			using System.IO.StringWriter stringWriter = new();
+			using (JsonWriter writer = new(stringWriter))
+			{
+				writer.WriteObjectStart();
+				writer.WriteArrayStart("Targets");
+				foreach ((string name, TargetType type, FileReference? path, bool? isDefaultTarget) in entries.OrderBy(x => x.path))
+				{
 					// Write the target info
-					Writer.WriteObjectStart();
-					Writer.WriteValue("Name", TargetName);
+					writer.WriteObjectStart();
+					writer.WriteValue("Name", name);
 					if (path != null)
 					{
-						Writer.WriteValue("Path", path.MakeRelativeTo(OutputFile.Directory));
+						writer.WriteValue("Path", path.MakeRelativeTo(outputFile.Directory));
 					}
-					Writer.WriteValue("Type", TargetRules.Type.ToString());
+					writer.WriteValue("Type", type.ToString());
 
-					if (bIncludeAllTargets && bIsDefaultTarget.HasValue)
+					if (bIncludeAllTargets && isDefaultTarget.HasValue)
 					{
-						Writer.WriteValue("DefaultTarget", bIsDefaultTarget.Value);
+						writer.WriteValue("DefaultTarget", isDefaultTarget.Value);
 					}
-					Writer.WriteObjectEnd();
+					writer.WriteObjectEnd();
 				}
-				Writer.WriteArrayEnd();
-				Writer.WriteObjectEnd();
+				writer.WriteArrayEnd();
+				writer.WriteObjectEnd();
 			}
 
 			// Write the output file
-			DirectoryReference.CreateDirectory(OutputFile.Directory);
-			FileReference.WriteAllTextIfDifferent(OutputFile, StringWriter.ToString());
+			DirectoryReference.CreateDirectory(outputFile.Directory);
+			FileReference.WriteAllTextIfDifferent(outputFile, stringWriter.ToString());
 			// even if the contents don't change, we need to update the last write time (this is safer to do in parallel, unlike WriteAllText)
-			FileReference.SetLastWriteTimeUtc(OutputFile, DateTime.UtcNow);
+			FileReference.SetLastWriteTimeUtc(outputFile, DateTime.UtcNow);
 		}
 	}
 }
