UnrealBuild: Add 5.8 deprecations to existing Obsolete attributes with no deprecation version.

#Synced-CL 51408971

diff --git a/Engine/Source/Programs/UnrealBuildTool/Configuration/BuildConfiguration.cs b/Engine/Source/Programs/UnrealBuildTool/Configuration/BuildConfiguration.cs
--- a/Engine/Source/Programs/UnrealBuildTool/Configuration/BuildConfiguration.cs
+++ b/Engine/Source/Programs/UnrealBuildTool/Configuration/BuildConfiguration.cs
@@ -37,7 +37,7 @@ namespace UnrealBuildTool
 		/// Whether the hybrid executor will be used (a remote executor and local executor). No longer supported.
 		/// </summary>
 		[XmlConfigFile(Deprecated = true)]
-		[Obsolete("HybridExecutor is no longer supported")]
+		[Obsolete("Deprecated in UE5.8. HybridExecutor is no longer supported")]
 		public bool bAllowHybridExecutor { get; set; } = false;
 
 		/// <summary>
@@ -58,7 +58,7 @@ namespace UnrealBuildTool
 		/// Whether the UnrealBuildAccelerator (local only) executor will be used.
 		/// </summary>
 		[XmlConfigFile(Deprecated = true)]
-		[Obsolete("Replace with bAllowUBAExecutor")]
+		[Obsolete("Deprecated in UE5.8. Replace with bAllowUBAExecutor")]
 		public bool bAllowUBALocalExecutor
 		{
 			get => bAllowUBAExecutor;
@@ -143,7 +143,7 @@ namespace UnrealBuildTool
 		/// <summary>
 		/// If true, do not build UHT, assume it is already built.
 		/// </summary>
-		[Obsolete]
+		[Obsolete("Deprecated in UE5.8.")]
 		[XmlConfigFile(Category = "UEBuildConfiguration", Deprecated = true)]
 		public bool bDoNotBuildUHT { get; set; } = false;
 
@@ -164,21 +164,21 @@ namespace UnrealBuildTool
 		/// <summary>
 		/// If true, the Debug version of UnrealHeaderTool will be built and run instead of the Development version.
 		/// </summary>
-		[Obsolete]
+		[Obsolete("Deprecated in UE5.8.")]
 		[XmlConfigFile(Category = "UEBuildConfiguration", Deprecated = true)]
 		public bool bForceDebugUnrealHeaderTool { get; set; } = false;
 
 		/// <summary>
 		/// If true, use C# UHT internal to UBT
 		/// </summary>
-		[Obsolete]
+		[Obsolete("Deprecated in UE5.8.")]
 		[XmlConfigFile(Category = "UEBuildConfiguration", Deprecated = true)]
 		public bool bUseBuiltInUnrealHeaderTool { get; set; } = true;
 
 		/// <summary>
 		/// If true, generate warnings when C++ UHT is used
 		/// </summary>
-		[Obsolete]
+		[Obsolete("Deprecated in UE5.8.")]
 		[XmlConfigFile(Category = "UEBuildConfiguration", Deprecated = true)]
 		public bool bWarnOnCppUnrealHeaderTool { get; set; } = true;
 
diff --git a/Engine/Source/Programs/UnrealBuildTool/Configuration/Rules/TestModuleRules.cs b/Engine/Source/Programs/UnrealBuildTool/Configuration/Rules/TestModuleRules.cs
--- a/Engine/Source/Programs/UnrealBuildTool/Configuration/Rules/TestModuleRules.cs
+++ b/Engine/Source/Programs/UnrealBuildTool/Configuration/Rules/TestModuleRules.cs
@@ -551,7 +551,7 @@ namespace UnrealBuildTool
 			public HashSet<UnrealTargetPlatform> SupportedPlatforms
 			{
 				get;
-				[Obsolete("Modify the existing collection")]
+				[Obsolete("Deprecated in UE5.8. Modify the existing collection")]
 				set;
 			} = [UnrealTargetPlatform.Win64];
 
@@ -561,7 +561,7 @@ namespace UnrealBuildTool
 			public Dictionary<UnrealTargetPlatform, string> PlatformTags
 			{
 				get;
-				[Obsolete("Modify the existing collection")]
+				[Obsolete("Deprecated in UE5.8. Modify the existing collection")]
 				set;
 			} = [];
 
@@ -571,7 +571,7 @@ namespace UnrealBuildTool
 			public Dictionary<UnrealTargetPlatform, string> PlatformGauntletArgs
 			{
 				get;
-				[Obsolete("Modify the existing collection")]
+				[Obsolete("Deprecated in UE5.8. Modify the existing collection")]
 				set;
 			} = [];
 
@@ -581,7 +581,7 @@ namespace UnrealBuildTool
 			public Dictionary<UnrealTargetPlatform, string> PlatformCompilationExtraArgs
 			{
 				get;
-				[Obsolete("Modify the existing collection")]
+				[Obsolete("Deprecated in UE5.8. Modify the existing collection")]
 				set;
 			} = [];
 
@@ -591,7 +591,7 @@ namespace UnrealBuildTool
 			public List<UnrealTargetPlatform> PlatformsRunUnsupported
 			{
 				get;
-				[Obsolete("Modify the existing collection")]
+				[Obsolete("Deprecated in UE5.8. Modify the existing collection")]
 				set;
 			} = [UnrealTargetPlatform.Android, UnrealTargetPlatform.IOS, UnrealTargetPlatform.TVOS, UnrealTargetPlatform.VisionOS];
 
@@ -601,7 +601,7 @@ namespace UnrealBuildTool
 			public Dictionary<UnrealTargetPlatform, bool> PlatformRunContainerized
 			{
 				get;
-				[Obsolete("Modify the existing collection")]
+				[Obsolete("Deprecated in UE5.8. Modify the existing collection")]
 				set;
 			} = [];
 #pragma warning restore CA2227
diff --git a/Engine/Source/Programs/UnrealBuildTool/Configuration/UEBuildPlatform.cs b/Engine/Source/Programs/UnrealBuildTool/Configuration/UEBuildPlatform.cs
--- a/Engine/Source/Programs/UnrealBuildTool/Configuration/UEBuildPlatform.cs
+++ b/Engine/Source/Programs/UnrealBuildTool/Configuration/UEBuildPlatform.cs
@@ -748,7 +748,7 @@ namespace UnrealBuildTool
 		/// <summary>
 		/// Returns the delimiter used to separate paths in the PATH environment variable for the platform we are executing on.
 		/// </summary>
-		[Obsolete("Replace with System.IO.Path.PathSeparator")]
+		[Obsolete("Deprecated in UE5.8. Replace with System.IO.Path.PathSeparator")]
 		public static string GetPathVarDelimiter()
 		{
 			return Path.PathSeparator.ToString();
diff --git a/Engine/Source/Programs/UnrealBuildTool/Executors/UnrealBuildAccelerator/UnrealBuildAcceleratorHordeConfig.cs b/Engine/Source/Programs/UnrealBuildTool/Executors/UnrealBuildAccelerator/UnrealBuildAcceleratorHordeConfig.cs
--- a/Engine/Source/Programs/UnrealBuildTool/Executors/UnrealBuildAccelerator/UnrealBuildAcceleratorHordeConfig.cs
+++ b/Engine/Source/Programs/UnrealBuildTool/Executors/UnrealBuildAccelerator/UnrealBuildAcceleratorHordeConfig.cs
@@ -35,7 +35,7 @@ namespace UnrealBuildTool
 		/// <summary>
 		/// OIDC id for the login to use
 		/// </summary>
-		[Obsolete("The OidcProvider option is no longer used")]
+		[Obsolete("Deprecated in UE5.8. The OidcProvider option is no longer used")]
 		[XmlConfigFile(Category = "Horde", Name = "OidcProvider")]
 		[CommandLine("-UBAHordeOidc=")]
 		public string? HordeOidcProvider { get; set; }
diff --git a/Engine/Source/Programs/UnrealBuildTool/Modes/ServerMode.cs b/Engine/Source/Programs/UnrealBuildTool/Modes/ServerMode.cs
--- a/Engine/Source/Programs/UnrealBuildTool/Modes/ServerMode.cs
+++ b/Engine/Source/Programs/UnrealBuildTool/Modes/ServerMode.cs
@@ -124,8 +124,7 @@ namespace UnrealBuildTool
 		public List<TargetIntellisenseInfo.CompileSettings?> Settings { get; set; } = new();
 	}
 
-	// TODO: Deprecate 
-	[Obsolete]
+	[Obsolete("Deprecated in UE5.8.")]
 	[ToolMode("Server", ToolModeOptions.BuildPlatforms | ToolModeOptions.XmlConfig | ToolModeOptions.UseStartupTraceListener)]
 	class ServerMode : ToolMode
 	{
diff --git a/Engine/Source/Programs/UnrealBuildTool/Platform/Apple/AppleExports.cs b/Engine/Source/Programs/UnrealBuildTool/Platform/Apple/AppleExports.cs
--- a/Engine/Source/Programs/UnrealBuildTool/Platform/Apple/AppleExports.cs
+++ b/Engine/Source/Programs/UnrealBuildTool/Platform/Apple/AppleExports.cs
@@ -17,9 +17,9 @@ namespace UnrealBuildTool
 	{
 #pragma warning disable CA1823 // unused private fields - This will be removed in a future version
 #pragma warning disable IDE0052 //  unused private member - This will be removed in a future version
-		[Obsolete("Modern Xcode is now the only option")]
+		[Obsolete("Deprecated in UE5.8. Modern Xcode is now the only option")]
 		private static readonly bool bForceModernXcode = Environment.CommandLine.Contains("-modernxcode", StringComparison.OrdinalIgnoreCase);
-		[Obsolete("Legacy Xcode is no longer supported")]
+		[Obsolete("Deprecated in UE5.8. Legacy Xcode is no longer supported")]
 		private static readonly bool bForceLegacyXcode = Environment.CommandLine.Contains("-legacy xcode", StringComparison.OrdinalIgnoreCase);
 #pragma warning restore IDE0052
 #pragma warning restore CA1823
@@ -42,7 +42,7 @@ namespace UnrealBuildTool
 		/// </summary>
 		/// <param name="ProjectFile"></param>
 		/// <returns></returns>
-		[Obsolete("Use !CreatingAppOnWindows")]
+		[Obsolete("Deprecated in UE5.8. Use !CreatingAppOnWindows")]
 		public static bool UseModernXcode(FileReference? ProjectFile)
 		{
 			return !CreatingAppOnWindows(ProjectFile);
diff --git a/Engine/Source/Programs/UnrealBuildTool/Platform/Apple/AppleToolChain.cs b/Engine/Source/Programs/UnrealBuildTool/Platform/Apple/AppleToolChain.cs
--- a/Engine/Source/Programs/UnrealBuildTool/Platform/Apple/AppleToolChain.cs
+++ b/Engine/Source/Programs/UnrealBuildTool/Platform/Apple/AppleToolChain.cs
@@ -1237,7 +1237,7 @@ namespace UnrealBuildTool
 		[CommandLine("-XmlConfigCache=")]
 		public FileReference? XmlConfigCache = null;
 
-		[Obsolete("Modern Xcode is now the only option")] 
+		[Obsolete("Deprecated in UE5.8. Modern Xcode is now the only option")] 
 		[CommandLine("-ModernXcode")]
 		public bool bModernXcode = true;
 
diff --git a/Engine/Source/Programs/UnrealBuildTool/Platform/Windows/UEBuildWindows.cs b/Engine/Source/Programs/UnrealBuildTool/Platform/Windows/UEBuildWindows.cs
--- a/Engine/Source/Programs/UnrealBuildTool/Platform/Windows/UEBuildWindows.cs
+++ b/Engine/Source/Programs/UnrealBuildTool/Platform/Windows/UEBuildWindows.cs
@@ -506,7 +506,7 @@ namespace UnrealBuildTool
 		/// Enable building with the Win10 SDK instead of the older Win8.1 SDK 
 		/// </summary>
 		[ConfigFile(ConfigHierarchyType.Engine, "/Script/WindowsTargetPlatform.WindowsTargetSettings", "bUseWindowsSDK10")]
-		[Obsolete("Set MainVersion in Windows_SDK.json")]
+		[Obsolete("Deprecated in UE5.8. Set MainVersion in Windows_SDK.json")]
 		public bool bUseWindowsSDK10 = false;
 
 		/// <summary>
diff --git a/Engine/Source/Programs/UnrealBuildTool/ProjectFiles/ProjectFileGenerator.cs b/Engine/Source/Programs/UnrealBuildTool/ProjectFiles/ProjectFileGenerator.cs
--- a/Engine/Source/Programs/UnrealBuildTool/ProjectFiles/ProjectFileGenerator.cs
+++ b/Engine/Source/Programs/UnrealBuildTool/ProjectFiles/ProjectFileGenerator.cs
@@ -2721,7 +2721,7 @@ namespace UnrealBuildTool
 			}
 		}
 
-		[Obsolete("Replace with 3 paramater version of AddTargetForIntellisense")]
+		[Obsolete("Deprecated in UE5.8. Replace with 3 paramater version of AddTargetForIntellisense")]
 		protected virtual void AddTargetForIntellisense(UEBuildTarget Target, ILogger Logger)
 		{
 			AddTargetForIntellisense(Target, [], Logger);
diff --git a/Engine/Source/Programs/UnrealBuildTool/System/CppCompileEnvironment.cs b/Engine/Source/Programs/UnrealBuildTool/System/CppCompileEnvironment.cs
--- a/Engine/Source/Programs/UnrealBuildTool/System/CppCompileEnvironment.cs
+++ b/Engine/Source/Programs/UnrealBuildTool/System/CppCompileEnvironment.cs
@@ -31,13 +31,13 @@ namespace UnrealBuildTool
 		/// <summary>
 		/// Obsolete C++14 standard. No longer allowed, code will not compile with this version.
 		/// </summary>
-		[Obsolete("C++14 is no longer allowed, code will not compile with this version")]
+		[Obsolete("Deprecated in UE5.8. C++14 is no longer allowed, code will not compile with this version")]
 		Cpp14,
 
 		/// <summary>
 		/// Obsolete C++17 standard. No longer allowed, code will not compile with this version.
 		/// </summary>
-		[Obsolete("C++17 is no longer allowed, code will not compile with this version")]
+		[Obsolete("Deprecated in UE5.8. C++17 is no longer allowed, code will not compile with this version")]
 		Cpp17,
 
 		/// <summary>
diff --git a/Engine/Source/Programs/UnrealBuildTool/System/PlatformExports.cs b/Engine/Source/Programs/UnrealBuildTool/System/PlatformExports.cs
--- a/Engine/Source/Programs/UnrealBuildTool/System/PlatformExports.cs
+++ b/Engine/Source/Programs/UnrealBuildTool/System/PlatformExports.cs
@@ -104,7 +104,7 @@ namespace UnrealBuildTool
 		/// <param name="Platform">Platform to check</param>
 		/// <param name="Logger">Logger for output</param>
 		/// <returns>True if the platform supports the parallel executor in UAT</returns>
-		[Obsolete("All platforms support the parallel executor")]
+		[Obsolete("Deprecated in UE5.8. All platforms support the parallel executor")]
 		public static bool CanUseParallelExecutor(UnrealTargetPlatform Platform, ILogger Logger)
 		{
 			return true;
diff --git a/Engine/Source/Programs/UnrealBuildTool/System/Utils.cs b/Engine/Source/Programs/UnrealBuildTool/System/Utils.cs
--- a/Engine/Source/Programs/UnrealBuildTool/System/Utils.cs
+++ b/Engine/Source/Programs/UnrealBuildTool/System/Utils.cs
@@ -950,7 +950,7 @@ namespace UnrealBuildTool
 		/// NOTE: This function may return null. Some accounts (eg. the SYSTEM account on Windows) do not have a personal folder, and Jenkins
 		/// runs using this account by default.
 		/// </summary>
-		[Obsolete("Replace with Unreal.UserSettingDirectory")]
+		[Obsolete("Deprecated in UE5.8. Replace with Unreal.UserSettingDirectory")]
 		public static DirectoryReference? GetUserSettingDirectory() => Unreal.UserSettingDirectory;
 
 		enum LOGICAL_PROCESSOR_RELATIONSHIP
