UnrealBuildTool: Minor opt to QueryTargets

#rnx

#Synced-CL 51388341

diff --git a/Engine/Source/Programs/UnrealBuildTool/Modes/QueryTargetsMode.cs b/Engine/Source/Programs/UnrealBuildTool/Modes/QueryTargetsMode.cs
--- a/Engine/Source/Programs/UnrealBuildTool/Modes/QueryTargetsMode.cs
+++ b/Engine/Source/Programs/UnrealBuildTool/Modes/QueryTargetsMode.cs
@@ -37,7 +37,7 @@ namespace UnrealBuildTool
 		bool bIncludeProgramTargets = true;
 
 		/// <summary>
-		/// Write out all targets, even if a default is specified in the BuildSettings section of the Default*.ini files. 
+		/// Write out all targets, even if a default is specified in the BuildSettings section of the Default*.ini files.
 		/// </summary>
 		[CommandLine("-IncludeAllTargets")]
 		bool bIncludeAllTargets = false;
@@ -91,9 +91,7 @@ namespace UnrealBuildTool
 		/// <returns>Path to the output file</returns>
 		public static FileReference GetDefaultOutputFile(FileReference? ProjectFile)
 		{
-			return ProjectFile == null
-				? FileReference.Combine(Unreal.EngineDirectory, "Intermediate", "TargetInfo.json")
-				: FileReference.Combine(ProjectFile.Directory, "Intermediate", "TargetInfo.json");
+			return FileReference.Combine(ProjectFile?.Directory ?? Unreal.EngineDirectory, "Intermediate", "TargetInfo.json");
 		}
 
 		/// <summary>
@@ -114,11 +112,12 @@ namespace UnrealBuildTool
 			List<string> targetNames = new List<string>();
 			assembly.GetAllTargetNames(targetNames, bIncludeParentAssembly);
 
+			ConcurrentDictionary<TargetType, string?> defaultTargetNameCache = [];
 			ConcurrentBag<(string name, TargetType type, FileReference? path, bool? isDefaultTarget)> entries = [];
 			Parallel2.ForEach(targetNames, (targetName) =>
 			{
 				// skip target rules that are platform extension or platform group specializations
-				string[] targetPathSplit = targetName.Split(new char[] { '_' }, StringSplitOptions.RemoveEmptyEntries);
+				string[] targetPathSplit = targetName.Split('_', StringSplitOptions.RemoveEmptyEntries);
 				if (targetPathSplit.Length > 1 && (UnrealTargetPlatform.IsValidName(targetPathSplit.Last()) || UnrealPlatformGroup.IsValidName(targetPathSplit.Last())))
 				{
 					return;
@@ -153,9 +152,9 @@ namespace UnrealBuildTool
 				{
 					if (projectFile != null)
 					{
-						string? defaultTargetName = ProjectFileGenerator.GetProjectDefaultTargetNameForType(projectFile.Directory, targetRules.Type);
-
-						// GetProjectDefaultTargetNameForType returns
+						string? defaultTargetName = defaultTargetNameCache.GetOrAdd(
+							targetRules.Type,
+							t => ProjectFileGenerator.GetProjectDefaultTargetNameForType(projectFile.Directory, t));
 						if (defaultTargetName != null)
 						{
 							bIsDefaultTarget = defaultTargetName == targetName;
