UnrealBuildTool: Fix double negative in warning log message.

#Synced-CL 51454738
#Synced-CL 51454738

diff --git a/Engine/Source/Programs/UnrealBuildTool/ProjectFiles/VisualStudio/VCProjectFileGenerator.cs b/Engine/Source/Programs/UnrealBuildTool/ProjectFiles/VisualStudio/VCProjectFileGenerator.cs
--- a/Engine/Source/Programs/UnrealBuildTool/ProjectFiles/VisualStudio/VCProjectFileGenerator.cs
+++ b/Engine/Source/Programs/UnrealBuildTool/ProjectFiles/VisualStudio/VCProjectFileGenerator.cs
@@ -406,7 +406,7 @@ namespace UnrealBuildTool
 									}
 									if (!DowngradeAvailable)
 									{
-										Logger.LogInformation("Generated solution cannot be downgraded as no prior Visual Studio version is not installed. Please install the prior version of Visual Studio if {SupportedPlatform} SDK support is required.", SupportedPlatform);
+										Logger.LogInformation("Generated solution cannot be downgraded as no prior Visual Studio version is installed. Please install the prior version of Visual Studio if {SupportedPlatform} SDK support is required.", SupportedPlatform);
 									}
 									else
 									{
