UnrealBuildTool: Update to 5.6 being the lowest supported EngineIncludeOrderVersion.

#Synced-CL 51454738

#Synced-CL 51454738

diff --git a/Engine/Source/Programs/UnrealBuildTool/Configuration/Rules/TargetRules.cs b/Engine/Source/Programs/UnrealBuildTool/Configuration/Rules/TargetRules.cs
--- a/Engine/Source/Programs/UnrealBuildTool/Configuration/Rules/TargetRules.cs
+++ b/Engine/Source/Programs/UnrealBuildTool/Configuration/Rules/TargetRules.cs
@@ -245,7 +245,7 @@ namespace UnrealBuildTool
 		/// Do not delete old enum entries to prevent breaking project generation
 		/// </summary>
 #pragma warning disable CS0618 // Type or member is obsolete
-		Oldest = Unreal5_4,
+		Oldest = Unreal5_6,
 #pragma warning restore CS0618 // Type or member is obsolete
 	}
 #pragma warning restore CA1707
@@ -475,8 +475,6 @@ namespace UnrealBuildTool
 			EngineIncludeOrderVersion.Unreal5_7 => "UE_ENABLE_INCLUDE_ORDER_DEPRECATED_IN_5_7",
 #pragma warning disable CS0618 // Type or member is obsolete
 			EngineIncludeOrderVersion.Unreal5_6 => "UE_ENABLE_INCLUDE_ORDER_DEPRECATED_IN_5_6",
-			EngineIncludeOrderVersion.Unreal5_5 => "UE_ENABLE_INCLUDE_ORDER_DEPRECATED_IN_5_5",
-			EngineIncludeOrderVersion.Unreal5_4 => "UE_ENABLE_INCLUDE_ORDER_DEPRECATED_IN_5_4",
 #pragma warning restore CS0618 // Type or member is obsolete
 			_ => throw new BuildLogEventException("Unsupported EngineIncludeOrderVersion {Value}", inVersion)
 		};
