UnrealBuildTool: Add min AGDE vesion for VS2026

#rnx
#Synced-CL 51442092

diff --git a/Engine/Source/Programs/UnrealBuildTool/Platform/Android/UEBuildAndroid.cs b/Engine/Source/Programs/UnrealBuildTool/Platform/Android/UEBuildAndroid.cs
--- a/Engine/Source/Programs/UnrealBuildTool/Platform/Android/UEBuildAndroid.cs
+++ b/Engine/Source/Programs/UnrealBuildTool/Platform/Android/UEBuildAndroid.cs
@@ -249,20 +249,27 @@ namespace UnrealBuildTool
 		{
 			IEnumerable<VisualStudioInstallation> Installations = MicrosoftPlatformSDK.FindVisualStudioInstallations(Logger);
 
-			bool vsiInstalled = Installations
+			Version? version = Installations
 				.Where(x => x.Compiler == WindowsCompiler.VisualStudio2022)
 				.SelectMany(x => x.Packages.Where(x => x.Key.StartsWith("Component.GoogleAndroid", StringComparison.Ordinal)).Select(x => x.Value.Version))
-				.Any();
+				.Max();
 
-			if (vsiInstalled)
+			if (version != null && version < GetVersionRequiredForVisualStudio(VCProjectFileFormat.VisualStudio2026))
 			{
-				// Currently unknown when 2026 will be supported and at which version
 				return VCProjectFileFormat.VisualStudio2022;
 			}
 
 			return VCProjectFileFormat.Default;
 		}
 
+		/// <inheritdoc/>
+		public override Version GetVersionRequiredForVisualStudio(VCProjectFileFormat Format)
+		{
+			return Format >= VCProjectFileFormat.VisualStudio2026
+				? new Version(26, 1, 102)
+				: new Version();
+		}
+
 		public override bool IsBuildProduct(string FileName, string[] NamePrefixes, string[] NameSuffixes)
 		{
 			return IsBuildProductWithArch(FileName, NamePrefixes, NameSuffixes, ".so")
