Include Verse files in GenerateProjectFiles.

We now include `Verse/` directories, as well as `Content/*` and `Content/Verse/**`. This will make Verse easier to discover within an IDE.

#Synced-CL 51288979

diff --git a/Engine/Source/Programs/UnrealBuildTool/ProjectFiles/ProjectFileGenerator.cs b/Engine/Source/Programs/UnrealBuildTool/ProjectFiles/ProjectFileGenerator.cs
--- a/Engine/Source/Programs/UnrealBuildTool/ProjectFiles/ProjectFileGenerator.cs
+++ b/Engine/Source/Programs/UnrealBuildTool/ProjectFiles/ProjectFileGenerator.cs
@@ -2082,6 +2082,11 @@ namespace UnrealBuildTool
 				{
 					GameProject.AddFilesToProject(SourceFileSearch.FindFiles(GameShaderDirectory), GameProjectDirectory);
 				}
+
+				foreach (DirectoryReference GameVerseDirectory in Unreal.GetExtensionDirs(GameProjectDirectory, "Verse"))
+				{
+					GameProject.AddFilesToProject(SourceFileSearch.FindFiles(GameVerseDirectory), GameProjectDirectory);
+				}
 			});
 		}
 
@@ -3108,6 +3113,50 @@ namespace UnrealBuildTool
 					ProjectFile.AddFilesToProject(files, BaseFolder);
 				}
 			}
+
+			// Add plugin Verse files from the Content directory if we have any.
+			// We filter by extension here rather than using FindFiles because Content
+			// directories can contain large binary assets that should not be added.
+			if (PluginDir.TryGetDirectory("Content", out DirectoryItem? PluginContentFolder))
+			{
+				List<FileReference> files = FindVerseFilesInDirectory(PluginContentFolder);
+				if (files.Count > 0)
+				{
+					lock (ProjectFile)
+					{
+						ProjectFile.AddFilesToProject(files, BaseFolder);
+					}
+				}
+			}
+		}
+
+		private static readonly string[] VerseFileExtensions = [".verse", ".versetest", ".vmodule", ".vproject"];
+
+		private static List<FileReference> FindVerseFilesInDirectory(DirectoryItem Directory, bool bFullRecurse = false)
+		{
+			List<FileReference> Result = [];
+			foreach (FileItem File in Directory.EnumerateFiles())
+			{
+				foreach (string Ext in VerseFileExtensions)
+				{
+					if (File.Name.EndsWith(Ext, StringComparison.OrdinalIgnoreCase))
+					{
+						Result.Add(File.Location);
+						break;
+					}
+				}
+			}
+			foreach (DirectoryItem SubDir in Directory.EnumerateDirectories())
+			{
+				// When scanning from a Content root, only descend into subdirectories
+				// named "Verse" to avoid traversing large asset trees (Textures, Materials,
+				// etc.). Once inside a Verse subtree, recurse fully.
+				if (bFullRecurse || SubDir.Name.Equals("Verse", StringComparison.OrdinalIgnoreCase))
+				{
+					Result.AddRange(FindVerseFilesInDirectory(SubDir, bFullRecurse: true));
+				}
+			}
+			return Result;
 		}
 
 		private ProjectFile FindOrAddProjectHelper(string InProjectFileNameBase, DirectoryReference InBaseFolder)
