UnrealBuildTool: Remove unintended vs install filter

#rnx
#Synced-CL 51442092

diff --git a/Engine/Source/Programs/UnrealBuildTool/Platform/Android/UEBuildAndroid.cs b/Engine/Source/Programs/UnrealBuildTool/Platform/Android/UEBuildAndroid.cs
--- a/Engine/Source/Programs/UnrealBuildTool/Platform/Android/UEBuildAndroid.cs
+++ b/Engine/Source/Programs/UnrealBuildTool/Platform/Android/UEBuildAndroid.cs
@@ -250,7 +250,6 @@ namespace UnrealBuildTool
 			IEnumerable<VisualStudioInstallation> Installations = MicrosoftPlatformSDK.FindVisualStudioInstallations(Logger);
 
 			Version? version = Installations
-				.Where(x => x.Compiler == WindowsCompiler.VisualStudio2022)
 				.SelectMany(x => x.Packages.Where(x => x.Key.StartsWith("Component.GoogleAndroid", StringComparison.Ordinal)).Select(x => x.Value.Version))
 				.Max();
 
