[iOS] Fix for an issue where in blueprint only projects, the version config is not added to the packaged version causing errors when you try to upload the app to test flight.

The problem is that in BP only projects, we compile Unreal Game not the project itself and in that path the generated version file is on the Engine directory not the Project directory.

This CL makes it so when the Xcode is generated, it will include both the engine and project version of the files using ?. if both files exist, the last one will override the previous one, if one does not exist, it just skips it silently.

#Synced-CL 51440721

diff --git a/Engine/Source/Programs/UnrealBuildTool/ProjectFiles/Xcode/XcodeProject.cs b/Engine/Source/Programs/UnrealBuildTool/ProjectFiles/Xcode/XcodeProject.cs
--- a/Engine/Source/Programs/UnrealBuildTool/ProjectFiles/Xcode/XcodeProject.cs
+++ b/Engine/Source/Programs/UnrealBuildTool/ProjectFiles/Xcode/XcodeProject.cs
@@ -2169,8 +2169,18 @@ namespace UnrealBuildTool.XcodeProjectXcconfig
 			XcodeUtils.FindPlistId(PlistMetadata, "LSApplicationCategoryType", ref AppCategory);
 
 			// include another xcconfig for versions that UBT writes out
-			Xcconfig!.AppendLine($"#include? \"{ProjectOrEngineDir}/Intermediate/Build/Versions.xcconfig\"");
-			if (UnrealData.UProjectFileLocation == null)
+
+			// Blueprint-only projects have Versions.xcconfig in the engine dir, not the project dir because it technically builds UnrealGame.
+			// If the project directory has a valid version file (like in code projects), that will be used as the last declared one overrides everything else. 
+			// If any of them do not exist, they are skipped. This is why we add both here with "?"
+			
+			Xcconfig!.AppendLine($"#include? \"{Unreal.EngineDirectory}/Intermediate/Build/Versions.xcconfig\"");
+
+			if (UnrealData.UProjectFileLocation != null)
+			{
+				Xcconfig!.AppendLine($"#include? \"{UnrealData.UProjectFileLocation.Directory}/Intermediate/Build/Versions.xcconfig\"");
+			}
+			else
 			{
 				Xcconfig!.AppendLine($"#include? \"{XcodeProjectFileGenerator.ContentOnlySettingsFile}\"");
 			}
@@ -2670,8 +2680,18 @@ namespace UnrealBuildTool.XcodeProjectXcconfig
 			XcodeUtils.FindPlistId(PlistMetadata, "LSApplicationCategoryType", ref AppCategory);
 
 			// include another xcconfig for versions that UBT writes out
-			Xcconfig!.AppendLine($"#include? \"{ProjectOrEngineDir}/Intermediate/Build/Versions.xcconfig\"");
-			if (UnrealData.UProjectFileLocation == null)
+	
+			// Blueprint-only projects have Versions.xcconfig in the engine dir, not the project dir
+			// If the project directory has a valid version file (like in code projects), that will be used as the last declared one overrides everything else. 
+			// If any of them do not exist, they are skipped. This is why we add both here with "?"
+			
+			Xcconfig!.AppendLine($"#include? \"{Unreal.EngineDirectory}/Intermediate/Build/Versions.xcconfig\"");
+
+			if (UnrealData.UProjectFileLocation != null)
+			{
+				Xcconfig!.AppendLine($"#include? \"{UnrealData.UProjectFileLocation.Directory}/Intermediate/Build/Versions.xcconfig\"");
+			}
+			else
 			{
 				Xcconfig!.AppendLine($"#include? \"{XcodeProjectFileGenerator.ContentOnlySettingsFile}\"");
 			}
