[Mac/iOS] Fix for issue where Apple platforms were maked valid in windows even if the require sdk to use the remote mac build workflow is not installed.

This CL is basically reverting the change made on CL 49184814 that was made to make build machines and dev machines match behaviour wise. Build machines bypass the sdk check because they din't need it,  but the problem is that dev machines do need the itunes sdk.

#Synced-CL 51524632

diff --git a/Engine/Source/Programs/UnrealBuildTool/Platform/Apple/ApplePlatformSDK.cs b/Engine/Source/Programs/UnrealBuildTool/Platform/Apple/ApplePlatformSDK.cs
--- a/Engine/Source/Programs/UnrealBuildTool/Platform/Apple/ApplePlatformSDK.cs
+++ b/Engine/Source/Programs/UnrealBuildTool/Platform/Apple/ApplePlatformSDK.cs
@@ -8,6 +8,7 @@ using System.Text.RegularExpressions;
 using EpicGames.Core;
 using Microsoft.Extensions.Logging;
 using Microsoft.Win32;
+using UnrealBuildBase;
 
 ///////////////////////////////////////////////////////////////////
 // If you are looking for supported version numbers, look in the
@@ -269,7 +270,7 @@ namespace UnrealBuildTool
 
 			// iTunes is technically only need to deploy to and run on connected devices.
 			// This code removes requirement for Windows builders to have Xcode installed.
-			if (Status == SDKStatus.Invalid && !OperatingSystem.IsMacOS() && true) // Unreal.IsBuildMachine())
+			if (Status == SDKStatus.Invalid && !OperatingSystem.IsMacOS() && Unreal.IsBuildMachine())
 			{
 				Status = SDKStatus.Valid;
 			}
