Pass the CustomConfig parameter when loading the config hierarchy in SetDefaultsForCookedEditor.

#rnx
#Synced-CL 51534191

diff --git a/Engine/Source/Programs/UnrealBuildTool/Configuration/Rules/TargetRules.cs b/Engine/Source/Programs/UnrealBuildTool/Configuration/Rules/TargetRules.cs
--- a/Engine/Source/Programs/UnrealBuildTool/Configuration/Rules/TargetRules.cs
+++ b/Engine/Source/Programs/UnrealBuildTool/Configuration/Rules/TargetRules.cs
@@ -3504,7 +3504,7 @@ namespace UnrealBuildTool
 			// this setting can be used to compile out the shader compiler if that is important
 			//GlobalDefinitions.Add("UE_ALLOW_SHADER_COMPILING=0");
 
-			ConfigHierarchy projectGameIni = ConfigCache.ReadHierarchy(ConfigHierarchyType.Game, ProjectFile?.Directory, Platform);
+			ConfigHierarchy projectGameIni = ConfigCache.ReadHierarchy(ConfigHierarchyType.Game, ProjectFile?.Directory, Platform, CustomConfig);
 			List<string>? disabledPlugins;
 			List<string> allDisabledPlugins = [];
 
