UnrealBuildTool: Include vcxproj project guids in VS slnx files such that Visual Studio doesn't mark them as dirty on load.

#Synced-CL 51580970

diff --git a/Engine/Source/Programs/UnrealBuildTool/ProjectFiles/VisualStudio/VCProjectFileGenerator.cs b/Engine/Source/Programs/UnrealBuildTool/ProjectFiles/VisualStudio/VCProjectFileGenerator.cs
--- a/Engine/Source/Programs/UnrealBuildTool/ProjectFiles/VisualStudio/VCProjectFileGenerator.cs
+++ b/Engine/Source/Programs/UnrealBuildTool/ProjectFiles/VisualStudio/VCProjectFileGenerator.cs
@@ -7,10 +7,14 @@ using System.Linq;
 using System.Runtime.Versioning;
 using System.Security.Cryptography;
 using System.Text;
+using System.Threading;
+using System.Threading.Tasks;
 using System.Xml;
 using System.Xml.Linq;
 using EpicGames.Core;
 using Microsoft.Extensions.Logging;
+using Microsoft.VisualStudio.SolutionPersistence.Model;
+using Microsoft.VisualStudio.SolutionPersistence.Serializer;
 using UnrealBuildBase;
 using UnrealBuildTool.ProjectFiles.VisualStudio;
 
@@ -497,7 +501,7 @@ namespace UnrealBuildTool
 					// Generate a slnx for VS2026 until properly supported
 					if (success && Environment.Version.Major >= 10)
 					{
-						ConvertSolutionFileToSlnx(SolutionFileName, SolutionFilePath, Logger);
+						ConvertSolutionFileToSlnx(SolutionFileName, SolutionFilePath, Logger).Wait();
 					}
 				}
 			}
@@ -849,7 +853,7 @@ namespace UnrealBuildTool
 				// Generate a slnx for VS2026 until properly supported
 				if (bSuccess && Settings.ProjectFileFormat == VCProjectFileFormat.VisualStudio2026 && Environment.Version.Major >= 10)
 				{
-					bSuccess = ConvertSolutionFileToSlnx(SolutionFileName, SolutionFilePath, Logger);
+					bSuccess = ConvertSolutionFileToSlnx(SolutionFileName, SolutionFilePath, Logger).GetAwaiter().GetResult();
 				}
 			}
 
@@ -868,7 +872,7 @@ namespace UnrealBuildTool
 			return bSuccess;
 		}
 
-		private static bool ConvertSolutionFileToSlnx(string SolutionFileName, FileReference SolutionFilePath, ILogger Logger)
+		private async Task<bool> ConvertSolutionFileToSlnx(string SolutionFileName, FileReference SolutionFilePath, ILogger Logger)
 		{
 			FileReference TempSolutionFilePath = FileReference.Combine(new DirectoryReference(Path.GetTempPath()), SolutionFileName);
 			FileReference.Copy(SolutionFilePath, TempSolutionFilePath, true);
@@ -882,11 +886,31 @@ namespace UnrealBuildTool
 				System.Diagnostics.ProcessPriorityClass.Normal
 			);
 			process.ReadAllLines();
-			process.WaitForExit();
+			await process.WaitForExitAsync();
 			bool bSuccess = process.ExitCode == 0;
 			if (bSuccess)
 			{
-				string contents = FileReference.ReadAllText(TempSolutionFilePath.ChangeExtension(".slnx"));
+				string contents = await FileReference.ReadAllTextAsync(TempSolutionFilePath.ChangeExtension(".slnx"));
+
+				// Visual Studio marks UE .slnx files with vcxproj files as dirty, as the `dotnet sln migrate` command doesn't include project GUIDs.
+				// We add them in here.
+				List<VCProjectFile> vcProjectFiles = [.. AllProjectFiles.OfType<VCProjectFile>()];
+				if (vcProjectFiles.Count != 0)
+				{
+					SolutionModel solutionModel = await SolutionSerializers.SlnXml.OpenAsync(TempSolutionFilePath.ChangeExtension(".slnx").FullName, CancellationToken.None);
+					foreach (MSBuildProjectFile projectFile in vcProjectFiles)
+					{
+						SolutionProjectModel project = solutionModel.SolutionProjects.Single(x =>
+							FileReference.Combine(SolutionFilePath.Directory, x.FilePath) == projectFile.ProjectFilePath);
+						project.Id = projectFile.ProjectGUID;
+					}
+
+					using MemoryStream memoryStream = new();
+					await SolutionSerializers.SlnXml.SaveAsync(memoryStream, solutionModel, CancellationToken.None);
+					UTF8Encoding encoding = new UTF8Encoding(encoderShouldEmitUTF8Identifier: false);
+					contents = encoding.GetString(memoryStream.GetBuffer());
+				}
+
 				bSuccess = WriteFileIfChanged(SolutionFilePath.ChangeExtension(".slnx").FullName, contents, Logger);
 				FileReference.Delete(TempSolutionFilePath);
 				FileReference.Delete(TempSolutionFilePath.ChangeExtension(".slnx"));
diff --git a/Engine/Source/Programs/UnrealBuildTool/UnrealBuildTool.csproj b/Engine/Source/Programs/UnrealBuildTool/UnrealBuildTool.csproj
--- a/Engine/Source/Programs/UnrealBuildTool/UnrealBuildTool.csproj
+++ b/Engine/Source/Programs/UnrealBuildTool/UnrealBuildTool.csproj
@@ -71,6 +71,7 @@
     <PackageReference Include="Microsoft.CodeAnalysis.CSharp" Version="4.13.0" />
     <PackageReference Include="Microsoft.Extensions.Logging.Console" Version="8.0.1" />
     <PackageReference Include="Microsoft.VisualStudio.Setup.Configuration.Interop" Version="3.14.2075" />
+    <PackageReference Include="Microsoft.VisualStudio.SolutionPersistence" Version="1.0.52" />
     <PackageReference Include="Microsoft.Win32.Registry" Version="5.0.0" />
     <PackageReference Include="OpenTracing" Version="0.12.1" />
     <PackageReference Include="System.CodeDom" Version="8.0.0" />
