UnrealBuildTool: Include extra information in Automation solution files such that Visual Studio doesn't mark them as dirty on load.

#Synced-CL 51580970

diff --git a/Engine/Source/Programs/UnrealBuildTool/ProjectFiles/VisualStudio/VCProjectFileGenerator.cs b/Engine/Source/Programs/UnrealBuildTool/ProjectFiles/VisualStudio/VCProjectFileGenerator.cs
--- a/Engine/Source/Programs/UnrealBuildTool/ProjectFiles/VisualStudio/VCProjectFileGenerator.cs
+++ b/Engine/Source/Programs/UnrealBuildTool/ProjectFiles/VisualStudio/VCProjectFileGenerator.cs
@@ -481,6 +481,16 @@ namespace UnrealBuildTool
 
 				AddAutomationProjects(AllGameProjects, Logger);
 
+				// All projects must be C# here - if they're not, this is a bug and should be fixed.
+				IEnumerable<VCSharpProjectFile> CSharpProjects = AllProjectFiles.Cast<VCSharpProjectFile>();
+				List<string> SolutionConfigurations =
+				[.. CSharpProjects
+					.SelectMany(x => x.Configurations)
+					.Distinct(StringComparer.OrdinalIgnoreCase)
+					// VS will order the solution configurations alphabetically
+					.OrderBy(x => x, StringComparer.OrdinalIgnoreCase)
+				];
+
 				using (Timeline.ScopeEvent("Writing automation project files", Logger))
 				{
 					string SolutionFileName = $"Automation_{PrimaryProjectName}.sln";
@@ -491,6 +501,8 @@ namespace UnrealBuildTool
 					SolutionBuilder.WriteProjects();
 
 					SolutionBuilder.WriteGlobalStart();
+					SolutionBuilder.WriteSolutionConfigurationPlatformsForCSharpOnlySolution(SolutionConfigurations);
+					SolutionBuilder.WriteProjectConfigurationPlatformsForCSharpOnlySolution(SolutionConfigurations);
 					SolutionBuilder.WriteSolutionProperties();
 					SolutionBuilder.WriteNestedProjects();
 					SolutionBuilder.WriteGlobalEnd();
diff --git a/Engine/Source/Programs/UnrealBuildTool/ProjectFiles/VisualStudio/VCSharpProjectFile.cs b/Engine/Source/Programs/UnrealBuildTool/ProjectFiles/VisualStudio/VCSharpProjectFile.cs
--- a/Engine/Source/Programs/UnrealBuildTool/ProjectFiles/VisualStudio/VCSharpProjectFile.cs
+++ b/Engine/Source/Programs/UnrealBuildTool/ProjectFiles/VisualStudio/VCSharpProjectFile.cs
@@ -182,7 +182,21 @@ internal class VCSharpProjectFile : MSBuildProjectFile
 	}
 
 	/// <inheritdoc/>
-	public override MSBuildProjectContext? GetMatchingProjectContext(TargetType SolutionTarget, UnrealTargetConfiguration SolutionConfiguration, UnrealTargetPlatform SolutionPlatform, PlatformProjectGeneratorCollection PlatformProjectGenerators, UnrealArch? Architecture, ILogger Logger)
+	public override MSBuildProjectContext? GetMatchingProjectContext(
+		TargetType SolutionTarget,
+		UnrealTargetConfiguration SolutionConfiguration,
+		UnrealTargetPlatform SolutionPlatform,
+		PlatformProjectGeneratorCollection PlatformProjectGenerators,
+		UnrealArch? Architecture,
+		ILogger Logger)
+	{
+		bool bPreferBuildByDefault = SolutionTarget == TargetType.Game || SolutionTarget == TargetType.Editor;
+		return GetMatchingProjectContextForCSharp(bPreferBuildByDefault, SolutionConfiguration.ToString());
+	}
+
+	public MSBuildProjectContext GetMatchingProjectContextForCSharp(
+		bool bPreferBuildByDefault,
+		string SolutionConfiguration)
 	{
 		// Find the matching platform name
 		string ProjectPlatformName;
@@ -197,9 +211,9 @@ internal class VCSharpProjectFile : MSBuildProjectFile
 
 		// Find the matching configuration
 		string ProjectConfigurationName;
-		if (_configurations.Contains(SolutionConfiguration.ToString()))
+		if (_configurations.Contains(SolutionConfiguration))
 		{
-			ProjectConfigurationName = SolutionConfiguration.ToString();
+			ProjectConfigurationName = SolutionConfiguration;
 		}
 		else if (_configurations.Contains("Development"))
 		{
@@ -211,11 +225,7 @@ internal class VCSharpProjectFile : MSBuildProjectFile
 		}
 
 		// Figure out whether to build it by default
-		bool bBuildByDefault = ShouldBuildByDefaultForSolutionTargets;
-		if (SolutionTarget == TargetType.Game || SolutionTarget == TargetType.Editor)
-		{
-			bBuildByDefault = true;
-		}
+		bool bBuildByDefault = ShouldBuildByDefaultForSolutionTargets || bPreferBuildByDefault;
 
 		// Create the context
 		return new MSBuildProjectContext(ProjectConfigurationName, ProjectPlatformName) { bBuildByDefault = bBuildByDefault };
diff --git a/Engine/Source/Programs/UnrealBuildTool/ProjectFiles/VisualStudio/VisualStudioSolutionBuilder.cs b/Engine/Source/Programs/UnrealBuildTool/ProjectFiles/VisualStudio/VisualStudioSolutionBuilder.cs
--- a/Engine/Source/Programs/UnrealBuildTool/ProjectFiles/VisualStudio/VisualStudioSolutionBuilder.cs
+++ b/Engine/Source/Programs/UnrealBuildTool/ProjectFiles/VisualStudio/VisualStudioSolutionBuilder.cs
@@ -186,6 +186,48 @@ internal sealed class VisualStudioSolutionBuilder
 		_sb.AppendLine("	EndGlobalSection");
 	}
 
+	public void WriteSolutionConfigurationPlatformsForCSharpOnlySolution(List<string> solutionConfigurations)
+	{
+		_sb.AppendLine("	GlobalSection(SolutionConfigurationPlatforms) = preSolution");
+
+		foreach (string solutionConfiguration in solutionConfigurations)
+		{
+			string combination = $"{solutionConfiguration}|Any CPU";
+			_sb.AppendLine($"		{combination} = {combination}");
+		}
+
+		_sb.AppendLine("	EndGlobalSection");
+	}
+
+	public void WriteProjectConfigurationPlatformsForCSharpOnlySolution(List<string> solutionConfigurations)
+	{
+		_sb.AppendLine("	GlobalSection(ProjectConfigurationPlatforms) = postSolution");
+
+		foreach (VCSharpProjectFile project in _allProjectFiles.Cast<VCSharpProjectFile>())
+		{
+			foreach (string solutionConfiguration in solutionConfigurations)
+			{
+				MSBuildProjectContext projectContext = project.GetMatchingProjectContextForCSharp(bPreferBuildByDefault: true, solutionConfiguration);
+
+				// Write the solution mapping (e.g.  "{4232C52C-680F-4850-8855-DC39419B5E9B}.Debug|iOS.ActiveCfg = iOS_Debug|Win32")
+				string guid = FormatGuid(project.ProjectGUID);
+				string configAndPlatformName = $"{solutionConfiguration}|Any CPU";
+
+				_sb.AppendLine($"		{guid}.{configAndPlatformName}.ActiveCfg = {projectContext.Name}");
+				if (projectContext.bBuildByDefault)
+				{
+					_sb.AppendLine($"		{guid}.{configAndPlatformName}.Build.0 = {projectContext.Name}");
+					if (projectContext.bDeployByDefault)
+					{
+						_sb.AppendLine($"		{guid}.{configAndPlatformName}.Deploy.0 = {projectContext.Name}");
+					}
+				}
+			}
+		}
+
+		_sb.AppendLine("	EndGlobalSection");
+	}
+
 	public void WriteUnrealVs(HashSet<UnrealTargetPlatform> validPlatforms)
 	{
 		const string UnrealVSGuid = "ddbf523f-7eb6-4887-bd51-85a714ff87eb";
