UnrealBuildTool: Print the path to the solution files being generated so users can more easily find them.

#Synced-CL 51580970

diff --git a/Engine/Source/Programs/UnrealBuildTool/ProjectFiles/VisualStudio/VCProjectFileGenerator.cs b/Engine/Source/Programs/UnrealBuildTool/ProjectFiles/VisualStudio/VCProjectFileGenerator.cs
--- a/Engine/Source/Programs/UnrealBuildTool/ProjectFiles/VisualStudio/VCProjectFileGenerator.cs
+++ b/Engine/Source/Programs/UnrealBuildTool/ProjectFiles/VisualStudio/VCProjectFileGenerator.cs
@@ -508,6 +508,7 @@ namespace UnrealBuildTool
 					SolutionBuilder.WriteGlobalEnd();
 
 					FileReference SolutionFilePath = FileReference.Combine(PrimaryProjectPath, SolutionFileName);
+					Logger.LogInformation("Writing '{SolutionPath}'...", SolutionFilePath);
 					bool success = WriteFileIfChanged(SolutionFilePath.FullName, SolutionBuilder.ToString(), Logger);
 
 					// Generate a slnx for VS2026 until properly supported
@@ -860,6 +861,7 @@ namespace UnrealBuildTool
 			if (bSuccess)
 			{
 				FileReference SolutionFilePath = FileReference.Combine(PrimaryProjectPath, SolutionFileName);
+				Logger.LogInformation("Writing '{SolutionPath}'...", SolutionFilePath);
 				bSuccess = WriteFileIfChanged(SolutionFilePath.FullName, SolutionBuilder.ToString(), Logger);
 
 				// Generate a slnx for VS2026 until properly supported
@@ -902,14 +904,15 @@ namespace UnrealBuildTool
 			bool bSuccess = process.ExitCode == 0;
 			if (bSuccess)
 			{
-				string contents = await FileReference.ReadAllTextAsync(TempSolutionFilePath.ChangeExtension(".slnx"));
+				FileReference TempSlnxFilePath = TempSolutionFilePath.ChangeExtension(".slnx");
+				string contents = await FileReference.ReadAllTextAsync(TempSlnxFilePath);
 
 				// Visual Studio marks UE .slnx files with vcxproj files as dirty, as the `dotnet sln migrate` command doesn't include project GUIDs.
 				// We add them in here.
 				List<VCProjectFile> vcProjectFiles = [.. AllProjectFiles.OfType<VCProjectFile>()];
 				if (vcProjectFiles.Count != 0)
 				{
-					SolutionModel solutionModel = await SolutionSerializers.SlnXml.OpenAsync(TempSolutionFilePath.ChangeExtension(".slnx").FullName, CancellationToken.None);
+					SolutionModel solutionModel = await SolutionSerializers.SlnXml.OpenAsync(TempSlnxFilePath.FullName, CancellationToken.None);
 					foreach (MSBuildProjectFile projectFile in vcProjectFiles)
 					{
 						SolutionProjectModel project = solutionModel.SolutionProjects.Single(x =>
@@ -923,9 +926,11 @@ namespace UnrealBuildTool
 					contents = encoding.GetString(memoryStream.GetBuffer());
 				}
 
-				bSuccess = WriteFileIfChanged(SolutionFilePath.ChangeExtension(".slnx").FullName, contents, Logger);
+				FileReference SlnxFilePath = SolutionFilePath.ChangeExtension(".slnx");
+				Logger.LogInformation("Writing '{SolutionPath}'...", SlnxFilePath);
+				bSuccess = WriteFileIfChanged(SlnxFilePath.FullName, contents, Logger);
 				FileReference.Delete(TempSolutionFilePath);
-				FileReference.Delete(TempSolutionFilePath.ChangeExtension(".slnx"));
+				FileReference.Delete(TempSlnxFilePath);
 			}
 
 			return bSuccess;
