Compile arguments should be included in MSVC shared response files. Previously, only include paths were included.

#Synced-CL 51569266

diff --git a/Engine/Source/Programs/UnrealBuildTool/Platform/Windows/VCToolChain.cs b/Engine/Source/Programs/UnrealBuildTool/Platform/Windows/VCToolChain.cs
--- a/Engine/Source/Programs/UnrealBuildTool/Platform/Windows/VCToolChain.cs
+++ b/Engine/Source/Programs/UnrealBuildTool/Platform/Windows/VCToolChain.cs
@@ -2045,8 +2045,6 @@ namespace UnrealBuildTool
 			// Add additional response files
 			BaseCompileAction.AdditionalResponseFiles.AddRange(CompileEnvironment.AdditionalResponseFiles);
 
-			AppendCLArguments_Global(CompileEnvironment, BaseCompileAction.Arguments);
-
 			BaseCompileAction.bIsAnalyzing = Target.StaticAnalyzer != StaticAnalyzer.None && !CompileEnvironment.bDisableStaticAnalysis && !(Target.WindowsPlatform.Compiler.IsClang() && CompileEnvironment.PrecompiledHeaderAction == PrecompiledHeaderAction.Create);
 			BaseCompileAction.bWarningsAsErrors = CompileEnvironment.bWarningsAsErrors;
 
@@ -2056,8 +2054,14 @@ namespace UnrealBuildTool
 
 			if (!CompileEnvironment.bHasSharedResponseFile)
 			{
+				AppendCLArguments_Global(CompileEnvironment, BaseCompileAction.Arguments);
 				BaseCompileAction.SystemIncludePaths.AddRange(EnvVars.IncludePaths);
 			}
+			else
+			{
+				// /nologo in the shared file doesn't work to suppress logging.
+				BaseCompileAction.Arguments.Add("/nologo");
+			}
 
 			// Remember the architecture
 			BaseCompileAction.Architecture = CompileEnvironment.Architecture;
@@ -3140,6 +3144,7 @@ namespace UnrealBuildTool
 			CppCompileEnvironment NewCompileEnvironment = new CppCompileEnvironment(CompileEnvironment);
 			List<string> Arguments = new List<string>();
 
+			AppendCLArguments_Global(NewCompileEnvironment, Arguments);
 			AddIncludePathArguments(NewCompileEnvironment, Target.WindowsPlatform.Compiler, Arguments);
 
 			// Stash shared include paths for validation purposes
