Changes to prepare for TargetFramework net10.0 update

#rnx
#Synced-CL 51592834

diff --git a/Engine/Source/Programs/Shared/EpicGames.Core/BinaryArchiveReader.cs b/Engine/Source/Programs/Shared/EpicGames.Core/BinaryArchiveReader.cs
--- a/Engine/Source/Programs/Shared/EpicGames.Core/BinaryArchiveReader.cs
+++ b/Engine/Source/Programs/Shared/EpicGames.Core/BinaryArchiveReader.cs
@@ -60,7 +60,7 @@ namespace EpicGames.Core
 		public BinaryArchiveReader(Stream stream)
 		{
 			_buffer = new byte[stream.Length];
-			stream.Read(_buffer, 0, _buffer.Length);
+			stream.ReadExactly(_buffer, 0, _buffer.Length);
 		}
 
 		/// <summary>
diff --git a/Engine/Source/Programs/Shared/EpicGames.Core/FileReference.cs b/Engine/Source/Programs/Shared/EpicGames.Core/FileReference.cs
--- a/Engine/Source/Programs/Shared/EpicGames.Core/FileReference.cs
+++ b/Engine/Source/Programs/Shared/EpicGames.Core/FileReference.cs
@@ -255,7 +255,7 @@ namespace EpicGames.Core
 				Span<byte> span = stackalloc byte[4];
 				using (FileStream stream = Open(location, FileMode.Open, FileAccess.Read, FileShare.Read))
 				{
-					stream.Read(span);
+					stream.ReadExactly(span);
 				}
 				return FileUtils.GetEncoding(span, out skipBytes);
 			}
@@ -286,7 +286,7 @@ namespace EpicGames.Core
 				byte[] buffer = new byte[4];
 				using (FileStream stream = Open(location, FileMode.Open, FileAccess.Read, FileShare.Read))
 				{
-					await stream.ReadAsync(buffer, 0, 4, cancellationToken);
+					await stream.ReadExactlyAsync(buffer, 0, 4, cancellationToken);
 				}
 				return FileUtils.GetEncoding(buffer);
 			}
diff --git a/Engine/Source/Programs/Shared/EpicGames.Core/FileTransaction.cs b/Engine/Source/Programs/Shared/EpicGames.Core/FileTransaction.cs
--- a/Engine/Source/Programs/Shared/EpicGames.Core/FileTransaction.cs
+++ b/Engine/Source/Programs/Shared/EpicGames.Core/FileTransaction.cs
@@ -82,7 +82,7 @@ namespace EpicGames.Core
 				if (stream != null)
 				{
 					byte[] data = new byte[stream.Length];
-					await stream.ReadAsync(data);
+					await stream.ReadExactlyAsync(data);
 					return data;
 				}
 			}
diff --git a/Engine/Source/Programs/Shared/EpicGames.Core/FileUtils.cs b/Engine/Source/Programs/Shared/EpicGames.Core/FileUtils.cs
--- a/Engine/Source/Programs/Shared/EpicGames.Core/FileUtils.cs
+++ b/Engine/Source/Programs/Shared/EpicGames.Core/FileUtils.cs
@@ -531,7 +531,7 @@ namespace EpicGames.Core
 				byte[] bytes = new byte[4];
 				using (FileStream stream = File.Open(filePath, FileMode.Open, FileAccess.Read, FileShare.Read))
 				{
-					stream.Read(bytes, 0, 4);
+					stream.ReadExactly(bytes, 0, 4);
 				}
 				return GetEncoding(bytes, out skipBytes);
 			}
diff --git a/Engine/Source/Programs/Shared/EpicGames.Horde/Storage/Nodes/DirectoryNode.Update.cs b/Engine/Source/Programs/Shared/EpicGames.Horde/Storage/Nodes/DirectoryNode.Update.cs
--- a/Engine/Source/Programs/Shared/EpicGames.Horde/Storage/Nodes/DirectoryNode.Update.cs
+++ b/Engine/Source/Programs/Shared/EpicGames.Horde/Storage/Nodes/DirectoryNode.Update.cs
@@ -608,7 +608,9 @@ namespace EpicGames.Horde.Storage.Nodes
 				{
 					if (entry.Name.Length > 0)
 					{
+#pragma warning disable CA1849 // ZipArchiveEntry.Open() has no async equivalent
 						using (Stream entryStream = entry.Open())
+#pragma warning restore CA1849
 						{
 							LeafChunkedData leafChunkedData = await LeafChunkedDataNode.CreateFromStreamAsync(writer, entryStream, options.LeafOptions, cancellationToken);
 							entries.Add((entry, leafChunkedData));
diff --git a/Engine/Source/Programs/Shared/EpicGames.Horde/Storage/Nodes/DirectoryNode.ZipStream.cs b/Engine/Source/Programs/Shared/EpicGames.Horde/Storage/Nodes/DirectoryNode.ZipStream.cs
--- a/Engine/Source/Programs/Shared/EpicGames.Horde/Storage/Nodes/DirectoryNode.ZipStream.cs
+++ b/Engine/Source/Programs/Shared/EpicGames.Horde/Storage/Nodes/DirectoryNode.ZipStream.cs
@@ -151,7 +151,9 @@ namespace EpicGames.Horde.Storage.Nodes
 						entry.ExternalAttributes |= 0b_110_110_100 << 16; // rw- rw- r--
 					}
 
+#pragma warning disable CA1849 // ZipArchiveEntry.Open() has no async equivalent
 					using Stream entryStream = entry.Open();
+#pragma warning restore CA1849
 					await fileEntry.CopyToStreamAsync(entryStream, cancellationToken);
 				}
 				numCopiedFiles++;
diff --git a/Engine/Source/Programs/Shared/EpicGames.MsBuild/CsProjBuilder.cs b/Engine/Source/Programs/Shared/EpicGames.MsBuild/CsProjBuilder.cs
--- a/Engine/Source/Programs/Shared/EpicGames.MsBuild/CsProjBuilder.cs
+++ b/Engine/Source/Programs/Shared/EpicGames.MsBuild/CsProjBuilder.cs
@@ -148,10 +148,10 @@ namespace EpicGames.MsBuild
 						| targetFramework.Contains("net7.0", StringComparison.Ordinal)
 #endif
 #if NET9_0_OR_GREATER
-						| targetFramework.Contains("net8.0", StringComparison.Ordinal);
+						| targetFramework.Contains("net8.0", StringComparison.Ordinal)
 #endif
 #if NET10_0_OR_GREATER
-						| targetFramework.Contains("net9.0", StringComparison.Ordinal);
+						| targetFramework.Contains("net9.0", StringComparison.Ordinal)
 #endif
 					;
 
diff --git a/Engine/Source/Programs/Shared/EpicGames.Redis.Tests/SerializationTests.cs b/Engine/Source/Programs/Shared/EpicGames.Redis.Tests/SerializationTests.cs
--- a/Engine/Source/Programs/Shared/EpicGames.Redis.Tests/SerializationTests.cs
+++ b/Engine/Source/Programs/Shared/EpicGames.Redis.Tests/SerializationTests.cs
@@ -129,7 +129,7 @@ namespace EpicGames.Redis.Tests
 
 			RedisValue value = RedisSerializer.Serialize(cls);
 
-			ProtoClass clsPb = Serializer.Deserialize<ProtoClass>(value);
+			ProtoClass clsPb = Serializer.Deserialize<ProtoClass>(((byte[]?)value).AsMemory());
 			Assert.AreEqual(clsPb.Text, "hello world");
 
 			ProtoClass clsRedis = RedisSerializer.Deserialize<ProtoClass>(value);
diff --git a/Engine/Source/Programs/UnrealBuildTool/Platform/Android/UEDeployAndroid.cs b/Engine/Source/Programs/UnrealBuildTool/Platform/Android/UEDeployAndroid.cs
--- a/Engine/Source/Programs/UnrealBuildTool/Platform/Android/UEDeployAndroid.cs
+++ b/Engine/Source/Programs/UnrealBuildTool/Platform/Android/UEDeployAndroid.cs
@@ -3636,9 +3636,9 @@ LD_HWASAN=1 exec $cmd";
 					ushort FilenameLen = SourceReader.ReadUInt16();
 					ushort ExtraLen = SourceReader.ReadUInt16();
 					byte[] Filename = new byte[FilenameLen];
-					SourceReader.BaseStream.Read(Filename, 0, FilenameLen);
+					SourceReader.BaseStream.ReadExactly(Filename, 0, FilenameLen);
 					byte[] Extra = new byte[ExtraLen];
-					SourceReader.BaseStream.Read(Extra, 0, ExtraLen);
+					SourceReader.BaseStream.ReadExactly(Extra, 0, ExtraLen);
 
 					string FilenameStr = Encoding.UTF8.GetString(Filename, 0, FilenameLen);
 					if (!FilenameStr.Equals(RelativeSOFile, StringComparison.OrdinalIgnoreCase))
diff --git a/Engine/Source/Programs/UnrealBuildTool/Preprocessor/SourceFile.cs b/Engine/Source/Programs/UnrealBuildTool/Preprocessor/SourceFile.cs
--- a/Engine/Source/Programs/UnrealBuildTool/Preprocessor/SourceFile.cs
+++ b/Engine/Source/Programs/UnrealBuildTool/Preprocessor/SourceFile.cs
@@ -133,7 +133,7 @@ namespace UnrealBuildTool
 			int length = (int)stream.Length;
 
 			byte[] data = new byte[length + 1];
-			stream.Read(data, 0, length);
+			stream.ReadExactly(data, 0, length);
 
 			return ConvertToNullTerminatedUtf8(data, length);
 		}
