[UBT]
* Fixed bug where _hordeSessionTasks could be accessed outside range

#Synced-CL 51656697

diff --git a/Engine/Source/Programs/UnrealBuildTool/Executors/UnrealBuildAccelerator/UBAAgentCoordinatorHorde.cs b/Engine/Source/Programs/UnrealBuildTool/Executors/UnrealBuildAccelerator/UBAAgentCoordinatorHorde.cs
--- a/Engine/Source/Programs/UnrealBuildTool/Executors/UnrealBuildAccelerator/UBAAgentCoordinatorHorde.cs
+++ b/Engine/Source/Programs/UnrealBuildTool/Executors/UnrealBuildAccelerator/UBAAgentCoordinatorHorde.cs
@@ -228,7 +228,7 @@ namespace UnrealBuildTool
 				{
 					return null;
 				}
-				if (_hordeSessionTasks.Length < i)
+				if (i >= _hordeSessionTasks.Length)
 				{
 					return null;
 				}
