UnrealBuildTool: Use Microsoft's solution persistence library to build solutions, including natively supporting .slnx files.

#Synced-CL 51664322
#Synced-CL 51664322

diff --git a/Engine/Source/Programs/UnrealBuildTool/ProjectFiles/ProjectFile.cs b/Engine/Source/Programs/UnrealBuildTool/ProjectFiles/ProjectFile.cs
--- a/Engine/Source/Programs/UnrealBuildTool/ProjectFiles/ProjectFile.cs
+++ b/Engine/Source/Programs/UnrealBuildTool/ProjectFiles/ProjectFile.cs
@@ -362,29 +362,6 @@ namespace UnrealBuildTool
 			IntelliSenseEnableCoroutines |= Value;
 		}
 
-		/// <summary>
-		/// Add the given project to the DepondsOn project list.
-		/// </summary>
-		/// <param name="InProjectFile">The project this project is dependent on</param>
-		public void AddDependsOnProject(ProjectFile InProjectFile)
-		{
-			// Make sure that it doesn't exist already
-			bool AlreadyExists = false;
-			foreach (ProjectFile ExistingDependentOn in DependsOnProjects)
-			{
-				if (ExistingDependentOn == InProjectFile)
-				{
-					AlreadyExists = true;
-					break;
-				}
-			}
-
-			if (AlreadyExists == false)
-			{
-				DependsOnProjects.Add(InProjectFile);
-			}
-		}
-
 		/// <summary>
 		/// Writes a project file to disk
 		/// </summary>
@@ -560,11 +537,6 @@ namespace UnrealBuildTool
 		/// <summary>
 		/// Set of unique preprocessor definitions
 		/// </summary>
-		readonly HashSet<string> KnownIntelliSensePreprocessorDefinitions = new HashSet<string>(StringComparer.InvariantCultureIgnoreCase);
-
-		/// <summary>
-		/// Projects that this project is dependent on
-		/// </summary>
-		public readonly List<ProjectFile> DependsOnProjects = new List<ProjectFile>();
+		private readonly HashSet<string> KnownIntelliSensePreprocessorDefinitions = new HashSet<string>(StringComparer.InvariantCultureIgnoreCase);
 	}
 }
diff --git a/Engine/Source/Programs/UnrealBuildTool/ProjectFiles/ProjectFileGenerator.cs b/Engine/Source/Programs/UnrealBuildTool/ProjectFiles/ProjectFileGenerator.cs
--- a/Engine/Source/Programs/UnrealBuildTool/ProjectFiles/ProjectFileGenerator.cs
+++ b/Engine/Source/Programs/UnrealBuildTool/ProjectFiles/ProjectFileGenerator.cs
@@ -46,6 +46,22 @@ namespace UnrealBuildTool
 			FolderName = InitFolderName;
 		}
 
+		/// <summary>
+		/// Deep clone a primary project folder.
+		/// </summary>
+		/// <param name="Other">The folder to clone.</param>
+		/// <remarks>
+		/// Not all project files are necessarily immutable, so check your use-case carefully.
+		/// </remarks>
+		public PrimaryProjectFolder(PrimaryProjectFolder Other)
+		{
+			OwnerProjectFileGenerator = Other.OwnerProjectFileGenerator;
+			FolderName = Other.FolderName;
+			SubFolders.AddRange(Other.SubFolders.Select(x => new PrimaryProjectFolder(x)));
+			ChildProjects.AddRange(Other.ChildProjects);
+			Files.AddRange(Other.Files);
+		}
+
 		/// Name of this folder
 		public string FolderName
 		{
diff --git a/Engine/Source/Programs/UnrealBuildTool/ProjectFiles/Rider/RiderProjectFileGenerator.cs b/Engine/Source/Programs/UnrealBuildTool/ProjectFiles/Rider/RiderProjectFileGenerator.cs
--- a/Engine/Source/Programs/UnrealBuildTool/ProjectFiles/Rider/RiderProjectFileGenerator.cs
+++ b/Engine/Source/Programs/UnrealBuildTool/ProjectFiles/Rider/RiderProjectFileGenerator.cs
@@ -6,6 +6,7 @@ using System.Linq;
 using EpicGames.Core;
 using Microsoft.Extensions.Logging;
 using UnrealBuildBase;
+using UnrealBuildTool.ProjectFiles.VisualStudio;
 
 namespace UnrealBuildTool
 {
@@ -520,21 +521,21 @@ namespace UnrealBuildTool
 			try
 			{
 				HashSet<UnrealTargetPlatform> ValidPlatforms;
-				List<VCProjectFileGenerator.VCSolutionConfigCombination> ValidSolutionConfigs;
+				List<VisualStudioPrimarySolutionProjectParams> VcxSolutionProjectParamList;
 				VCProjectFileGenerator.CollectSolutionConfigurations(SupportedConfigurations,
-					SupportedPlatforms, AllProjectFiles, bMakeProjectPerTarget, out ValidPlatforms, out ValidSolutionConfigs);
+					SupportedPlatforms, AllProjectFiles, bMakeProjectPerTarget, out ValidPlatforms, out VcxSolutionProjectParamList);
 				FileReference OutputFile = GetRiderProjectLocation("root.json");
 				DirectoryReference.CreateDirectory(OutputFile.Directory);
 				using (JsonWriter Writer = new JsonWriter(OutputFile))
 				{
 					Writer.WriteObjectStart();
 					Writer.WriteArrayStart("ConfigurationsAndPlatforms");
-					foreach (VCProjectFileGenerator.VCSolutionConfigCombination Config in ValidSolutionConfigs)
+					foreach (VisualStudioPrimarySolutionProjectParams Params in VcxSolutionProjectParamList)
 					{
 						Writer.WriteObjectStart();
-						Writer.WriteValue("Configuration", Config.Configuration.ToString());
-						Writer.WriteValue("Platform", Config.Platform.ToString());
-						Writer.WriteValue("Target", Config.TargetConfigurationName.ToString());
+						Writer.WriteValue("Configuration", Params.ProjectConfiguration.ToString());
+						Writer.WriteValue("Platform", Params.ProjectPlatform.ToString());
+						Writer.WriteValue("Target", Params.ProjectTargetType.ToString());
 						Writer.WriteObjectEnd();
 					}
 
diff --git a/Engine/Source/Programs/UnrealBuildTool/ProjectFiles/VisualStudio/VCProject.cs b/Engine/Source/Programs/UnrealBuildTool/ProjectFiles/VisualStudio/VCProject.cs
--- a/Engine/Source/Programs/UnrealBuildTool/ProjectFiles/VisualStudio/VCProject.cs
+++ b/Engine/Source/Programs/UnrealBuildTool/ProjectFiles/VisualStudio/VCProject.cs
@@ -14,6 +14,7 @@ using EpicGames.Core;
 using Microsoft.CodeAnalysis;
 using Microsoft.Extensions.Logging;
 using UnrealBuildBase;
+using UnrealBuildTool.ProjectFiles.VisualStudio;
 using static UnrealBuildTool.PlatformProjectGenerator;
 
 namespace UnrealBuildTool
@@ -174,14 +175,14 @@ namespace UnrealBuildTool
 		/// <summary>
 		/// Get the project context for the given solution context
 		/// </summary>
-		/// <param name="SolutionTarget">The solution target type</param>
-		/// <param name="SolutionConfiguration">The solution configuration</param>
-		/// <param name="SolutionPlatform">The solution platform</param>
+		/// <param name="SolutionProjectParams">The solution project params to get a context for</param>
 		/// <param name="PlatformProjectGenerators">Set of platform project generators</param>
-		/// <param name="Architecture">The target architecture</param>
 		/// <param name="Logger">Logger for output</param>
 		/// <returns>Project context matching the given solution context</returns>
-		public abstract MSBuildProjectContext? GetMatchingProjectContext(TargetType SolutionTarget, UnrealTargetConfiguration SolutionConfiguration, UnrealTargetPlatform SolutionPlatform, PlatformProjectGeneratorCollection PlatformProjectGenerators, UnrealArch? Architecture, ILogger Logger);
+		public abstract MSBuildProjectContext? GetMatchingProjectContext(
+			VisualStudioPrimarySolutionProjectParams SolutionProjectParams,
+			PlatformProjectGeneratorCollection PlatformProjectGenerators,
+			ILogger Logger);
 
 		/// <summary>
 		/// Checks to see if the specified solution platform and configuration is able to map to this project
@@ -326,7 +327,10 @@ namespace UnrealBuildTool
 		}
 
 		/// <inheritdoc/>
-		public override MSBuildProjectContext? GetMatchingProjectContext(TargetType SolutionTarget, UnrealTargetConfiguration SolutionConfiguration, UnrealTargetPlatform SolutionPlatform, PlatformProjectGeneratorCollection PlatformProjectGenerators, UnrealArch? Architecture, ILogger Logger)
+		public override MSBuildProjectContext? GetMatchingProjectContext(
+			VisualStudioPrimarySolutionProjectParams ProjectParams,
+			PlatformProjectGeneratorCollection PlatformProjectGenerators,
+			ILogger Logger)
 		{
 			// Stub projects always build in the same configuration
 			if (IsStubProject)
@@ -350,12 +354,12 @@ namespace UnrealBuildTool
 			bool IsProgramProject = ProjectTargets.Any() && ProjectTargets[0].TargetRules != null && ProjectTargets[0].TargetRules!.Type == TargetType.Program;
 			{
 				// Get the target type we expect to find for this project
-				TargetType TargetConfigurationName = SolutionTarget;
+				TargetType TargetConfigurationName = ProjectParams.ProjectTargetType;
 				if (IsProgramProject)
 				{
 					TargetConfigurationName = TargetType.Program;
 				}
-				else if (SolutionTarget == TargetType.Program)
+				else if (ProjectParams.ProjectTargetType == TargetType.Program)
 				{
 					TargetConfigurationName = TargetType.Game;
 				}
@@ -365,7 +369,7 @@ namespace UnrealBuildTool
 				List<Project> MatchingProjectTargets = new List<Project>();
 				foreach (Project ProjectTarget in ProjectTargets)
 				{
-					if (VCProjectFile.IsValidProjectPlatformAndConfiguration(ProjectTarget, SolutionPlatform, SolutionConfiguration))
+					if (VCProjectFile.IsValidProjectPlatformAndConfiguration(ProjectTarget, ProjectParams.ProjectPlatform, ProjectParams.ProjectConfiguration))
 					{
 						if (ProjectTarget.TargetRules != null)
 						{
@@ -385,7 +389,7 @@ namespace UnrealBuildTool
 				}
 
 				// Always allow SCW and UnrealLighmass to build in editor configurations
-				if (MatchingProjectTargets.Count == 0 && SolutionTarget == TargetType.Editor && SolutionPlatform == UnrealTargetPlatform.Win64)
+				if (MatchingProjectTargets.Count == 0 && ProjectParams.ProjectTargetType == TargetType.Editor && ProjectParams.ProjectPlatform == UnrealTargetPlatform.Win64)
 				{
 					foreach (Project ProjectTarget in ProjectTargets)
 					{
@@ -401,7 +405,12 @@ namespace UnrealBuildTool
 				// Make sure there's only one matching project target
 				if (MatchingProjectTargets.Count > 1)
 				{
-					throw new BuildException("Not expecting more than one target for project {0} to match solution configuration {1} {2} {3}", ProjectFilePath, SolutionTarget, SolutionConfiguration, SolutionPlatform);
+					throw new BuildException(
+						"Not expecting more than one target for project {0} to match solution configuration {1} {2} {3}",
+						ProjectFilePath,
+						ProjectParams.ProjectTargetType,
+						ProjectParams.ProjectConfiguration,
+						ProjectParams.ProjectPlatform);
 				}
 
 				// If we found a matching project, get matching configuration
@@ -411,14 +420,14 @@ namespace UnrealBuildTool
 					Project MatchingProjectTarget = MatchingProjectTargets[0];
 
 					// If the project wants to always build in "Development", regardless of what the solution configuration is set to, then we'll do that here.
-					UnrealTargetConfiguration ProjectConfiguration = SolutionConfiguration;
-					if (MatchingProjectTarget.ForceDevelopmentConfiguration && SolutionTarget != TargetType.Game)
+					UnrealTargetConfiguration ProjectConfiguration = ProjectParams.ProjectConfiguration;
+					if (MatchingProjectTarget.ForceDevelopmentConfiguration && ProjectParams.ProjectTargetType != TargetType.Game)
 					{
 						ProjectConfiguration = UnrealTargetConfiguration.Development;
 					}
 
 					// Get the matching project configuration
-					UnrealTargetPlatform ProjectPlatform = SolutionPlatform;
+					UnrealTargetPlatform ProjectPlatform = ProjectParams.ProjectPlatform;
 					if (IsStubProject)
 					{
 						if (ProjectConfiguration != UnrealTargetConfiguration.Unknown)
@@ -436,7 +445,7 @@ namespace UnrealBuildTool
 						}
 
 						ProjectConfigAndTargetCombination? Combination = ProjectConfigAndTargetCombinations.FirstOrDefault(
-							ProjectPlatform == UnrealTargetPlatform.Android && Architecture == null ?
+							ProjectPlatform == UnrealTargetPlatform.Android && ProjectParams.ProjectArchitecture == null ?
 								// Match first preferred Architecture if none is provided
 								Combination =>
 									Combination.Platform == ProjectPlatform &&
@@ -446,7 +455,7 @@ namespace UnrealBuildTool
 									Combination.Platform == ProjectPlatform &&
 									Combination.Configuration == ProjectConfiguration &&
 									Combination.ProjectTarget == MatchingProjectTarget &&
-									Combination.Architecture == Architecture);
+									Combination.Architecture == ProjectParams.ProjectArchitecture);
 
 						if (Combination == null)
 						{
@@ -472,13 +481,13 @@ namespace UnrealBuildTool
 							SupportedPlatforms = new List<UnrealTargetPlatform>();
 							SupportedPlatforms.AddRange(MatchingProjectTarget.SupportedPlatforms);
 						}
-						if (SupportedPlatforms == null || SupportedPlatforms.Contains(SolutionPlatform))
+						if (SupportedPlatforms == null || SupportedPlatforms.Contains(ProjectParams.ProjectPlatform))
 						{
 							bBuildByDefault = true;
 
-							PlatformProjectGenerator? ProjGen = PlatformProjectGenerators.GetPlatformProjectGenerator(SolutionPlatform, true);
+							PlatformProjectGenerator? ProjGen = PlatformProjectGenerators.GetPlatformProjectGenerator(ProjectParams.ProjectPlatform, true);
 							if (MatchingProjectTarget.ProjectDeploys ||
-								((ProjGen != null) && (ProjGen.GetVisualStudioDeploymentEnabled(new VSSettings(ProjectPlatform, ProjectConfiguration, ProjectFileFormat, Architecture)) == true)))
+								((ProjGen != null) && (ProjGen.GetVisualStudioDeploymentEnabled(new VSSettings(ProjectPlatform, ProjectConfiguration, ProjectFileFormat, ProjectParams.ProjectArchitecture)) == true)))
 							{
 								bDeployByDefault = true;
 							}
diff --git a/Engine/Source/Programs/UnrealBuildTool/ProjectFiles/VisualStudio/VCProjectFileGenerator.cs b/Engine/Source/Programs/UnrealBuildTool/ProjectFiles/VisualStudio/VCProjectFileGenerator.cs
--- a/Engine/Source/Programs/UnrealBuildTool/ProjectFiles/VisualStudio/VCProjectFileGenerator.cs
+++ b/Engine/Source/Programs/UnrealBuildTool/ProjectFiles/VisualStudio/VCProjectFileGenerator.cs
@@ -7,14 +7,10 @@ using System.Linq;
 using System.Runtime.Versioning;
 using System.Security.Cryptography;
 using System.Text;
-using System.Threading;
-using System.Threading.Tasks;
 using System.Xml;
 using System.Xml.Linq;
 using EpicGames.Core;
 using Microsoft.Extensions.Logging;
-using Microsoft.VisualStudio.SolutionPersistence.Model;
-using Microsoft.VisualStudio.SolutionPersistence.Serializer;
 using UnrealBuildBase;
 using UnrealBuildTool.ProjectFiles.VisualStudio;
 
@@ -481,41 +477,33 @@ namespace UnrealBuildTool
 
 				AddAutomationProjects(AllGameProjects, Logger);
 
-				// All projects must be C# here - if they're not, this is a bug and should be fixed.
-				IEnumerable<VCSharpProjectFile> CSharpProjects = AllProjectFiles.Cast<VCSharpProjectFile>();
-				List<string> SolutionConfigurations =
-				[.. CSharpProjects
-					.SelectMany(x => x.Configurations)
-					.Distinct(StringComparer.OrdinalIgnoreCase)
-					// VS will order the solution configurations alphabetically
-					.OrderBy(x => x, StringComparer.OrdinalIgnoreCase)
-				];
-
 				using (Timeline.ScopeEvent("Writing automation project files", Logger))
 				{
+					// All projects must be C# here - if they're not, this is a bug and should be fixed.
+					IEnumerable<VCSharpProjectFile> CSharpProjects = AllProjectFiles.Cast<VCSharpProjectFile>();
+					List<string> SolutionConfigurations =
+					[.. CSharpProjects
+						.SelectMany(x => x.Configurations)
+						.Distinct(StringComparer.OrdinalIgnoreCase)
+						// VS will order the solution configurations alphabetically
+						.OrderBy(x => x, StringComparer.OrdinalIgnoreCase)
+					];
+					List<VisualStudioSolutionProjectParams> SolutionCombinations =
+					[..
+						SolutionConfigurations.Select(x => new VisualStudioSolutionProjectParams(x, "Any CPU"))
+					];
 					string SolutionFileName = $"Automation_{PrimaryProjectName}.sln";
-					VisualStudioSolutionBuilder SolutionBuilder = new(RootFolder, AllProjectFiles, Logger);
-
-					SolutionBuilder.WritePreamble(Settings.ProjectFileFormat);
-					SolutionBuilder.WriteProjectFolders();
-					SolutionBuilder.WriteProjects();
+					FileReference SlnFilePath = FileReference.Combine(PrimaryProjectPath, SolutionFileName);
+					FileReference SlnxFilePath = SlnFilePath.ChangeExtension(".slnx");
 
-					SolutionBuilder.WriteGlobalStart();
-					SolutionBuilder.WriteSolutionConfigurationPlatformsForCSharpOnlySolution(SolutionConfigurations);
-					SolutionBuilder.WriteProjectConfigurationPlatformsForCSharpOnlySolution(SolutionConfigurations);
-					SolutionBuilder.WriteSolutionProperties();
-					SolutionBuilder.WriteNestedProjects();
-					SolutionBuilder.WriteGlobalEnd();
+					AutomationVisualStudioSolutionBuilder Builder = new(SlnFilePath.Directory, RootFolder, SolutionCombinations, Logger);
+					string slnContents = Builder.BuildSolution(Settings.ProjectFileFormat, forSlnx: false);
+					string slnxContents = Builder.BuildSolution(Settings.ProjectFileFormat, forSlnx: true);
 
-					FileReference SolutionFilePath = FileReference.Combine(PrimaryProjectPath, SolutionFileName);
-					Logger.LogInformation("Writing '{SolutionPath}'...", SolutionFilePath);
-					bool success = WriteFileIfChanged(SolutionFilePath.FullName, SolutionBuilder.ToString(), Logger);
-
-					// Generate a slnx for VS2026 until properly supported
-					if (success && Environment.Version.Major >= 10)
-					{
-						ConvertSolutionFileToSlnx(SolutionFileName, SolutionFilePath, Logger).Wait();
-					}
+					Logger.LogInformation("Writing '{SolutionPath}'...", SlnFilePath);
+					_ = WriteFileIfChanged(SlnFilePath.FullName, slnContents, Logger);
+					Logger.LogInformation("Writing '{SolutionPath}'...", SlnxFilePath);
+					_ =WriteFileIfChanged(SlnxFilePath.FullName, slnxContents, Logger);
 				}
 			}
 			finally
@@ -529,47 +517,6 @@ namespace UnrealBuildTool
 			}
 		}
 
-		/// <summary>
-		/// Used to sort VC solution config names along with the config and platform values
-		/// </summary>
-		public class VCSolutionConfigCombination
-		{
-			/// <summary>
-			/// Visual Studio solution configuration name for this config+platform
-			/// </summary>
-			public string VCSolutionConfigAndPlatformName;
-
-			/// <summary>
-			/// Configuration name
-			/// </summary>
-			public UnrealTargetConfiguration Configuration;
-
-			/// <summary>
-			/// Platform name
-			/// </summary>
-			public UnrealTargetPlatform Platform;
-
-			/// <summary>
-			/// The target type
-			/// </summary>
-			public TargetType TargetConfigurationName;
-
-			/// <summary>
-			/// The target architecture
-			/// </summary>
-			public UnrealArch? Architecture;
-
-			public override string ToString()
-			{
-				return String.Format("{0}={1} {2} {3}{4}", VCSolutionConfigAndPlatformName, Configuration, Platform, TargetConfigurationName, Architecture != null ? " " + Architecture : String.Empty);
-			}
-
-			public VCSolutionConfigCombination(string VCSolutionConfigAndPlatformName)
-			{
-				this.VCSolutionConfigAndPlatformName = VCSolutionConfigAndPlatformName;
-			}
-		}
-
 		/// <summary>
 		/// Composes a string to use for the Visual Studio solution configuration, given a build configuration and target rules configuration name
 		/// </summary>
@@ -793,89 +740,56 @@ namespace UnrealBuildTool
 
 			string SolutionFileName = PrimaryProjectName + ".sln";
 
-			// Setup solution file content
-			VisualStudioSolutionBuilder SolutionBuilder = new(RootFolder, AllProjectFiles, Logger);
-
-			SolutionBuilder.WritePreamble(Settings.ProjectFileFormat);
-
-			// Solution folders, files and project entries
-			{
-				SolutionBuilder.WriteProjectFolders();
+			// Clone the root folder as we want to add visualizers to it for VS specifically.
+			PrimaryProjectFolder RootFolderClone = new(RootFolder);
 
-				// Get the path to the visualizers file. Try to make it relative to the solution directory, but fall back to a full path if it's a foreign project.
-				FileReference VisualizersFile = FileReference.Combine(Unreal.EngineDirectory, "Extras", "VisualStudioDebugging", "Unreal.natvis");
-				FileReference VisualizersStepFile = FileReference.Combine(Unreal.EngineDirectory, "Extras", "VisualStudioDebugging", "Unreal.natstepfilter");
-				FileReference VisualizersJustMyCodeFile = FileReference.Combine(Unreal.EngineDirectory, "Extras", "VisualStudioDebugging", "Unreal.natjmc");
+			// Get the path to the visualizers file. Try to make it relative to the solution directory, but fall back to a full path if it's a foreign project.
+			FileReference VisualizersFile = FileReference.Combine(Unreal.EngineDirectory, "Extras", "VisualStudioDebugging", "Unreal.natvis");
+			FileReference VisualizersStepFile = FileReference.Combine(Unreal.EngineDirectory, "Extras", "VisualStudioDebugging", "Unreal.natstepfilter");
+			FileReference VisualizersJustMyCodeFile = FileReference.Combine(Unreal.EngineDirectory, "Extras", "VisualStudioDebugging", "Unreal.natjmc");
 
-				// Add the visualizers at the solution level. Doesn't seem to be picked up from a makefile project in VS2017 15.8.5.
-				PrimaryProjectFolder visualizers = new PrimaryProjectFolder(this, "Visualizers");
-				visualizers.Files.AddRange([
-					VisualizersFile.MakeRelativeTo(PrimaryProjectPath),
+			// Add the visualizers at the solution level
+			PrimaryProjectFolder visualizers = RootFolderClone.AddSubFolder("Visualizers");
+			visualizers.Files.AddRange([
+				VisualizersFile.MakeRelativeTo(PrimaryProjectPath),
 					VisualizersStepFile.MakeRelativeTo(PrimaryProjectPath),
 					VisualizersJustMyCodeFile.MakeRelativeTo(PrimaryProjectPath),
 				]);
-				SolutionBuilder.WriteProjectStart(
-					VisualStudioSolutionBuilder.SolutionFolderGuid,
-					"Visualizers",
-					"Visualizers",
-					Guid.Parse("1CCEC849-CC72-4C59-8C36-2F7C38706D4C")
-				);
-				SolutionBuilder.WriteProjectSolutionItems(visualizers);
-				SolutionBuilder.WriteProjectEnd();
-
-				SolutionBuilder.WriteProjects();
-			}
 
 			// Solution configuration platforms. This is just a list of all of the platforms and configurations that
 			// appear in Visual Studio's build configuration selector.
-			List<VCSolutionConfigCombination> SolutionConfigCombinations;
-
-			// The "Global" section has source control, solution configurations, project configurations,
-			// preferences, and project hierarchy data
-			{
-				SolutionBuilder.WriteGlobalStart();
-				{
-					HashSet<UnrealTargetPlatform> PlatformsValidForProjects;
-					CollectSolutionConfigurations(SupportedConfigurations, SupportedPlatforms, AllProjectFiles, bMakeProjectPerTarget,
-						out PlatformsValidForProjects, out SolutionConfigCombinations);
-
-					SolutionBuilder.WriteSolutionConfigurationPlatforms(SolutionConfigCombinations);
-
-					// Embed UnrealVS section, which is parsed by that VSPackage (Extension) to know how to handle this solution
-					SolutionBuilder.WriteUnrealVs(PlatformsValidForProjects);
-
-					// Assign each project's "project configuration" to our "solution platform + configuration" pairs.  This
-					// also sets up which projects are actually built when building the solution.
-					bool DebugPredicate(MSBuildProjectFile Project) => Settings.bBuildUBTInDebug && Project == UBTProject;
-					SolutionBuilder.WriteProjectConfigurationPlatforms(SolutionConfigCombinations, PlatformProjectGenerators, DebugPredicate);
-
-					SolutionBuilder.WriteSolutionProperties();
-
-					SolutionBuilder.WriteNestedProjects();
-				}
-
-				SolutionBuilder.WriteGlobalEnd();
-			}
+			List<VisualStudioPrimarySolutionProjectParams> ProjectParamsList;
+			HashSet<UnrealTargetPlatform> PlatformsValidForProjects;
+				CollectSolutionConfigurations(SupportedConfigurations, SupportedPlatforms, AllProjectFiles, bMakeProjectPerTarget,
+					out PlatformsValidForProjects, out ProjectParamsList);
+			bool DebugPredicate(MSBuildProjectFile Project) => Settings.bBuildUBTInDebug && Project == UBTProject;
 
 			// Save the solution file
-			if (bSuccess)
-			{
-				FileReference SolutionFilePath = FileReference.Combine(PrimaryProjectPath, SolutionFileName);
-				Logger.LogInformation("Writing '{SolutionPath}'...", SolutionFilePath);
-				bSuccess = WriteFileIfChanged(SolutionFilePath.FullName, SolutionBuilder.ToString(), Logger);
+			FileReference SlnFilePath = FileReference.Combine(PrimaryProjectPath, SolutionFileName);
+			FileReference SlnxFilePath = SlnFilePath.ChangeExtension(".slnx");
+			PrimaryVisualStudioSolutionBuilder Builder = new(
+				SlnFilePath.Directory,
+				RootFolderClone,
+				ProjectParamsList,
+				PlatformProjectGenerators,
+				[.. PlatformsValidForProjects],
+				DebugPredicate,
+				Logger
+			);
 
-				// Generate a slnx for VS2026 until properly supported
-				if (bSuccess && Settings.ProjectFileFormat == VCProjectFileFormat.VisualStudio2026 && Environment.Version.Major >= 10)
-				{
-					bSuccess = ConvertSolutionFileToSlnx(SolutionFileName, SolutionFilePath, Logger).GetAwaiter().GetResult();
-				}
-			}
+			string slnContents = Builder.BuildSolution(Settings.ProjectFileFormat, forSlnx: false);
+			string slnxContents = Builder.BuildSolution(Settings.ProjectFileFormat, forSlnx: true);
+
+			Logger.LogInformation("Writing '{SolutionPath}'...", SlnFilePath);
+			bSuccess &= WriteFileIfChanged(SlnFilePath.FullName, slnContents, Logger);
+			Logger.LogInformation("Writing '{SolutionPath}'...", SlnxFilePath);
+			bSuccess &= WriteFileIfChanged(SlnxFilePath.FullName, slnxContents, Logger);
 
 			// Save a solution config file which selects the development editor configuration by default.
 			// .suo file writable only on Windows, requires ole32
 			if (bSuccess && Settings.bWriteSolutionOptionFile && OperatingSystem.IsWindows())
 			{
-				WriteSolutionOptionFile(SolutionFileName, SolutionConfigCombinations);
+				WriteSolutionOptionFile(SolutionFileName, ProjectParamsList);
 			}
 
 			if (bSuccess && Settings.bVsConfigFile && OperatingSystem.IsWindows())
@@ -886,58 +800,8 @@ namespace UnrealBuildTool
 			return bSuccess;
 		}
 
-		private async Task<bool> ConvertSolutionFileToSlnx(string SolutionFileName, FileReference SolutionFilePath, ILogger Logger)
-		{
-			FileReference TempSolutionFilePath = FileReference.Combine(new DirectoryReference(Path.GetTempPath()), SolutionFileName);
-			FileReference.Copy(SolutionFilePath, TempSolutionFilePath, true);
-			using ManagedProcessGroup group = new();
-			using ManagedProcess process = new(
-				group,
-				Unreal.DotnetPath.FullName,
-				$"sln migrate \"{TempSolutionFilePath}\"",
-				TempSolutionFilePath.Directory.FullName,
-				null,
-				System.Diagnostics.ProcessPriorityClass.Normal
-			);
-			process.ReadAllLines();
-			await process.WaitForExitAsync();
-			bool bSuccess = process.ExitCode == 0;
-			if (bSuccess)
-			{
-				FileReference TempSlnxFilePath = TempSolutionFilePath.ChangeExtension(".slnx");
-				string contents = await FileReference.ReadAllTextAsync(TempSlnxFilePath);
-
-				// Visual Studio marks UE .slnx files with vcxproj files as dirty, as the `dotnet sln migrate` command doesn't include project GUIDs.
-				// We add them in here.
-				List<VCProjectFile> vcProjectFiles = [.. AllProjectFiles.OfType<VCProjectFile>()];
-				if (vcProjectFiles.Count != 0)
-				{
-					SolutionModel solutionModel = await SolutionSerializers.SlnXml.OpenAsync(TempSlnxFilePath.FullName, CancellationToken.None);
-					foreach (MSBuildProjectFile projectFile in vcProjectFiles)
-					{
-						SolutionProjectModel project = solutionModel.SolutionProjects.Single(x =>
-							FileReference.Combine(SolutionFilePath.Directory, x.FilePath) == projectFile.ProjectFilePath);
-						project.Id = projectFile.ProjectGUID;
-					}
-
-					using MemoryStream memoryStream = new();
-					await SolutionSerializers.SlnXml.SaveAsync(memoryStream, solutionModel, CancellationToken.None);
-					UTF8Encoding encoding = new UTF8Encoding(encoderShouldEmitUTF8Identifier: false);
-					contents = encoding.GetString(memoryStream.GetBuffer());
-				}
-
-				FileReference SlnxFilePath = SolutionFilePath.ChangeExtension(".slnx");
-				Logger.LogInformation("Writing '{SolutionPath}'...", SlnxFilePath);
-				bSuccess = WriteFileIfChanged(SlnxFilePath.FullName, contents, Logger);
-				FileReference.Delete(TempSolutionFilePath);
-				FileReference.Delete(TempSlnxFilePath);
-			}
-
-			return bSuccess;
-		}
-
 		[SupportedOSPlatform("windows")]
-		private void WriteSolutionOptionFile(string SolutionFileName, List<VCSolutionConfigCombination> SolutionConfigCombinations)
+		private void WriteSolutionOptionFile(string SolutionFileName, List<VisualStudioPrimarySolutionProjectParams> SolutionVcxProjectParamsList)
 		{
 			// Figure out the filename for the SUO file. VS will automatically import the options from earlier versions if necessary.
 			FileReference SolutionOptionsFileName = Settings.ProjectFileFormat switch
@@ -955,15 +819,15 @@ namespace UnrealBuildTool
 				VCSolutionOptions Options = new VCSolutionOptions(Settings.ProjectFileFormat);
 
 				// Set the default configuration and startup project
-				VCSolutionConfigCombination? DefaultConfig = SolutionConfigCombinations.Find(x =>
-					x.Configuration == UnrealTargetConfiguration.Development &&
-					x.Platform == UnrealTargetPlatform.Win64 &&
-					(x.Architecture == null || x.Architecture == UnrealArch.X64) &&
-					(bMakeProjectPerTarget || x.TargetConfigurationName == TargetType.Editor));
+				VisualStudioPrimarySolutionProjectParams? DefaultConfig = SolutionVcxProjectParamsList.Find(x =>
+					x.ProjectConfiguration == UnrealTargetConfiguration.Development &&
+					x.ProjectPlatform == UnrealTargetPlatform.Win64 &&
+					(x.ProjectArchitecture == null || x.ProjectArchitecture == UnrealArch.X64) &&
+					(bMakeProjectPerTarget || x.ProjectTargetType == TargetType.Editor));
 				if (DefaultConfig != null)
 				{
 					List<VCBinarySetting> Settings = new List<VCBinarySetting>();
-					Settings.Add(new VCBinarySetting("ActiveCfg", DefaultConfig.VCSolutionConfigAndPlatformName));
+					Settings.Add(new VCBinarySetting("ActiveCfg", DefaultConfig.SolutionConfigurationAndPlatform));
 					if (DefaultProject != null)
 					{
 						Settings.Add(new VCBinarySetting("StartupProject", ((MSBuildProjectFile)DefaultProject).ProjectGUID.ToString("B")));
@@ -1002,10 +866,10 @@ namespace UnrealBuildTool
 
 		public static void CollectSolutionConfigurations(List<UnrealTargetConfiguration> AllConfigurations,
 			List<UnrealTargetPlatform> AllPlatforms, List<ProjectFile> AllProjectFiles, bool bMakeProjectPerTarget,
-			out HashSet<UnrealTargetPlatform> OutValidPlatforms, out List<VCSolutionConfigCombination> OutSolutionConfigs)
+			out HashSet<UnrealTargetPlatform> OutValidPlatforms, out List<VisualStudioPrimarySolutionProjectParams> OutSolutionVcxProjectParamsList)
 		{
 			OutValidPlatforms = new HashSet<UnrealTargetPlatform>();
-			OutSolutionConfigs = new List<VCSolutionConfigCombination>();
+			OutSolutionVcxProjectParamsList = new List<VisualStudioPrimarySolutionProjectParams>();
 			Dictionary<string, Tuple<UnrealTargetConfiguration, Tuple<ProjectTarget, TargetType>>> SolutionConfigurationsValidForProjects = new();
 
 			foreach (UnrealTargetConfiguration CurConfiguration in AllConfigurations)
@@ -1060,7 +924,7 @@ namespace UnrealBuildTool
 					{
 						ProjectTarget ProjectTarget = SolutionConfigKeyValue.Value.Item2.Item1;
 
-						void AddSolutionConfig(UnrealArch? Arch, List<VCSolutionConfigCombination> OutSolutionConfigs, bool ForceArchSuffix = false)
+						void AddSolutionConfig(UnrealArch? Arch, List<VisualStudioPrimarySolutionProjectParams> OutSolutionConfigs, bool ForceArchSuffix = false)
 						{
 							// e.g.  "Development|Win64 = Development|Win64"
 							string SolutionConfigName = SolutionConfigKeyValue.Key;
@@ -1077,37 +941,30 @@ namespace UnrealBuildTool
 
 							string SolutionConfigAndPlatformPair = SolutionConfigName + "|" + SolutionPlatformName;
 							OutSolutionConfigs.Add(
-								new VCSolutionConfigCombination(SolutionConfigAndPlatformPair)
-								{
-									Configuration = Configuration,
-									Platform = CurPlatform,
-									TargetConfigurationName = TargetType,
-									Architecture = Arch
-								}
-							);
+								new VisualStudioPrimarySolutionProjectParams(SolutionConfigName, SolutionPlatformName, Configuration, CurPlatform, TargetType, Arch));
 						}
 
 						UnrealArchitectures? Architectures = GetPlatformArchitecturesToGenerate(BuildPlatform, ProjectTarget);
 						if (Architectures == null)
 						{
-							AddSolutionConfig(null, OutSolutionConfigs);
+							AddSolutionConfig(null, OutSolutionVcxProjectParamsList);
 						}
 						else
 						{
 							if (BuildPlatform.ArchitectureConfig.Mode == UnrealArchitectureMode.SingleTargetLinkSeparately && BuildPlatform.Platform == UnrealTargetPlatform.Android)
 							{
-								AddSolutionConfig(null, OutSolutionConfigs);
+								AddSolutionConfig(null, OutSolutionVcxProjectParamsList);
 
 								foreach (UnrealArch Arch in BuildPlatform.ArchitectureConfig.AllSupportedArchitectures.Architectures)
 								{
-									AddSolutionConfig(Arch, OutSolutionConfigs, true);
+									AddSolutionConfig(Arch, OutSolutionVcxProjectParamsList, true);
 								}
 							}
 							else
 							{
 								foreach (UnrealArch Arch in Architectures.Architectures)
 								{
-									AddSolutionConfig(Arch, OutSolutionConfigs);
+									AddSolutionConfig(Arch, OutSolutionVcxProjectParamsList);
 								}
 							}
 						}
@@ -1116,11 +973,11 @@ namespace UnrealBuildTool
 			}
 
 			// Sort the list of solution platform strings alphabetically (Visual Studio prefers it)
-			OutSolutionConfigs.Sort(
-				new Comparison<VCSolutionConfigCombination>(
+			OutSolutionVcxProjectParamsList.Sort(
+				new Comparison<VisualStudioPrimarySolutionProjectParams>(
 					(x, y) =>
 					{
-						return String.Compare(x.VCSolutionConfigAndPlatformName, y.VCSolutionConfigAndPlatformName,
+						return String.Compare(x.SolutionConfigurationAndPlatform, y.SolutionConfigurationAndPlatform,
 							StringComparison.OrdinalIgnoreCase);
 					}
 				)
diff --git a/Engine/Source/Programs/UnrealBuildTool/ProjectFiles/VisualStudio/VCSharpProjectFile.cs b/Engine/Source/Programs/UnrealBuildTool/ProjectFiles/VisualStudio/VCSharpProjectFile.cs
--- a/Engine/Source/Programs/UnrealBuildTool/ProjectFiles/VisualStudio/VCSharpProjectFile.cs
+++ b/Engine/Source/Programs/UnrealBuildTool/ProjectFiles/VisualStudio/VCSharpProjectFile.cs
@@ -8,6 +8,7 @@ using System.Xml;
 using EpicGames.Core;
 using Microsoft.Extensions.Logging;
 using UnrealBuildBase;
+using UnrealBuildTool.ProjectFiles.VisualStudio;
 
 namespace UnrealBuildTool;
 
@@ -183,15 +184,12 @@ internal class VCSharpProjectFile : MSBuildProjectFile
 
 	/// <inheritdoc/>
 	public override MSBuildProjectContext? GetMatchingProjectContext(
-		TargetType SolutionTarget,
-		UnrealTargetConfiguration SolutionConfiguration,
-		UnrealTargetPlatform SolutionPlatform,
+		VisualStudioPrimarySolutionProjectParams SolutionCombination,
 		PlatformProjectGeneratorCollection PlatformProjectGenerators,
-		UnrealArch? Architecture,
 		ILogger Logger)
 	{
-		bool bPreferBuildByDefault = SolutionTarget == TargetType.Game || SolutionTarget == TargetType.Editor;
-		return GetMatchingProjectContextForCSharp(bPreferBuildByDefault, SolutionConfiguration.ToString());
+		bool bPreferBuildByDefault = SolutionCombination.ProjectTargetType is TargetType.Game or TargetType.Editor;
+		return GetMatchingProjectContextForCSharp(bPreferBuildByDefault, SolutionCombination.SolutionConfiguration);
 	}
 
 	public MSBuildProjectContext GetMatchingProjectContextForCSharp(
diff --git a/Engine/Source/Programs/UnrealBuildTool/ProjectFiles/VisualStudio/VisualStudioSolutionBuilder.cs b/Engine/Source/Programs/UnrealBuildTool/ProjectFiles/VisualStudio/VisualStudioSolutionBuilder.cs
--- a/Engine/Source/Programs/UnrealBuildTool/ProjectFiles/VisualStudio/VisualStudioSolutionBuilder.cs
+++ b/Engine/Source/Programs/UnrealBuildTool/ProjectFiles/VisualStudio/VisualStudioSolutionBuilder.cs
@@ -2,372 +2,252 @@
 
 using System;
 using System.Collections.Generic;
+using System.Collections.Immutable;
+using System.IO;
 using System.Linq;
 using System.Text;
+using System.Threading;
+using EpicGames.Core;
 using Microsoft.Extensions.Logging;
-using static UnrealBuildTool.VCProjectFileGenerator;
+using Microsoft.VisualStudio.SolutionPersistence.Model;
+using Microsoft.VisualStudio.SolutionPersistence.Serializer;
 
 namespace UnrealBuildTool.ProjectFiles.VisualStudio;
 
-internal sealed class VisualStudioSolutionBuilder
+internal abstract class VisualStudioSolutionBuilder<TProjectParams> where TProjectParams : VisualStudioSolutionProjectParams
 {
-	public VisualStudioSolutionBuilder(PrimaryProjectFolder rootFolder, List<ProjectFile> allProjectFiles, ILogger logger)
+	public VisualStudioSolutionBuilder(
+		DirectoryReference solutionDirectory,
+		PrimaryProjectFolder rootFolder,
+		IReadOnlyList<TProjectParams> solutionCombinations,
+		ILogger logger)
 	{
+		_solutionDirectory = solutionDirectory;
 		_rootFolder = rootFolder;
+		_solutionCombinations = [.. solutionCombinations];
+		_solutionConfigurations = [.. solutionCombinations.Select(x => x.SolutionConfiguration).Distinct()];
+		_solutionPlatforms = [.. solutionCombinations.Select(x => x.SolutionPlatform).Distinct()];
 		_logger = logger;
-		_allProjectFiles = [.. allProjectFiles.OfType<MSBuildProjectFile>()];
-		_projectFolderGuids = GenerateProjectFolderGuids();
 	}
 
-	// This the GUID that Visual Studio uses to identify a solution folder
-	public const string SolutionFolderGuid = "{2150E333-8FDC-42A3-9474-1A3956D46DE8}";
-
-	public static string FormatGuid(Guid guid) => guid.ToString("B").ToUpperInvariant();
-
-	public void WritePreamble(VCProjectFileFormat format)
+	public string BuildSolution(VCProjectFileFormat format, bool forSlnx)
 	{
-		// Solution file header. Note that a leading newline is required for file type detection to work correctly in the shell.
+		SolutionModel solution = new();
 
-		_sb.AppendLine("Microsoft Visual Studio Solution File, Format Version 12.00");
-
-		switch (format)
+		if (!forSlnx)
 		{
-			case VCProjectFileFormat.VisualStudio2022:
-				_sb.AppendLine("# Visual Studio Version 17");
-				_sb.AppendLine("VisualStudioVersion = 17.0.31314.256");
-				break;
-			case VCProjectFileFormat.VisualStudio2026:
-				_sb.AppendLine("# Visual Studio Version 18");
-				_sb.AppendLine("VisualStudioVersion = 18.0.11205.157");
-				break;
-			default:
-				throw new BuildException("Unexpected ProjectFileFormat '{0}'", format);
-		}
+#pragma warning disable CS0618 // Write HideSolutionNode for legacy compatibility.
+			solution.VisualStudioProperties.HideSolutionNode = false;
+#pragma warning restore CS0618
+			solution.VisualStudioProperties.MinimumVersion = new(10, 0, 40219, 1);
 
-		_sb.AppendLine("MinimumVisualStudioVersion = 10.0.40219.1");
-	}
-
-	public void WriteProjectStart(string typeGuid, string name, string path, Guid guid)
-	{
-		string guidString = FormatGuid(guid);
-		_sb.AppendLine($"Project(\"{typeGuid}\") = \"{name}\", \"{path}\", \"{guidString}\"");
-	}
-
-	public void WriteProjectEnd()
-	{
-		_sb.AppendLine("EndProject");
-	}
+			switch (format)
+			{
+				case VCProjectFileFormat.VisualStudio2022:
+					solution.VisualStudioProperties.OpenWith = "Visual Studio Version 17";
+					solution.VisualStudioProperties.Version = new(17, 0, 31314, 256);
+					break;
+				case VCProjectFileFormat.VisualStudio2026:
+					solution.VisualStudioProperties.OpenWith = "Visual Studio Version 18";
+					solution.VisualStudioProperties.Version = new(18, 0, 11205, 157);
+					break;
+				default:
+					throw new BuildException("Unexpected ProjectFileFormat '{0}'", format);
+			}
+		}
 
-	public void WriteProjectFolders()
-	{
-		IEnumerable<PrimaryProjectFolder> allSolutionFolders = _projectFolderGuids.Keys.OrderBy(folder => folder.FolderName).ThenBy(folder => _projectFolderGuids[folder]);
-		foreach (PrimaryProjectFolder folder in allSolutionFolders)
+		foreach (string platform in _solutionPlatforms)
 		{
-			Guid guid = _projectFolderGuids[folder];
-
-			WriteProjectStart(SolutionFolderGuid, folder.FolderName, folder.FolderName, guid);
-			WriteProjectSolutionItems(folder);
-			WriteProjectEnd();
+			solution.AddPlatform(platform.ToString());
 		}
-	}
 
-	public void WriteProjects()
-	{
-		// Project files
-		foreach (MSBuildProjectFile project in _allProjectFiles)
+		foreach (string configuration in _solutionConfigurations)
 		{
-			// Visual Studio uses different GUID types depending on the project type
-			string typeGuid = project.ProjectTypeGUID;
-
-			// NOTE: The project name in the solution doesn't actually *have* to match the project file name on disk. However,
-			// we prefer it when it does match so we use the actual file name here.
-			string name = project.ProjectFilePath.GetFileNameWithoutExtension();
-			string relativePath = project.ProjectFilePath.MakeRelativeTo(ProjectFileGenerator.PrimaryProjectPath);
-
-			WriteProjectStart(typeGuid, name, relativePath, project.ProjectGUID);
-
-			if (!project.IsStubProject)
-			{
-				if (project.DependsOnProjects.Count > 0)
-				{
-					_sb.AppendLine("\tProjectSection(ProjectDependencies) = postProject");
-
-					// Setup any addition dependencies this project has
-					foreach (MSBuildProjectFile dependsOnProject in project.DependsOnProjects.Cast<MSBuildProjectFile>())
-					{
-						string dependsOnProjectGuid = FormatGuid(dependsOnProject.ProjectGUID);
-						_sb.AppendLine($"\t\t{dependsOnProjectGuid} = {dependsOnProjectGuid}");
-					}
-
-					_sb.AppendLine("\tEndProjectSection");
-				}
-			}
-
-			WriteProjectEnd();
+			solution.AddBuildType(configuration);
 		}
-	}
 
-	public void WriteGlobalStart()
-	{
-		_sb.AppendLine("Global");
-	}
-
-	public void WriteGlobalEnd()
-	{
-		_sb.AppendLine("EndGlobal");
-	}
+		AddUnrealVS(solution);
 
-	public void WriteSolutionConfigurationPlatforms(List<VCSolutionConfigCombination> solutionConfigCombinations)
-	{
-		_sb.AppendLine("	GlobalSection(SolutionConfigurationPlatforms) = preSolution");
+		Queue<(PrimaryProjectFolder folder, string path)> folders = new();
 
-		HashSet<string> AppendedSolutionConfigAndPlatformNames = new HashSet<string>(StringComparer.OrdinalIgnoreCase);
-		foreach (VCSolutionConfigCombination SolutionConfigCombination in solutionConfigCombinations)
+		AddProjectsAndItems(solution, null, _rootFolder, !forSlnx);
+		foreach (PrimaryProjectFolder folder in _rootFolder.SubFolders)
 		{
-			// We alias "Game" and "Program" to both have the same solution configuration, so we're careful not to add the same combination twice.
-			if (!AppendedSolutionConfigAndPlatformNames.Contains(SolutionConfigCombination.VCSolutionConfigAndPlatformName))
-			{
-				_sb.AppendLine($"		{SolutionConfigCombination.VCSolutionConfigAndPlatformName} = {SolutionConfigCombination.VCSolutionConfigAndPlatformName}");
-				AppendedSolutionConfigAndPlatformNames.Add(SolutionConfigCombination.VCSolutionConfigAndPlatformName);
-			}
+			folders.Enqueue((folder, $"/{folder.FolderName}/"));
 		}
 
-		_sb.AppendLine("	EndGlobalSection");
-	}
-
-	public void WriteProjectConfigurationPlatforms(
-		List<VCSolutionConfigCombination> solutionConfigCombinations,
-		PlatformProjectGeneratorCollection platformProjectGenerators,
-		Predicate<MSBuildProjectFile> debugPredicate)
-	{
-		_sb.AppendLine("	GlobalSection(ProjectConfigurationPlatforms) = postSolution");
-
-		foreach (MSBuildProjectFile project in _allProjectFiles)
+		while (folders.Count > 0)
 		{
-			foreach (VCSolutionConfigCombination solutionConfigCombination in solutionConfigCombinations)
+			(PrimaryProjectFolder folder, string path) = folders.Dequeue();
+			SolutionFolderModel solutionFolder = solution.AddFolder(path);
+			if (!forSlnx)
 			{
-				// Get the context for the current solution context
-				MSBuildProjectContext? projectContext = project.GetMatchingProjectContext(
-					solutionConfigCombination.TargetConfigurationName,
-					solutionConfigCombination.Configuration,
-					solutionConfigCombination.Platform,
-					platformProjectGenerators,
-					solutionConfigCombination.Architecture,
-					_logger
-				);
-
-				if (projectContext == null)
+				// Legacy compatibility: Set the ID of the generated folder.
+				if (path == "/Visualizers/")
 				{
-					continue;
+					solutionFolder.Id = Guid.Parse("1CCEC849-CC72-4C59-8C36-2F7C38706D4C");
 				}
-
-				// Override the configuration to build for UBT
-				if (debugPredicate(project))
+				else
 				{
-					projectContext.ConfigurationName = "Debug";
+					solutionFolder.Id = VCProjectFileGenerator.MakeMd5Guid($"UE5{path[..^1]}");
 				}
+			}
 
-				// Write the solution mapping (e.g.  "{4232C52C-680F-4850-8855-DC39419B5E9B}.Debug|iOS.ActiveCfg = iOS_Debug|Win32")
-				string guid = FormatGuid(project.ProjectGUID);
-				string configAndPlatformName = solutionConfigCombination.VCSolutionConfigAndPlatformName;
+			AddProjectsAndItems(solution, solutionFolder, folder, !forSlnx);
 
-				_sb.AppendLine($"		{guid}.{configAndPlatformName}.ActiveCfg = {projectContext.Name}");
-				if (projectContext.bBuildByDefault)
-				{
-					_sb.AppendLine($"		{guid}.{configAndPlatformName}.Build.0 = {projectContext.Name}");
-					if (projectContext.bDeployByDefault)
-					{
-						_sb.AppendLine($"		{guid}.{configAndPlatformName}.Deploy.0 = {projectContext.Name}");
-					}
-				}
+			foreach (PrimaryProjectFolder childFolder in folder.SubFolders)
+			{
+				folders.Enqueue((childFolder, $"{path}{childFolder.FolderName}/"));
 			}
 		}
 
-		_sb.AppendLine("	EndGlobalSection");
-	}
+		solution.DistillProjectConfigurations();
 
-	public void WriteSolutionConfigurationPlatformsForCSharpOnlySolution(List<string> solutionConfigurations)
-	{
-		_sb.AppendLine("	GlobalSection(SolutionConfigurationPlatforms) = preSolution");
-
-		foreach (string solutionConfiguration in solutionConfigurations)
+		using MemoryStream memoryStream = new();
+		UTF8Encoding encoding = new UTF8Encoding(encoderShouldEmitUTF8Identifier: false);
+		if (forSlnx)
+		{
+			SolutionSerializers.SlnXml.SaveAsync(memoryStream, solution, CancellationToken.None).Wait();
+		}
+		else
 		{
-			string combination = $"{solutionConfiguration}|Any CPU";
-			_sb.AppendLine($"		{combination} = {combination}");
+			SolutionSerializers.SlnFileV12.SaveAsync(memoryStream, solution, CancellationToken.None).Wait();
 		}
 
-		_sb.AppendLine("	EndGlobalSection");
+		return encoding.GetString(memoryStream.ToArray());
 	}
 
-	public void WriteProjectConfigurationPlatformsForCSharpOnlySolution(List<string> solutionConfigurations)
+	private void AddProjectsAndItems(SolutionModel solution, SolutionFolderModel? solutionFolder, PrimaryProjectFolder folder, bool setID)
 	{
-		_sb.AppendLine("	GlobalSection(ProjectConfigurationPlatforms) = postSolution");
-
-		foreach (VCSharpProjectFile project in _allProjectFiles.Cast<VCSharpProjectFile>())
+		foreach (ProjectFile projectFile in folder.ChildProjects)
 		{
-			foreach (string solutionConfiguration in solutionConfigurations)
+			// Cast must be valid here, there is a bug further up if not.
+			MSBuildProjectFile msBuildProjectFile = (MSBuildProjectFile)projectFile;
+			string relativeProjectPath = projectFile.ProjectFilePath.MakeRelativeTo(_solutionDirectory);
+			SolutionProjectModel project = solution.AddProject(relativeProjectPath, folder: solutionFolder);
+			// To prevent the solution from being marked dirty, always set the ID if it's a C++ project.
+			if (projectFile is VCProjectFile || setID)
 			{
-				MSBuildProjectContext projectContext = project.GetMatchingProjectContextForCSharp(bPreferBuildByDefault: true, solutionConfiguration);
-
-				// Write the solution mapping (e.g.  "{4232C52C-680F-4850-8855-DC39419B5E9B}.Debug|iOS.ActiveCfg = iOS_Debug|Win32")
-				string guid = FormatGuid(project.ProjectGUID);
-				string configAndPlatformName = $"{solutionConfiguration}|Any CPU";
-
-				_sb.AppendLine($"		{guid}.{configAndPlatformName}.ActiveCfg = {projectContext.Name}");
-				if (projectContext.bBuildByDefault)
-				{
-					_sb.AppendLine($"		{guid}.{configAndPlatformName}.Build.0 = {projectContext.Name}");
-					if (projectContext.bDeployByDefault)
-					{
-						_sb.AppendLine($"		{guid}.{configAndPlatformName}.Deploy.0 = {projectContext.Name}");
-					}
-				}
+				project.Id = msBuildProjectFile.ProjectGUID;
 			}
-		}
 
-		_sb.AppendLine("	EndGlobalSection");
-	}
-
-	public void WriteUnrealVs(HashSet<UnrealTargetPlatform> validPlatforms)
-	{
-		const string UnrealVSGuid = "ddbf523f-7eb6-4887-bd51-85a714ff87eb";
-		string validPlatformsString = String.Join(";", validPlatforms.Select(platform => platform.ToString()));
-
-		_sb.AppendLine("\t# UnrealVS Section");
-		_sb.AppendLine($"\tGlobalSection({UnrealVSGuid}) = preSolution");
-		_sb.AppendLine($"\t\tAvailablePlatforms={validPlatformsString}");
-		_sb.AppendLine("\tEndGlobalSection");
-	}
-
-	public void WriteSolutionProperties()
-	{
-		// HideSolutionNode sets whether or not the top-level solution entry is completely hidden in the UI.
-		// We don't want that, as we need users to be able to right click on the solution tree item.
-		_sb.AppendLine("	GlobalSection(SolutionProperties) = preSolution");
-		_sb.AppendLine("		HideSolutionNode = FALSE");
-		_sb.AppendLine("	EndGlobalSection");
-	}
-
-	public void WriteNestedProjects()
-	{
-		_sb.AppendLine("	GlobalSection(NestedProjects) = preSolution");
-
-		// Every entry in this section is in the format "Guid1 = Guid2".  Guid1 is the child project (or solution
-		// filter)'s GUID, and Guid2 is the solution filter directory to parent the child project (or solution
-		// filter) to. This sets up the hierarchical solution explorer tree for all solution folders and projects.
-
-		void FolderProcessorFunction(List<PrimaryProjectFolder> localFolders)
-		{
-			foreach (PrimaryProjectFolder folder in localFolders)
+			foreach (TProjectParams combination in _solutionCombinations)
 			{
-				if (!_projectFolderGuids.TryGetValue(folder, out Guid folderGuid))
+				string solutionConfiguration = combination.SolutionConfiguration;
+				string solutionPlatform = combination.SolutionPlatform;
+				MSBuildProjectContext? projectContext = GetContextForProject(msBuildProjectFile, combination);
+
+				if (projectContext is null)
 				{
 					continue;
 				}
 
-				string folderGuidString = FormatGuid(_projectFolderGuids[folder]);
+				AdjustProjectContext(projectContext, msBuildProjectFile);
 
-				foreach (MSBuildProjectFile childProject in folder.ChildProjects)
+				// Only set these if they're not the default values.
+				if (!String.Equals(projectContext.ConfigurationName, solutionConfiguration, StringComparison.Ordinal))
 				{
-					if (_allProjectFiles.Contains(childProject))
-					{
-						//	e.g. "{BF6FB09F-A2A6-468F-BE6F-DEBE07EAD3EA} = {C43B6BB5-3EF0-4784-B896-4099753BCDA9}"
-						_sb.AppendLine($"		{FormatGuid(childProject.ProjectGUID)} = {folderGuidString}");
-					}
+					project.AddProjectConfigurationRule(new(BuildDimension.BuildType, solutionConfiguration, solutionPlatform, projectContext.ConfigurationName));
 				}
-
-				foreach (PrimaryProjectFolder subFolder in folder.SubFolders)
+				if (!String.Equals(projectContext.PlatformName, solutionPlatform, StringComparison.Ordinal))
 				{
-					// if this folder has no projects anywhere under it then skip it
-					if (_projectFolderGuids.TryGetValue(subFolder, out Guid subFolderGuid))
-					{
-						//	e.g. "{BF6FB09F-A2A6-468F-BE6F-DEBE07EAD3EA} = {C43B6BB5-3EF0-4784-B896-4099753BCDA9}"
-						_sb.AppendLine($"		{FormatGuid(subFolderGuid)} = {folderGuidString}");
-					}
+					project.AddProjectConfigurationRule(new(BuildDimension.Platform, solutionConfiguration, solutionPlatform, projectContext.PlatformName));
+				}
+				if (!projectContext.bBuildByDefault)
+				{
+					project.AddProjectConfigurationRule(new(BuildDimension.Build, solutionConfiguration, solutionPlatform, "false"));
+				}
+				if (projectContext.bDeployByDefault)
+				{
+					project.AddProjectConfigurationRule(new(BuildDimension.Deploy, solutionConfiguration, solutionPlatform, "true"));
 				}
-
-				// Recurse into subfolders
-				FolderProcessorFunction(folder.SubFolders);
 			}
 		}
 
-		FolderProcessorFunction(_rootFolder.SubFolders);
-
-		_sb.AppendLine("	EndGlobalSection");
-	}
-
-	/// <summary>
-	/// Add any files that are inlined right inside the solution folder
-	/// </summary>
-	public void WriteProjectSolutionItems(PrimaryProjectFolder folder)
-	{
-		if (folder.Files.Count > 0)
+		if (solutionFolder is not null)
 		{
-			_sb.AppendLine("	ProjectSection(SolutionItems) = preProject");
-			foreach (string CurFile in folder.Files)
+			foreach (string item in folder.Files)
 			{
-				// Syntax is:  <relative file path> = <relative file path>
-				_sb.AppendLine("		" + CurFile + " = " + CurFile);
+				solutionFolder.AddFile(item);
 			}
-			_sb.AppendLine("	EndProjectSection");
 		}
 	}
 
-	private bool FolderHasAnyProjects(PrimaryProjectFolder ProjectFolder)
+	protected virtual void AddUnrealVS(SolutionModel solution)
 	{
-		foreach (MSBuildProjectFile ChildProject in ProjectFolder.ChildProjects)
-		{
-			if (_allProjectFiles.Contains(ChildProject))
-			{
-				return true;
-			}
-		}
+	}
 
-		foreach (PrimaryProjectFolder SubFolder in ProjectFolder.SubFolders)
-		{
-			if (FolderHasAnyProjects(SubFolder))
-			{
-				return true;
-			}
-		}
+	protected virtual void AdjustProjectContext(MSBuildProjectContext projectContext, MSBuildProjectFile projectFile)
+	{
+	}
+
+	protected abstract MSBuildProjectContext? GetContextForProject(MSBuildProjectFile projectFile, TProjectParams projectParams);
 
-		return false;
+	private readonly DirectoryReference _solutionDirectory;
+	private readonly PrimaryProjectFolder _rootFolder;
+	private readonly ImmutableArray<TProjectParams> _solutionCombinations;
+	private readonly ImmutableArray<string> _solutionConfigurations;
+	private readonly ImmutableArray<string> _solutionPlatforms;
+	protected readonly ILogger _logger;
+}
+
+internal sealed class PrimaryVisualStudioSolutionBuilder : VisualStudioSolutionBuilder<VisualStudioPrimarySolutionProjectParams>
+{
+	public PrimaryVisualStudioSolutionBuilder(
+		DirectoryReference solutionDirectory,
+		PrimaryProjectFolder rootFolder,
+		IReadOnlyList<VisualStudioPrimarySolutionProjectParams> solutionProjectParams,
+		PlatformProjectGeneratorCollection platformProjectGeneratorCollection,
+		IReadOnlyList<UnrealTargetPlatform> unrealVSPlatforms,
+		Predicate<MSBuildProjectFile> debugPredicate,
+		ILogger logger
+	) : base(solutionDirectory, rootFolder, solutionProjectParams, logger)
+	{
+		_platformProjectGeneratorCollection = platformProjectGeneratorCollection;
+		_debugPredicate = debugPredicate;
+		_unrealVSPlatformsString = String.Join(";", unrealVSPlatforms.Select(platform => platform.ToString()));
 	}
 
-	private Dictionary<PrimaryProjectFolder, Guid> GenerateProjectFolderGuids()
+	protected override void AddUnrealVS(SolutionModel solution)
 	{
-		Dictionary<PrimaryProjectFolder, Guid> guids = [];
-		foreach (PrimaryProjectFolder folder in _rootFolder.SubFolders)
-		{
-			GenerateProjectFolderGuids("UE5", folder, guids);
-		}
-		return guids;
+		const string UnrealVSGuid = "ddbf523f-7eb6-4887-bd51-85a714ff87eb";
+		SolutionPropertyBag unrealVSProperties = solution.AddProperties(UnrealVSGuid);
+		unrealVSProperties.Add("AvailablePlatforms", _unrealVSPlatformsString);
 	}
 
-	private void GenerateProjectFolderGuids(
-		string parentPath,
-		PrimaryProjectFolder folder,
-		Dictionary<PrimaryProjectFolder, Guid> guids)
+	protected override void AdjustProjectContext(MSBuildProjectContext projectContext, MSBuildProjectFile projectFile)
 	{
-		if (!FolderHasAnyProjects(folder))
+		if (_debugPredicate(projectFile))
 		{
-			return;
+			projectContext.ConfigurationName = "Debug";
 		}
+	}
+
+	protected override MSBuildProjectContext? GetContextForProject(MSBuildProjectFile projectFile, VisualStudioPrimarySolutionProjectParams projectParams)
+	{
+		return projectFile.GetMatchingProjectContext(projectParams, _platformProjectGeneratorCollection, _logger);
+	}
+
+	private readonly PlatformProjectGeneratorCollection _platformProjectGeneratorCollection;
+	private readonly Predicate<MSBuildProjectFile> _debugPredicate;
+	private readonly string _unrealVSPlatformsString;
 
-		string path = String.Format("{0}/{1}", parentPath, folder.FolderName);
-		guids[folder] = MakeMd5Guid(path);
+}
 
-		foreach (PrimaryProjectFolder subFolder in folder.SubFolders)
+internal sealed class AutomationVisualStudioSolutionBuilder(
+	DirectoryReference solutionDirectory,
+	PrimaryProjectFolder rootFolder,
+	IReadOnlyList<VisualStudioSolutionProjectParams> solutionProjectParams,
+	ILogger logger
+) : VisualStudioSolutionBuilder<VisualStudioSolutionProjectParams>(solutionDirectory, rootFolder, solutionProjectParams, logger)
+{
+	protected override MSBuildProjectContext? GetContextForProject(MSBuildProjectFile projectFile, VisualStudioSolutionProjectParams projectParams)
+	{
+		if (projectFile is not VCSharpProjectFile cSharpProjectFile)
 		{
-			GenerateProjectFolderGuids(path, subFolder, guids);
+			throw new ArgumentException($"{nameof(AutomationVisualStudioSolutionBuilder)} only supports C# projects.");
 		}
-	}
 
-	public override string ToString() => _sb.ToString();
-
-	private readonly StringBuilder _sb = new();
-	private readonly PrimaryProjectFolder _rootFolder;
-	private readonly ILogger _logger;
-	private readonly List<MSBuildProjectFile> _allProjectFiles;
-	private readonly Dictionary<PrimaryProjectFolder, Guid> _projectFolderGuids;
+		return cSharpProjectFile.GetMatchingProjectContextForCSharp(bPreferBuildByDefault: true, projectParams.SolutionConfiguration);
+	}
 }
diff --git a/Engine/Source/Programs/UnrealBuildTool/ProjectFiles/VisualStudio/VisualStudioSolutionProjectParams.cs b/Engine/Source/Programs/UnrealBuildTool/ProjectFiles/VisualStudio/VisualStudioSolutionProjectParams.cs
new file mode 100644
--- /dev/null
+++ b/Engine/Source/Programs/UnrealBuildTool/ProjectFiles/VisualStudio/VisualStudioSolutionProjectParams.cs
@@ -0,0 +1,40 @@
+// Copyright Epic Games, Inc. All Rights Reserved.
+
+using System;
+
+namespace UnrealBuildTool.ProjectFiles.VisualStudio;
+
+/// <summary>
+/// Parameters to configure a project inside a Visual Studio solution.
+/// </summary>
+internal class VisualStudioSolutionProjectParams(string solutionConfiguration, string solutionPlatform)
+{
+	public string SolutionConfiguration { get; } = solutionConfiguration;
+
+	public string SolutionPlatform { get; } = solutionPlatform;
+
+	public string SolutionConfigurationAndPlatform => $"{SolutionConfiguration}|{SolutionPlatform}";
+}
+
+/// <summary>
+/// Parameters to configure a project inside the primary Visual Studio solution for this Unreal Engine project.
+/// </summary>
+internal sealed class VisualStudioPrimarySolutionProjectParams(
+	string solutionConfiguration,
+	string solutionPlatform,
+	UnrealTargetConfiguration projectConfiguration,
+	UnrealTargetPlatform projectPlatform,
+	TargetType projectTargetType,
+	UnrealArch? projectArchitecture) :
+	VisualStudioSolutionProjectParams(solutionConfiguration, solutionPlatform)
+{
+	public override string ToString() => $"{SolutionConfiguration}|{SolutionPlatform}={ProjectConfiguration} {ProjectPlatform} {ProjectTargetType}{(ProjectArchitecture != null ? " " + ProjectArchitecture : String.Empty)}";
+
+	public UnrealTargetConfiguration ProjectConfiguration { get; } = projectConfiguration;
+
+	public UnrealTargetPlatform ProjectPlatform { get; } = projectPlatform;
+
+	public TargetType ProjectTargetType { get; } = projectTargetType;
+
+	public UnrealArch? ProjectArchitecture { get; } = projectArchitecture;
+}
