IWYU improvements
- Add VFS support
- Fix duplicate includes in deprecation scope on repeated IWYU runs
- Header compilation should match non-IWYU: compile headers via .h.cpp files instead of directly
- -IncludeHeaders was not using Shared response files
- Minor C# code style improvements

#Synced-CL 51674282

diff --git a/Engine/Source/Programs/UnrealBuildTool/Configuration/Rules/TargetRules.cs b/Engine/Source/Programs/UnrealBuildTool/Configuration/Rules/TargetRules.cs
--- a/Engine/Source/Programs/UnrealBuildTool/Configuration/Rules/TargetRules.cs
+++ b/Engine/Source/Programs/UnrealBuildTool/Configuration/Rules/TargetRules.cs
@@ -3381,7 +3381,6 @@ namespace UnrealBuildTool
 				bUsePCHFiles = false; // No pch for iwyu
 				bUseVFS = false; // Do not use vfs because IWYUMode can't handle virtual paths (yet)
 				bUseUnityBuild = false; // No unity for iwyu
-				bUseVerseBPVM = false;
 				CppStandard = CppStandardVersion.Cpp23; // Must be set to 23 to handle function default values same as msvc
 				CppStandardEngine = CppStandardVersion.Cpp23; // Must be set to 23 to handle function default values same as msvc
 				// bAlwaysUseUnityForGeneratedFiles = false; // This should be used but adds too many extra files atm.
diff --git a/Engine/Source/Programs/UnrealBuildTool/Configuration/UEBuildModuleCPP.cs b/Engine/Source/Programs/UnrealBuildTool/Configuration/UEBuildModuleCPP.cs
--- a/Engine/Source/Programs/UnrealBuildTool/Configuration/UEBuildModuleCPP.cs
+++ b/Engine/Source/Programs/UnrealBuildTool/Configuration/UEBuildModuleCPP.cs
@@ -1216,18 +1216,8 @@ namespace UnrealBuildTool
 				List<FileItem> HeaderFileItems = GetCompilableHeaders(InputFiles, CppCompileEnvironment);
 				if (HeaderFileItems.Count > 0)
 				{
-					CppCompileEnvironment HeaderCompileEnvironment;
-
-					if (Target.bIWYU)
-					{
-						// We want to use CppCompileEnvironment with shared rsp etc in iwyu case
-						HeaderCompileEnvironment = new CppCompileEnvironment(CppCompileEnvironment);
-					}
-					else
-					{
-						HeaderCompileEnvironment = new CppCompileEnvironment(ModuleCompileEnvironment);
-						HeaderCompileEnvironment.bUseUnity = false;
-					}
+					CppCompileEnvironment HeaderCompileEnvironment = new(CppCompileEnvironment);
+					HeaderCompileEnvironment.bUseUnity = false;
 
 					DirectoryReference HeaderIntermediateDirectory = DirectoryReference.Combine(IntermediateDirectory, "H");
 					CreateHeaderForDefinitions(HeaderCompileEnvironment, HeaderIntermediateDirectory, "h", Graph);
@@ -1240,28 +1230,20 @@ namespace UnrealBuildTool
 						List<FileItem> HeadersToCompile = new();
 						foreach (FileItem Header in HeaderFileItems.Where(x => x.Directory == ParentDir))
 						{
-							if (Target.bIWYU)
+							string IncludeFileString = Header.AbsolutePath;
+							if (HeaderCompileEnvironment.RootPaths.GetVfsOverlayPath(Header.Location, out string? VfsPath))
 							{
-								// IWYU needs the actual headers as main file input
-								HeadersToCompile.Add(Header);
+								IncludeFileString = VfsPath;
 							}
-							else
+							else if (Header.Location.IsUnderDirectory(Unreal.RootDirectory))
 							{
-								string IncludeFileString = Header.AbsolutePath;
-								if (HeaderCompileEnvironment.RootPaths.GetVfsOverlayPath(Header.Location, out string? VfsPath))
-								{
-									IncludeFileString = VfsPath;
-								}
-								else if (Header.Location.IsUnderDirectory(Unreal.RootDirectory))
-								{
-									IncludeFileString = Header.Location.MakeRelativeTo(Unreal.EngineSourceDirectory);
-								}
-
-								List<string> GeneratedHeaderCppContents = UEBuildModuleCPP.GenerateHeaderCpp(Header.Name, IncludeFileString);
-								FileItem GeneratedHeaderCpp = FileItem.GetItemByFileReference(FileReference.Combine(HeaderSubIntermediateDirectory, $"{Header.Name}.cpp"));
-								Graph.CreateIntermediateTextFile(GeneratedHeaderCpp, GeneratedHeaderCppContents);
-								HeadersToCompile.Add(GeneratedHeaderCpp);
+								IncludeFileString = Header.Location.MakeRelativeTo(Unreal.EngineSourceDirectory);
 							}
+
+							List<string> GeneratedHeaderCppContents = UEBuildModuleCPP.GenerateHeaderCpp(Header.Name, IncludeFileString);
+							FileItem GeneratedHeaderCpp = FileItem.GetItemByFileReference(FileReference.Combine(HeaderSubIntermediateDirectory, $"{Header.Name}.cpp"));
+							Graph.CreateIntermediateTextFile(GeneratedHeaderCpp, GeneratedHeaderCppContents);
+							HeadersToCompile.Add(GeneratedHeaderCpp);
 						}
 
 						// Add the compile actions
diff --git a/Engine/Source/Programs/UnrealBuildTool/Modes/IWYUMode.cs b/Engine/Source/Programs/UnrealBuildTool/Modes/IWYUMode.cs
--- a/Engine/Source/Programs/UnrealBuildTool/Modes/IWYUMode.cs
+++ b/Engine/Source/Programs/UnrealBuildTool/Modes/IWYUMode.cs
@@ -803,6 +803,13 @@ namespace UnrealBuildTool
 						bool isHeader = !isCpp && info.File.EndsWith(".h", StringComparison.Ordinal);
 						bool isGeneratedHeader = isHeader && (info.File.EndsWith(".generated.h", StringComparison.Ordinal) || info.File.EndsWith(".gen.h", StringComparison.Ordinal));
 
+						// Set Source on regular headers so they pass the Source != null check in IsValidForUpdate
+						// and get processed by UpdateFile. Generated headers intentionally keep Source null.
+						if (isHeader && !isGeneratedHeader)
+						{
+							info.Source = IWYUFile;
+						}
+
 						lock (Infos)
 						{
 							if (!Infos.TryAdd(info.File, info))
@@ -821,10 +828,7 @@ namespace UnrealBuildTool
 
 						if (isGeneratedHeader)
 						{
-							if (generatedHeaders == null)
-							{
-								generatedHeaders = new();
-							}
+							generatedHeaders ??= new();
 
 							generatedHeaders.Add(info);
 						}
@@ -1219,6 +1223,8 @@ namespace UnrealBuildTool
 			// If we have built .gen.cpp it means that it is not inlined in another cpp file
 			// And we need to promote the includes needed to the header since we can never modify the .gen.cpp file
 			Logger.LogInformation($"Transferring needed includes from .gen.cpp to owning file...");
+
+			// Safe to parallelize: each .gen.cpp resolves to a unique IncluderInfo
 			Parallel2.ForEach(GeneratedCppInfos, entry =>
 			{
 				IWYUInfo GeneratedCpp = entry.Item1;
@@ -2138,6 +2144,12 @@ namespace UnrealBuildTool
 						else
 						{
 							// Scan the already added includes to prevent additional adds.
+							int EndIndex = IndexOfDeprecateScope;
+							while (EndIndex < NewLines.Count && NewLines[EndIndex] != "#endif")
+							{
+								LinesRemoved.Remove(NewLines[EndIndex]);
+								EndIndex++;
+							}
 						}
 
 						NewLines.InsertRange(IndexOfDeprecateScope, LinesRemoved);
@@ -2202,24 +2214,23 @@ namespace UnrealBuildTool
 				return 0;
 			}
 
-			List<System.Diagnostics.Process> P4Processes = new();
+			List<Process> P4Processes = new();
 
-			Action<string> AddP4Process = (Arguments) =>
+			void AddP4Process(string Arguments)
 			{
-				System.Diagnostics.Process Process = new System.Diagnostics.Process();
-				System.Diagnostics.ProcessStartInfo StartInfo = new System.Diagnostics.ProcessStartInfo();
-				Process.StartInfo.WindowStyle = System.Diagnostics.ProcessWindowStyle.Hidden;
+				Process Process = new Process();
+				Process.StartInfo.WindowStyle = ProcessWindowStyle.Hidden;
 				Process.StartInfo.CreateNoWindow = true;
 				Process.StartInfo.FileName = "p4.exe";
 				Process.StartInfo.Arguments = Arguments;
 				Process.Start();
 				P4Processes.Add(Process);
-			};
+			}
 
-			Func<bool> WaitForP4 = () =>
+			bool WaitForP4()
 			{
 				bool P4Success = true;
-				foreach (System.Diagnostics.Process P4Process in P4Processes)
+				foreach (Process P4Process in P4Processes)
 				{
 					P4Process.WaitForExit();
 					if (P4Process.ExitCode != 0)
@@ -2231,7 +2242,7 @@ namespace UnrealBuildTool
 				}
 				P4Processes.Clear();
 				return P4Success;
-			};
+			}
 
 			if (!bNoP4)
 			{
@@ -2472,7 +2483,7 @@ namespace UnrealBuildTool
 				return -1;
 			});
 
-			Func<long, string> PrettySize = (Size) =>
+			string PrettySize(long Size)
 			{
 				if (Size > 1024 * 1024)
 				{
@@ -2485,7 +2496,7 @@ namespace UnrealBuildTool
 				}
 
 				return Size + "b";
-			};
+			}
 
 			int PrintCount = 20;
 
diff --git a/Engine/Source/Programs/UnrealBuildTool/ToolChain/IWYUToolChain.cs b/Engine/Source/Programs/UnrealBuildTool/ToolChain/IWYUToolChain.cs
--- a/Engine/Source/Programs/UnrealBuildTool/ToolChain/IWYUToolChain.cs
+++ b/Engine/Source/Programs/UnrealBuildTool/ToolChain/IWYUToolChain.cs
@@ -92,6 +92,10 @@ namespace UnrealBuildTool
         /// <inheritdoc/>
         public override string NormalizeCommandLinePath(FileSystemReference Reference, CppRootPaths RootPath)
 		{
+			if (RootPath.GetVfsOverlayPath(Reference, out string? VfsPath))
+			{
+				return VfsPath;
+			}
 			return Reference.FullName.Replace('\\', '/');
 		}
 
@@ -254,7 +258,8 @@ namespace UnrealBuildTool
 			FileItem fileItem = base.GetCompileArguments_FileType(CompileEnvironment, SourceFile, OutputDir, Arguments, CompileAction, CompileResult);
 
 			Arguments.Add($"-Xiwyu --write_json_path=\"{fileItem.AbsolutePath.Replace('\\', '/')}\"");
-			if (SourceFile.HasExtension(".cpp"))
+			bool bIsHeaderCpp = SourceFile.HasExtension(".cpp") && SourceFile.Name.EndsWith(".h.cpp", StringComparison.Ordinal);
+			if (SourceFile.HasExtension(".cpp") && !bIsHeaderCpp)
 			{
 				if (CompileEnvironment.FileInlineSourceMap.TryGetValue(SourceFile, out HashSet<FileItem>? InlinedFiles))
 				{
@@ -277,7 +282,7 @@ namespace UnrealBuildTool
 			}
 			else
 			{
-				// We want to build header files first
+				// We want to build header files (and .h.cpp wrappers) first
 				CompileAction.bIsHighPriority = true;
 				Arguments.Add($"-Xiwyu --check_also=*/{SourceFile.Location.GetFileNameWithoutAnyExtensions()}.generated.h");
 			}
