UnrealBuildTool: Fix .lib files not being added to BuildProducts

#rnx

#Synced-CL 51691096

diff --git a/Engine/Source/Programs/UnrealBuildTool/Platform/Windows/VCToolChain.cs b/Engine/Source/Programs/UnrealBuildTool/Platform/Windows/VCToolChain.cs
--- a/Engine/Source/Programs/UnrealBuildTool/Platform/Windows/VCToolChain.cs
+++ b/Engine/Source/Programs/UnrealBuildTool/Platform/Windows/VCToolChain.cs
@@ -4048,7 +4048,7 @@ namespace UnrealBuildTool
 
 		public override void ModifyBuildProducts(ReadOnlyTargetRules Target, UEBuildBinary Binary, IEnumerable<string> Libraries, IEnumerable<UEBuildBundleResource> BundleResources, Dictionary<FileReference, BuildProductType> BuildProducts)
 		{
-			if (Binary.Type == UEBuildBinaryType.DynamicLinkLibrary && !Target.WindowsPlatform.bAllowClangLinker)
+			if (Binary.Type == UEBuildBinaryType.DynamicLinkLibrary && (Target.WindowsPlatform.Compiler.IsMSVC() || !Target.WindowsPlatform.bAllowClangLinker))
 			{
 				if (Target.bShouldCompileAsDLL)
 				{
