Use MallocBinned3 by default for Windows

#Synced-CL 51656536

diff --git a/Engine/Source/Programs/UnrealBuildTool/Platform/Windows/UEBuildWindows.cs b/Engine/Source/Programs/UnrealBuildTool/Platform/Windows/UEBuildWindows.cs
--- a/Engine/Source/Programs/UnrealBuildTool/Platform/Windows/UEBuildWindows.cs
+++ b/Engine/Source/Programs/UnrealBuildTool/Platform/Windows/UEBuildWindows.cs
@@ -1494,7 +1494,7 @@ namespace UnrealBuildTool
 
 			if (Target.StaticAllocator == StaticAllocatorType.None)
 			{
-				Target.StaticAllocator = StaticAllocatorType.Binned2;
+				Target.StaticAllocator = StaticAllocatorType.Binned3;
 			}
 			if (Target.WindowsPlatform.bEnableAddressSanitizer)
 			{
