UnrealBuildTool: PCH Variant for FPSemantics and OptimizationLevel

#rnx
#Synced-CL 51772246

diff --git a/Engine/Source/Programs/UnrealBuildTool/Configuration/Rules/ModuleRules.cs b/Engine/Source/Programs/UnrealBuildTool/Configuration/Rules/ModuleRules.cs
--- a/Engine/Source/Programs/UnrealBuildTool/Configuration/Rules/ModuleRules.cs
+++ b/Engine/Source/Programs/UnrealBuildTool/Configuration/Rules/ModuleRules.cs
@@ -752,7 +752,7 @@ namespace UnrealBuildTool
 		private OptimizationMode? OptimizationLevelOverride = null;
 
 		/// <summary>
-		/// Allows fine tuning optimization level for speed and\or code size. This requires a private PCH (or NoPCHs, which is not recommended)
+		/// Allows fine tuning optimization level for speed and\or code size.
 		/// </summary>
 		public OptimizationMode OptimizationLevel
 		{
@@ -779,7 +779,7 @@ namespace UnrealBuildTool
 		}
 
 		/// <summary>
-		/// Allows overriding the FP semantics for this module. This requires a private PCH (or NoPCHs, which is not recommended)
+		/// Allows overriding the FP semantics for this module.
 		/// </summary>
 		public FPSemanticsMode FPSemantics
 		{
diff --git a/Engine/Source/Programs/UnrealBuildTool/Configuration/UEBuildModuleCPP.cs b/Engine/Source/Programs/UnrealBuildTool/Configuration/UEBuildModuleCPP.cs
--- a/Engine/Source/Programs/UnrealBuildTool/Configuration/UEBuildModuleCPP.cs
+++ b/Engine/Source/Programs/UnrealBuildTool/Configuration/UEBuildModuleCPP.cs
@@ -1771,6 +1771,15 @@ namespace UnrealBuildTool
 			{
 				return false;
 			}
+			else if (ModuleCompileEnvironment.bOptimizeCode && ModuleCompileEnvironment.OptimizationLevel != CompileEnvironment.OptimizationLevel)
+			{
+				return false;
+			}
+
+			if (ModuleCompileEnvironment.FPSemantics != CompileEnvironment.FPSemantics)
+			{
+				return false;
+			}
 
 			if (ModuleCompileEnvironment.bUseRTTI != CompileEnvironment.bUseRTTI)
 			{
@@ -1850,13 +1859,21 @@ namespace UnrealBuildTool
 			{
 				if (CompileEnvironment.bOptimizeCode)
 				{
-					Variant += ".Optimized";
+					Variant += $".Opt{CompileEnvironment.OptimizationLevel}";
 				}
 				else
 				{
-					Variant += ".NonOptimized";
+					Variant += ".NonOpt";
 				}
 			}
+			else if (CompileEnvironment.bOptimizeCode && CompileEnvironment.OptimizationLevel != BaseCompileEnvironment.OptimizationLevel)
+			{
+				Variant += $".Opt{CompileEnvironment.OptimizationLevel}";
+			}
+			if (CompileEnvironment.FPSemantics != BaseCompileEnvironment.FPSemantics)
+			{
+				Variant += $".FP{CompileEnvironment.FPSemantics}";
+			}
 			if (CompileEnvironment.bUseRTTI != BaseCompileEnvironment.bUseRTTI)
 			{
 				if (CompileEnvironment.bUseRTTI)
@@ -1965,6 +1982,8 @@ namespace UnrealBuildTool
 		private static void CopySettingsForSharedPCH(CppCompileEnvironment ModuleCompileEnvironment, CppCompileEnvironment CompileEnvironment)
 		{
 			CompileEnvironment.bOptimizeCode = ModuleCompileEnvironment.bOptimizeCode;
+			CompileEnvironment.OptimizationLevel = ModuleCompileEnvironment.OptimizationLevel;
+			CompileEnvironment.FPSemantics = ModuleCompileEnvironment.FPSemantics;
 			CompileEnvironment.bCodeCoverage = ModuleCompileEnvironment.bCodeCoverage;
 			CompileEnvironment.bUseRTTI = ModuleCompileEnvironment.bUseRTTI;
 			CompileEnvironment.bEnableExceptions = ModuleCompileEnvironment.bEnableExceptions;
@@ -2623,6 +2642,8 @@ namespace UnrealBuildTool
 			Result.bValidateExperimentalApi = Rules.bValidateExperimentalApi;
 			Result.bUseAutoRTFMCompiler = Target.bUseAutoRTFMCompiler;
 			Result.FPassPlugins = Rules.FPassPlugins;
+			Result.OptimizationLevel = Rules.OptimizationLevel;
+			Result.FPSemantics = Rules.FPSemantics;
 
 			if (Target.Platform.IsInGroup(UnrealPlatformGroup.Apple))
 			{
@@ -2641,43 +2662,6 @@ namespace UnrealBuildTool
 				Result.bAutoRTFMRedirectionInstrumentation = Target.bAutoRTFMRedirectionInstrumentation;
 			}
 
-			if (Result.OptimizationLevel != Rules.OptimizationLevel)
-			{
-				Logger.LogInformation("Module {Name} - Optimization level changed for module due to override. Old: {OldOptimizationLevel} New: {NewOptimizationLevel}", Name, Result.OptimizationLevel, Rules.OptimizationLevel);
-				if (Rules.PrivatePCHHeaderFile == null)
-				{
-					if (Rules.PCHUsage != ModuleRules.PCHUsageMode.NoPCHs)
-					{
-						Logger.LogInformation("  Overriding OptimizationLevel requires a private PCH. Disabling PCH usage for {Name}", Name);
-						Rules.PCHUsage = ModuleRules.PCHUsageMode.NoPCHs;
-					}
-				}
-				else if (Rules.PCHUsage == ModuleRules.PCHUsageMode.UseSharedPCHs)
-				{
-					Logger.LogInformation("  Overriding OptimizationLevel requires a private PCH. A private PCH exists but UseSharedPCHs was specified. Overriding to NoSharedPCHs for {Name}", Name);
-					Rules.PCHUsage = ModuleRules.PCHUsageMode.NoSharedPCHs;
-				}
-				Result.OptimizationLevel = Rules.OptimizationLevel;
-			}
-
-			if (Result.FPSemantics != Rules.FPSemantics)
-			{
-				if (Rules.PrivatePCHHeaderFile == null)
-				{
-					if (Rules.PCHUsage != ModuleRules.PCHUsageMode.NoPCHs)
-					{
-						Logger.LogInformation("  Overriding FPSemantics requires a private PCH. Disabling PCH usage for {Name}", Name);
-						Rules.PCHUsage = ModuleRules.PCHUsageMode.NoPCHs;
-					}
-				}
-				else if (Rules.PCHUsage == ModuleRules.PCHUsageMode.UseSharedPCHs)
-				{
-					Logger.LogInformation("  Overriding FPSemantics requires a private PCH. A private PCH exists but UseSharedPCHs was specified. Overriding to NoSharedPCHs for {Name}", Name);
-					Rules.PCHUsage = ModuleRules.PCHUsageMode.NoSharedPCHs;
-				}
-				Result.FPSemantics = Rules.FPSemantics;
-			}
-
 			// If the module overrides the C++ language version, override it on the compile environment
 			if (Rules.CppStandard != null)
 			{
