Improve error message about incorrectly named files for hot reloading.

#rnx

- Add the name of the module failing the check to the error message, previously there was nothing indicating which module was causing the error.

#Synced-CL 51587407

diff --git a/Engine/Source/Programs/UnrealBuildTool/System/HotReload.cs b/Engine/Source/Programs/UnrealBuildTool/System/HotReload.cs
--- a/Engine/Source/Programs/UnrealBuildTool/System/HotReload.cs
+++ b/Engine/Source/Programs/UnrealBuildTool/System/HotReload.cs
@@ -820,7 +820,7 @@ namespace UnrealBuildTool
 			int HyphenIdx = FileName.IndexOf('-', StringComparison.Ordinal);
 			if (HyphenIdx == -1)
 			{
-				throw new BuildException("Hot-reloadable files are expected to contain a hyphen, eg. UnrealEditor-Core");
+				throw new BuildException("Hot-reloadable files are expected to contain a hyphen, eg. UnrealEditor-Core ({0})", FileName);
 			}
 
 			int NameEndIdx = HyphenIdx + 1;
