UnrealBuildTool: Remove hack to prune specific targets for SingleFile compiles, change NMake command-lines for specific file compiles to exclude extra targets.

#Synced-CL 51819546

diff --git a/Engine/Source/Programs/UnrealBuildTool/Modes/BuildMode.cs b/Engine/Source/Programs/UnrealBuildTool/Modes/BuildMode.cs
--- a/Engine/Source/Programs/UnrealBuildTool/Modes/BuildMode.cs
+++ b/Engine/Source/Programs/UnrealBuildTool/Modes/BuildMode.cs
@@ -191,12 +191,6 @@ namespace UnrealBuildTool
 					TargetDescriptors.ForEach(x => x.AdditionalArguments.ApplyTo(BuildConfiguration));
 				}
 
-				// Hack for specific files compile; don't build the ShaderCompileWorker target that's added to the command line for generated project files
-				if (TargetDescriptors.Count > 1)
-				{
-					TargetDescriptors.RemoveAll(x => (x.Name == "ShaderCompileWorker" || x.Name == "LiveCodingConsole" || x.Name == "InterchangeWorker") && x.SpecificFilesToCompile.Count > 0);
-				}
-
 				// Clean any target that wanted to be cleaned before being rebuilt
 				if (TargetDescriptors.Any(D => D.bRebuild))
 				{
diff --git a/Engine/Source/Programs/UnrealBuildTool/ProjectFiles/VisualStudio/VCProject.cs b/Engine/Source/Programs/UnrealBuildTool/ProjectFiles/VisualStudio/VCProject.cs
--- a/Engine/Source/Programs/UnrealBuildTool/ProjectFiles/VisualStudio/VCProject.cs
+++ b/Engine/Source/Programs/UnrealBuildTool/ProjectFiles/VisualStudio/VCProject.cs
@@ -2235,7 +2235,7 @@ namespace UnrealBuildTool
 				CleanScript = FileReference.Combine(BatchFilesDirectory, "Clean.bat");
 			}
 
-			public string GetBuildArguments()
+			public string GetBuildArguments(bool bAllowExtraTargets = true)
 			{
 				TargetRules TargetRulesObject = ProjectTarget.TargetRules!;
 				string TargetName = ProjectTarget.TargetFilePath.GetFileNameWithoutAnyExtensions();
@@ -2249,7 +2249,7 @@ namespace UnrealBuildTool
 				}
 
 				List<string> ExtraTargets = new List<string>();
-				if (!bUsePrecompiled)
+				if (!bUsePrecompiled && bAllowExtraTargets)
 				{
 					if (TargetRulesObject.Type == TargetType.Editor && bEditorDependsOnShaderCompileWorker && !Unreal.IsEngineInstalled())
 					{
@@ -2329,7 +2329,6 @@ namespace UnrealBuildTool
 			return Builder;
 		}
 
-		// Anonymous function that writes project configuration data
 		private void WriteConfiguration(string ProjectName, ProjectConfigAndTargetCombination Combination, StringBuilder VCProjectFileContent, PlatformProjectGeneratorCollection PlatformProjectGenerators, StringBuilder? VCUserFileContent, VisualStudioUserFileSettings? VCUserFileSettings)
 		{
 			UnrealTargetConfiguration Configuration = Combination.Configuration;
@@ -2557,9 +2556,10 @@ namespace UnrealBuildTool
 						VCProjectFileContent.Append(ProjGenerator.GetVisualStudioLayoutDirSection(VSSettings, ConditionString, Combination.ProjectTarget.TargetRules!.Type, Combination.ProjectTarget.TargetFilePath, ProjectFilePath, NMakePath));
 					}
 
+					string BuildArgumentsForFileCompiles = Builder.GetBuildArguments(bAllowExtraTargets: false);
 					VCProjectFileContent.AppendLine("  <ItemDefinitionGroup {0}>", ConditionString);
 					VCProjectFileContent.AppendLine("    <NMakeCompile>");
-					VCProjectFileContent.AppendLine("      <NMakeCompileFileCommandLine>$(BuildBatchScript) {0} -WorkingDir=$(MSBuildProjectDirectory) -Files=$(SelectedFiles)</NMakeCompileFileCommandLine>", BuildArguments);
+					VCProjectFileContent.AppendLine($"      <NMakeCompileFileCommandLine>$(BuildBatchScript) {BuildArgumentsForFileCompiles} -WorkingDir=$(MSBuildProjectDirectory) -Files=$(SelectedFiles)</NMakeCompileFileCommandLine>");
 					VCProjectFileContent.AppendLine("    </NMakeCompile>");
 					if (TargetRulesObject.bIsBuildingConsoleApplication)
 					{
