Add all additional libraries to RPATH, not just third party.

- This solves an issue when a plugin that depends on the main game module doesn't have it in RPATH (since it isn't under a Plugin or a ThirdParty directory) and fails to load it.

#Synced-CL 51825589

diff --git a/Engine/Source/Programs/UnrealBuildTool/Platform/Linux/LinuxToolChain.cs b/Engine/Source/Programs/UnrealBuildTool/Platform/Linux/LinuxToolChain.cs
--- a/Engine/Source/Programs/UnrealBuildTool/Platform/Linux/LinuxToolChain.cs
+++ b/Engine/Source/Programs/UnrealBuildTool/Platform/Linux/LinuxToolChain.cs
@@ -1087,8 +1087,8 @@ namespace UnrealBuildTool
 
 				if ((
 						AdditionalLibrary.Contains("Plugins", StringComparison.Ordinal) ||
-						AdditionalLibrary.Contains("Binaries/ThirdParty", StringComparison.Ordinal) ||
-						AdditionalLibrary.Contains("Binaries\\ThirdParty", StringComparison.Ordinal)
+						AdditionalLibrary.Contains("Binaries/", StringComparison.Ordinal) ||
+						AdditionalLibrary.Contains("Binaries\\", StringComparison.Ordinal)
 					) &&
 					Path.GetDirectoryName(AdditionalLibrary) != Path.GetDirectoryName(RealOutputFile.AbsolutePath))
 				{
