Fix LLT generated metadata (onlineservicesnull manually introduced)
Adjust path of onlineservicesnull to internal folder

#rnx

#Synced-CL 0

diff --git a/Engine/Build/LowLevelTests.xml b/Engine/Build/LowLevelTests.xml
--- a/Engine/Build/LowLevelTests.xml
+++ b/Engine/Build/LowLevelTests.xml
@@ -264,7 +264,7 @@
 			</Node>
 			<Label Category="Tests" Name="$(ShortName) $(Platform)" Requires="$(ShortName) Tests $(Platform)"/>
 			<Trace Message="$(ShortName) runs because $(TestRuns)" If="$(TestRuns) and $(DryRun)"/>
-			<Do If="$(TestRuns)">
+			<Do If="$(TestRuns) and '$(Deactivated)' == ''">
 				<Property Name="TestNodes" Value="$(TestNodes);$(ShortName) Tests $(Platform)"/>
 			</Do>
 			<Do If="'$(CheckedRunAfterStep)' == 'True' AND '$(HasAfterSteps)' == 'True'">
diff --git a/Engine/Build/LowLevelTests/OnlineServicesNull.xml b/Engine/Build/LowLevelTests/OnlineServicesNull.xml
deleted file mode 100644
diff --git a/Engine/Programs/Online/OnlineServicesNullTests/Config/DefaultEngine.ini b/Engine/Programs/Online/OnlineServicesNullTests/Config/DefaultEngine.ini
deleted file mode 100644
diff --git a/Engine/Programs/Online/OnlineServicesNullTests/Config/DefaultGame.ini b/Engine/Programs/Online/OnlineServicesNullTests/Config/DefaultGame.ini
deleted file mode 100644
diff --git a/Engine/Programs/Online/OnlineServicesNullTests/OnlineServicesNullTests.uproject b/Engine/Programs/Online/OnlineServicesNullTests/OnlineServicesNullTests.uproject
deleted file mode 100644
diff --git a/Engine/Source/Programs/Online/OnlineServicesNullTests/OnlineServicesNullTests.Build.cs b/Engine/Source/Programs/Online/OnlineServicesNullTests/OnlineServicesNullTests.Build.cs
deleted file mode 100644
diff --git a/Engine/Source/Programs/Online/OnlineServicesNullTests/OnlineServicesNullTests.Target.cs b/Engine/Source/Programs/Online/OnlineServicesNullTests/OnlineServicesNullTests.Target.cs
deleted file mode 100644
diff --git a/Engine/Source/Programs/Online/OnlineServicesNullTests/Private/ExampleTest.cpp b/Engine/Source/Programs/Online/OnlineServicesNullTests/Private/ExampleTest.cpp
deleted file mode 100644
diff --git a/Engine/Source/Programs/Online/OnlineServicesNullTests/Private/OnlineLLTModule.cpp b/Engine/Source/Programs/Online/OnlineServicesNullTests/Private/OnlineLLTModule.cpp
deleted file mode 100644
diff --git a/Engine/Source/Programs/UnrealBuildTool/Configuration/Rules/TestModuleRules.cs b/Engine/Source/Programs/UnrealBuildTool/Configuration/Rules/TestModuleRules.cs
--- a/Engine/Source/Programs/UnrealBuildTool/Configuration/Rules/TestModuleRules.cs
+++ b/Engine/Source/Programs/UnrealBuildTool/Configuration/Rules/TestModuleRules.cs
@@ -290,7 +290,9 @@ namespace UnrealBuildTool
 				}
 				ExpandArguments.Add("TargetName", Convert.ToString(TestTargetName));
 				ExpandArguments.Add("BinaryRelativePath", Convert.ToString(TestBinariesPath));
-				ExpandArguments.Add("ReportType", Convert.ToString(TestMetadata.ReportType));
+				string ReportTypeValue =
+					(ValidPlatform == UnrealTargetPlatform.IOS || ValidPlatform == UnrealTargetPlatform.TVOS || ValidPlatform == UnrealTargetPlatform.Android) ? "console" : TestMetadata.ReportType;
+				ExpandArguments.Add("ReportType", Convert.ToString(ReportTypeValue));
 				if (!String.IsNullOrEmpty(TestMetadata.GauntletArgs))
 				{
 					ExpandArguments.Add("GauntletArgs", Convert.ToString(TestMetadata.InitialExtraArgs) + Convert.ToString(TestMetadata.GauntletArgs));
