[GDK] Fixes bug causing image for Features to not appear in the package.

When generating the layout file, the <Feature> element will only contain the Image attribute when the file exists. The function being called was not searching the project directories, as expected.

With this change, PrepareResourceBinaryPaths will now have ProjectFile set and add the appropriate project-specific search paths.

#Synced-CL 51753475

#Synced-CL 51300879

diff --git a/Engine/Source/Programs/AutomationTool/Microsoft/GDKPackageGenerator.cs b/Engine/Source/Programs/AutomationTool/Microsoft/GDKPackageGenerator.cs
--- a/Engine/Source/Programs/AutomationTool/Microsoft/GDKPackageGenerator.cs
+++ b/Engine/Source/Programs/AutomationTool/Microsoft/GDKPackageGenerator.cs
@@ -838,7 +838,7 @@ namespace AutomationUtils.GDK
 
 								// see if there is an image resource for this feature (it will be added in GDKGameConfigGenerator)
 								string ImageFilename = $"Feature_{Mapping.Id}.png";
-								bool bHasImageFilename = GDKGameConfigGenerator.DoesDefaultResourceBinaryFileExistForPlatform(SC.StageTargetPlatform.PlatformType, Logger, ImageFilename, bAllowEngineFallback: false);
+								bool bHasImageFilename = GDKGameConfigGenerator.DoesDefaultResourceBinaryFileExistForPlatform(SC.RawProjectPath, SC.StageTargetPlatform.PlatformType, Logger, ImageFilename, bAllowEngineFallback: false);
 								Mapping.ImageFileName = bHasImageFilename ? $@"Resources\{ImageFilename}" : null;
 							}
 
diff --git a/Engine/Source/Programs/UnrealBuildTool/Platform/Microsoft/GDKGameConfigGenerator.cs b/Engine/Source/Programs/UnrealBuildTool/Platform/Microsoft/GDKGameConfigGenerator.cs
--- a/Engine/Source/Programs/UnrealBuildTool/Platform/Microsoft/GDKGameConfigGenerator.cs
+++ b/Engine/Source/Programs/UnrealBuildTool/Platform/Microsoft/GDKGameConfigGenerator.cs
@@ -718,10 +718,11 @@ namespace UnrealBuildTool.GDK
 		/// <summary>
 		/// Helper to find whether there is a packaging resource file available
 		/// </summary>
-		public static bool DoesDefaultResourceBinaryFileExistForPlatform(UnrealTargetPlatform Platform, ILogger Logger, string ImageFilename, bool bAllowEngineFallback = true)
+		public static bool DoesDefaultResourceBinaryFileExistForPlatform(FileReference ProjectFile, UnrealTargetPlatform Platform, ILogger Logger, string ImageFilename, bool bAllowEngineFallback = true)
 		{
 			GDKGameConfigGenerator Generator = GDKGameConfigGeneratorFactory.CreateGenerator(Platform,Logger);
 			Generator.AppXResources = new(Logger, Generator.GetMakePriBinaryPath());
+			Generator.ProjectFile = ProjectFile;
 			Generator.PrepareResourceBinaryPaths();
 
 			return Generator.AppXResources!.DoesDefaultResourceBinaryFileExist(ImageFilename, bAllowEngineFallback);
