Disable FMA contraction under Clang on Windows with FPSemanticsMode.Precise

This was causing divergence between binaries that did or didn't have hardware FMA instructions enabled. It matches what UBT does for Clang on other platforms.

#rnx
#Synced-CL 51749326

diff --git a/Engine/Source/Programs/UnrealBuildTool/Platform/Windows/VCToolChain.cs b/Engine/Source/Programs/UnrealBuildTool/Platform/Windows/VCToolChain.cs
--- a/Engine/Source/Programs/UnrealBuildTool/Platform/Windows/VCToolChain.cs
+++ b/Engine/Source/Programs/UnrealBuildTool/Platform/Windows/VCToolChain.cs
@@ -1303,7 +1303,18 @@ namespace UnrealBuildTool
 			{
 				case FPSemanticsMode.Default: // Default is imprecise FP semantics.
 				case FPSemanticsMode.Imprecise: Arguments.Add("/fp:fast"); break;
-				case FPSemanticsMode.Precise: Arguments.Add("/fp:precise"); break;
+				case FPSemanticsMode.Precise:
+					if (Target.WindowsPlatform.Compiler.IsClang())
+					{
+						// Clang defaults to -ffp-contract=on even with /fp:precise, which allows fusing multiplications and additions into FMAs.
+						// Passing this alongside /fp:precise emits a warning, but /fp:precise is already the default so it's safe to omit.
+						Arguments.Add("/clang:-ffp-contract=off");
+					}
+					else
+					{
+						Arguments.Add("/fp:precise");
+					}
+					break;
 				default:
 					throw new BuildException($"Unsupported FP semantics: {FPSemantics}");
 			}
