UnrealBuildTool: Add MSVC toolchain version check for precompiled modules

#rnx
#Synced-CL 51916670

diff --git a/Engine/Source/Programs/UnrealBuildTool/Configuration/UEBuildModuleCPP.cs b/Engine/Source/Programs/UnrealBuildTool/Configuration/UEBuildModuleCPP.cs
--- a/Engine/Source/Programs/UnrealBuildTool/Configuration/UEBuildModuleCPP.cs
+++ b/Engine/Source/Programs/UnrealBuildTool/Configuration/UEBuildModuleCPP.cs
@@ -561,6 +561,26 @@ namespace UnrealBuildTool
 				}
 
 				PrecompiledManifest Manifest = PrecompiledManifest.Read(PrecompiledManifestLocation);
+
+				// Validate the precompiled objects were not built with a newer MSVC toolchain than the one currently selected.
+				// A newer compiler can introduce intrinsics (e.g. __std_search_1) that don't exist in older runtimes, causing LNK2019.
+				if (Target.Platform.IsInGroup(UnrealPlatformGroup.Microsoft) && Target.WindowsPlatform.Compiler.IsMSVC()
+					&& Manifest.ToolChainVersion != null && Manifest.ToolChainVersion.Components.Count >= 2
+					&& Target.WindowsPlatform.Environment != null && Target.WindowsPlatform.Environment.ToolChainVersion.Components.Count >= 2)
+				{
+					VersionNumber ManifestFamily = new VersionNumber(Manifest.ToolChainVersion.GetComponent(0), Manifest.ToolChainVersion.GetComponent(1));
+					VersionNumber CurrentFamily  = new VersionNumber(Target.WindowsPlatform.Environment.ToolChainVersion.GetComponent(0), Target.WindowsPlatform.Environment.ToolChainVersion.GetComponent(1));
+
+					if (ManifestFamily > CurrentFamily)
+					{
+						throw new BuildLogEventException(
+							"Module '{Name}' was precompiled with MSVC toolchain {ManifestVersion} but the currently selected toolchain is {CurrentVersion}. " +
+							"The precompiled objects may reference intrinsics that do not exist in older runtimes, which will cause linker errors. " +
+							"Please use a newer toolchain, or recompile the engine from source.",
+							Name, Manifest.ToolChainVersion, Target.WindowsPlatform.Environment.ToolChainVersion);
+					}
+				}
+
 				foreach (FileReference OutputFile in Manifest.OutputFiles)
 				{
 					FileItem ObjectFile = FileItem.GetItemByFileReference(OutputFile);
@@ -1176,6 +1196,13 @@ namespace UnrealBuildTool
 
 				PrecompiledManifest Manifest = new PrecompiledManifest();
 				Manifest.OutputFiles.AddRange(LinkInputFiles.Select(x => x.Location));
+
+				// Store MSVC toolchain version to detect ABI mismatches
+				if (Target.Platform.IsInGroup(UnrealPlatformGroup.Microsoft) && Target.WindowsPlatform.Compiler.IsMSVC())
+				{
+					Manifest.ToolChainVersion = Target.WindowsPlatform.Environment?.ToolChainVersion;
+				}
+
 				Manifest.WriteIfModified(PrecompiledManifestLocation);
 			}
 
diff --git a/Engine/Source/Programs/UnrealBuildTool/System/PrecompiledManifest.cs b/Engine/Source/Programs/UnrealBuildTool/System/PrecompiledManifest.cs
--- a/Engine/Source/Programs/UnrealBuildTool/System/PrecompiledManifest.cs
+++ b/Engine/Source/Programs/UnrealBuildTool/System/PrecompiledManifest.cs
@@ -17,6 +17,12 @@ namespace UnrealBuildTool
 		/// </summary>
 		public List<FileReference> OutputFiles = new List<FileReference>();
 
+		/// <summary>
+		/// The MSVC toolchain version used to compile the precompiled objects.
+		/// Null for non-MSVC or manifests produced before this field was added.
+		/// </summary>
+		public VersionNumber? ToolChainVersion = null;
+
 		/// <summary>
 		/// Default constructor
 		/// </summary>
@@ -41,6 +47,11 @@ namespace UnrealBuildTool
 				Manifest.OutputFiles.Add(FileReference.Combine(BaseDir, OutputFile));
 			}
 
+			if (RawObject.TryGetStringField("ToolChainVersion", out string? ToolChainVersionStr))
+			{
+				_ = VersionNumber.TryParse(ToolChainVersionStr, out Manifest.ToolChainVersion);
+			}
+
 			return Manifest;
 		}
 
@@ -90,6 +101,11 @@ namespace UnrealBuildTool
 				}
 				Writer.WriteStringArrayField("OutputFiles", OutputFileStrings);
 
+				if (ToolChainVersion != null)
+				{
+					Writer.WriteValue("ToolChainVersion", ToolChainVersion.ToString());
+				}
+
 				Writer.WriteObjectEnd();
 			}
 
