UBT: Fix thread-safety of MacToolChain.DebugInfoFiles causing intermittent dSYM failures

DebugInfoFiles is a List<FileItem> mutated concurrently from Parallel2.ForEach in UEBuildTarget.Build.

Change to ConcurrentBag to prevent lost entries.

#Synced-CL 51629120

diff --git a/Engine/Source/Programs/UnrealBuildTool/Platform/Mac/MacToolChain.cs b/Engine/Source/Programs/UnrealBuildTool/Platform/Mac/MacToolChain.cs
--- a/Engine/Source/Programs/UnrealBuildTool/Platform/Mac/MacToolChain.cs
+++ b/Engine/Source/Programs/UnrealBuildTool/Platform/Mac/MacToolChain.cs
@@ -1,6 +1,7 @@
 // Copyright Epic Games, Inc. All Rights Reserved.
 
 using System;
+using System.Collections.Concurrent;
 using System.Collections.Generic;
 using System.IO;
 using System.Linq;
@@ -1013,7 +1014,7 @@ namespace UnrealBuildTool
 			}
 		}
 
-		private readonly List<FileItem> DebugInfoFiles = new List<FileItem>();
+		private readonly ConcurrentBag<FileItem> DebugInfoFiles = new ConcurrentBag<FileItem>();
 
 		public override void FinalizeOutput(ReadOnlyTargetRules Target, TargetMakefileBuilder MakefileBuilder, DirectoryReference ExeDir)
 		{
