UnrealBuildTool: Add logging to QueryTargetsMode.

We want this in order to try and track down an issue where default but non-standard editor target names can be missing from this file, causing Rider to incorrectly guess the target name and later fail to build.

#Synced-CL 52057491

diff --git a/Engine/Source/Programs/UnrealBuildTool/Modes/QueryTargetsMode.cs b/Engine/Source/Programs/UnrealBuildTool/Modes/QueryTargetsMode.cs
--- a/Engine/Source/Programs/UnrealBuildTool/Modes/QueryTargetsMode.cs
+++ b/Engine/Source/Programs/UnrealBuildTool/Modes/QueryTargetsMode.cs
@@ -65,6 +65,17 @@ namespace UnrealBuildTool
 			XmlConfig.ApplyTo(buildConfiguration);
 			arguments.ApplyTo(buildConfiguration);
 
+			// Create the log file, and flush the startup listener to it
+			if (!arguments.HasOption("-NoLog") && !Log.HasFileWriter())
+			{
+				FileReference LogFile = FileReference.Combine(Unreal.EngineProgramSavedDirectory, "UnrealBuildTool", "Log_QueryTargets.txt");
+				Log.AddFileWriter("DefaultLogTraceListener", LogFile);
+			}
+			else
+			{
+				Log.RemoveStartupTraceListener();
+			}
+
 			// Ensure the path to the output file is valid
 			OutputFile ??= GetDefaultOutputFile(ProjectFile);
 
