[SN-DBS] Use per-job tool template directories to avoid ambiguity.

When specifying tool templates using 'dbsbuild -templates', there may only be a single template for a tool of a given name. This usually isn't a problem when building a single platform, since there is only one of a given tool such as a compiler. However, UBT supports building multiple platforms at once. If multiple platforms all use a tool named clang.exe for example, there is no way to disambiguate which tool template to use.

SN-DBS 2.6.10.1020 added support for specifying the template directory for each job, rather than globally. UBT now writes the template files to per-platform directories, and each job specifies the directory corresponding to its platform.

To ease the cost of engine upgrades, we will continue to support older versions of SN-DBS. UBT will detect when there are two different templates for the same tool and issue a build failure. In that situation, the options are to upgrade SN-DBS or compile one platform at a time.

Existing templates were moved out of a global list in SNDBS.cs and added to the appropriate UEBuild*.cs file for the platform it belongs to. The templates were not changed, only moved.

#Synced-CL 51930710

diff --git a/Engine/Source/Programs/UnrealBuildTool/Configuration/UEBuildPlatform.cs b/Engine/Source/Programs/UnrealBuildTool/Configuration/UEBuildPlatform.cs
--- a/Engine/Source/Programs/UnrealBuildTool/Configuration/UEBuildPlatform.cs
+++ b/Engine/Source/Programs/UnrealBuildTool/Configuration/UEBuildPlatform.cs
@@ -786,6 +786,15 @@ namespace UnrealBuildTool
 			return false;
 		}
 
+		/// <summary>
+		/// Provides a set of SN-DBS tool templates for builds on this platform. The dictionary key is the name of the tool, such as "clang.exe", and the value is the template contents.
+		/// </summary>
+		/// <returns>Dictionary of tool templates or null if there are none.</returns>
+		public virtual IReadOnlyDictionary<string, string>? GetSNDBSToolTemplates()
+		{
+			return null;
+		}
+
 		/// <summary>
 		/// Set all the platform-specific defaults for a new target
 		/// </summary>
diff --git a/Engine/Source/Programs/UnrealBuildTool/Executors/SNDBS.cs b/Engine/Source/Programs/UnrealBuildTool/Executors/SNDBS.cs
--- a/Engine/Source/Programs/UnrealBuildTool/Executors/SNDBS.cs
+++ b/Engine/Source/Programs/UnrealBuildTool/Executors/SNDBS.cs
@@ -10,12 +10,14 @@ using System.Linq;
 using System.Net;
 using System.Net.NetworkInformation;
 using System.Runtime.InteropServices;
+using System.Runtime.Versioning;
 using System.ServiceProcess;
 using System.Text.Json;
 using System.Text.RegularExpressions;
 using System.Threading.Tasks;
 using EpicGames.Core;
 using Microsoft.Extensions.Logging;
+using Microsoft.Win32;
 using UnrealBuildBase;
 
 namespace UnrealBuildTool
@@ -50,7 +52,7 @@ namespace UnrealBuildTool
 		private static readonly FileReference IncludeRewriteRulesFile = FileReference.Combine(IntermediateDir, "include-rewrite-rules.ini");
 		private static readonly FileReference ScriptFile = FileReference.Combine(IntermediateDir, "sndbs.json");
 
-		private readonly Dictionary<string, string> ActiveTemplates = BuiltInTemplates.ToDictionary(p => p.Key, p => p.Value);
+		private readonly Dictionary<string, string> GlobalTemplates = new Dictionary<string, string>();
 
 #pragma warning disable IDE0044 // Make field readonly - these private static fields are set by XML configuration.
 		/// <summary>
@@ -91,7 +93,7 @@ namespace UnrealBuildTool
 
 		public SNDBS AddTemplate(string ExeName, string TemplateContents)
 		{
-			ActiveTemplates.Add(ExeName, TemplateContents);
+			GlobalTemplates.Add(ExeName, TemplateContents);
 			return this;
 		}
 
@@ -236,6 +238,12 @@ namespace UnrealBuildTool
 			return Task.FromResult(ExecuteActions(ActionsToExecute, Logger));
 		}
 
+		[SupportedOSPlatform("windows")]
+		private static string? GetInstalledSNDBSVersion()
+		{
+			return (string?)Registry.GetValue(@"HKEY_LOCAL_MACHINE\SOFTWARE\Wow6432Node\SCE\Common\SN-DBS", "Version", null);
+		}
+
 		bool ExecuteActions(IEnumerable<LinkedAction> Actions, ILogger Logger)
 		{
 			if (!Actions.Any())
@@ -243,6 +251,34 @@ namespace UnrealBuildTool
 				return true;
 			}
 
+			if (!OperatingSystem.IsWindows())
+			{
+				throw new BuildException("SN-DBS is only supported on Windows.");
+			}
+
+			// UE requires SN-DBS 2.6.10.1020 or later to avoid issues that may arise when compiling multiple platforms
+			// at once. To ease the pain of engine upgrades, we are temporarily maintaining support for older versions.
+			// When using older versions, a build failure will occur if the scenario requiring the newer version is
+			// encountered.
+			bool bUsePerJobToolTemplates = false;
+
+			string? SNDBSVersionString = GetInstalledSNDBSVersion();
+
+			if (SNDBSVersionString != null)
+			{
+				Version InstalledVersion = new Version(SNDBSVersionString);
+				Version RequiredVersion = new Version("2.6.10.1020");
+
+				if (InstalledVersion >= RequiredVersion)
+				{
+					bUsePerJobToolTemplates = true;
+				}
+			}
+			else
+			{
+				throw new BuildException("Failed to read SN-DBS version from registry.");
+			}
+
 			// Clean the intermediate directory in case there are any leftovers from previous builds
 			if (DirectoryReference.Exists(IntermediateDir))
 			{
@@ -274,6 +310,19 @@ namespace UnrealBuildTool
 						["run_locally"] = !(a.bCanExecuteRemotely && a.bCanExecuteRemotelyWithSNDBS)
 					};
 
+					// Associate this job with the platform-specific template directory
+					if (bUsePerJobToolTemplates)
+					{
+						if (a.Target != null)
+						{
+							// JSON format needs this to be an array, even though we only use one directory
+							DirectoryReference[] ToolTemplateDirectories = new DirectoryReference[1];
+							ToolTemplateDirectories[0] = DirectoryReference.Combine(IntermediateDir, a.Target.Platform.ToString());
+
+							Job["tool_template_directories"] = ToolTemplateDirectories;
+						}
+					}
+
 					if (a.PrerequisiteItems.Any())
 					{
 						Job["explicit_input_files"] = a.PrerequisiteItems
@@ -322,12 +371,14 @@ namespace UnrealBuildTool
 				}).ToArray()
 			}, JsonOption));
 
-			PrepareToolTemplates();
+			PrepareToolTemplates(bUsePerJobToolTemplates);
 			bool bHasRewrites = GenerateSNDBSIncludeRewriteRules();
 
 			IEnumerable<string> ConfigList = TargetDescriptors.Select(Descriptor => $"{Descriptor.Name}|{Descriptor.Platform}|{Descriptor.Configuration}");
 			string ConfigDescription = String.Join(",", ConfigList);
 
+			// The per-job "tool_template_directories" take precedent over the value passed to -templates. We use the
+			// -templates for GlobalTemplates.
 			ProcessStartInfo StartInfo = new ProcessStartInfo(
 				SNDBSBuildExe,
 				$"-q -p \"{ConfigDescription}\" -s \"{ScriptFile}\" -templates \"{IntermediateDir}\""
@@ -418,11 +469,120 @@ namespace UnrealBuildTool
 			}
 		}
 
-		private void PrepareToolTemplates()
+		private void PrepareToolTemplates(bool bUsePerJobToolTemplates)
+		{
+			// Get distinct platforms being used in this build.
+			List<UnrealTargetPlatform> TargetPlatforms = TargetDescriptors
+				.Select(TargetDescriptor => TargetDescriptor.Platform)
+				.Distinct()
+				.ToList();
+
+			if (bUsePerJobToolTemplates)
+			{
+				// Write templates to a per-platform directory. Each job will refer to the directory of the associated
+				// platform. This allows different platforms to have different templates for a tool of the same name.
+				foreach (UnrealTargetPlatform TargetPlatform in TargetPlatforms)
+				{
+					UEBuildPlatform? BuildPlatform = UEBuildPlatform.GetBuildPlatform(TargetPlatform);
+
+					if (BuildPlatform == null)
+					{
+						continue;
+					}
+
+					// Create the directory even if there are no templates to avoid an error when the JSON references a
+					// non-existant directory.
+					DirectoryReference PlatformTemplateDir = DirectoryReference.Combine(IntermediateDir, BuildPlatform.Platform.ToString());
+					DirectoryReference.CreateDirectory(PlatformTemplateDir);
+
+					// Not all platforms specify their own templates
+					IReadOnlyDictionary<string, string>? PlatformTemplates = BuildPlatform.GetSNDBSToolTemplates();
+
+					if (PlatformTemplates == null)
+					{
+						continue;
+					}
+
+					WriteToolTemplates(PlatformTemplates, PlatformTemplateDir);
+				}
+
+				// Global templates are always written to the directory passed to 'dbsbuild -templates'. This gets
+				// searched last, ensuring the per-job directories always override it.
+				WriteToolTemplates(GlobalTemplates, IntermediateDir);
+			}
+			else
+			{
+				// Older versions of SN-DBS have no way to disambiguate templates for a tool of the same name. This is
+				// only a problem when building multiple platforms at the same time that are using a different template
+				// for a tool of the same name. To provide a smoother upgrade path, we support using older versions of
+				// SN-DBS provided this scenario is not encountered.
+				//
+				// For example, if platforms A and B both have a tool named "compiler.exe" requiring a different
+				// template, that will produce a build error indicating an upgrade is required. However, compiling A by
+				// itself would work fine as would any other combination of platforms that did not include both A and B.
+				//
+				// If encountering this error, the recommended resolution is to upgrade SN-DBS to 2.6.10.1020 or later.
+
+				Dictionary<string, string> MergedTemplates = new Dictionary<string, string>();
+
+				foreach (UnrealTargetPlatform TargetPlatform in TargetPlatforms)
+				{
+					UEBuildPlatform? BuildPlatform = UEBuildPlatform.GetBuildPlatform(TargetPlatform);
+
+					if (BuildPlatform == null)
+					{
+						continue;
+					}
+
+					IReadOnlyDictionary<string, string>? PlatformTemplates = BuildPlatform.GetSNDBSToolTemplates();
+
+					if (PlatformTemplates == null)
+					{
+						continue;
+					}
+
+					// Merge all platform templates into a single set. Two platforms may specify an identical template
+					// for the same tool.
+					foreach (KeyValuePair<string, string> Template in PlatformTemplates)
+					{
+						string? PrevValue;
+						if (MergedTemplates.TryGetValue(Template.Key, out PrevValue))
+						{
+							if (PrevValue != null && !PrevValue.Equals(Template.Value, StringComparison.Ordinal))
+							{
+								throw new BuildException($"Conflicting SN-DBS tool template. {BuildPlatform} specifies a template for \"{Template.Key}\", but another platform is using a different template for the same tool. Upgrade SN-DBS to version 2.6.10.1020 or later to resolve this problem.");
+							}
+						}
+						else
+						{
+							MergedTemplates.Add(Template.Key, Template.Value);
+						}
+					}
+				}
+
+				// Now merge in any manually added templates
+				foreach (KeyValuePair<string, string> Template in GlobalTemplates)
+				{
+					string? PrevValue;
+					if (MergedTemplates.TryGetValue(Template.Key, out PrevValue))
+					{
+						throw new BuildException($"Conflicting SN-DBS tool templates. Template for \"{Template.Key}\" added via AddTemplate() conflicts with a platform-provided template.");
+					}
+					else
+					{
+						MergedTemplates.Add(Template.Key, Template.Value);
+					}
+				}
+
+				WriteToolTemplates(MergedTemplates, IntermediateDir);
+			}
+		}
+
+		private static void WriteToolTemplates(IReadOnlyDictionary<string, string> Templates, DirectoryReference Directory)
 		{
-			foreach (KeyValuePair<string, string> Template in ActiveTemplates)
+			foreach (KeyValuePair<string, string> Template in Templates)
 			{
-				FileReference TemplateFile = FileReference.Combine(IntermediateDir, $"{Template.Key}.sn-dbs-tool.ini");
+				FileReference TemplateFile = FileReference.Combine(Directory, $"{Template.Key}.sn-dbs-tool.ini");
 				string TemplateText = Template.Value;
 
 				foreach (Nullable<DictionaryEntry> Variable in Environment.GetEnvironmentVariables(EnvironmentVariableTarget.Process))
@@ -474,157 +634,5 @@ namespace UnrealBuildTool
 				return false;
 			}
 		}
-
-		/// <summary>
-		/// SN-DBS templates that are automatically included in the build.
-		/// </summary>
-		private static readonly IReadOnlyDictionary<string, string> BuiltInTemplates = new Dictionary<string, string>()
-		{
-			["cl-filter.exe"] = @"
-[tool]
-family=msvc
-vc_major_version=14
-use_surrogate=true
-force_synchronous_pdb_writes=true
-error_report_mode=prompt
-response_file_content_pattern=\s--\s"".*?cl\.exe""\s(.*)
-
-[group]
-server={VC_TOOLCHAIN_DIR}\mspdbsrv.exe
-
-[files]
-main=cl-filter.exe
-file01={VC_TOOLCHAIN_DIR}\c1.dll
-file02={VC_TOOLCHAIN_DIR}\c1xx.dll
-file03={VC_TOOLCHAIN_DIR}\c2.dll
-file04={VC_TOOLCHAIN_DIR}\mspdb140.dll
-file05={VC_TOOLCHAIN_DIR}\mspdbcore.dll
-file06={VC_TOOLCHAIN_DIR}\mspdbsrv.exe
-file07={VC_TOOLCHAIN_DIR}\mspft140.dll
-file08={VC_TOOLCHAIN_DIR}\vcmeta.dll
-file09={VC_TOOLCHAIN_DIR}\*\clui.dll
-file10={VC_TOOLCHAIN_DIR}\*\mspft140ui.dll
-file11={VC_TOOLCHAIN_DIR}\localespc.dll
-file12={VC_TOOLCHAIN_DIR}\cppcorecheck.dll
-file13={VC_TOOLCHAIN_DIR}\experimentalcppcorecheck.dll
-file14={VC_TOOLCHAIN_DIR}\espxengine.dll
-file15={VC_TOOLCHAIN_DIR}\c1.exe
-
-[output-file-patterns]
-outputfile01=\s*""([^ "",]+\.cpp\.txt)\""
-
-[output-file-rules]
-rule01=*\sqmcpp*.log|discard=true
-rule02=*\vctoolstelemetry*.dat|discard=true
-rule03=*\Microsoft\Windows\Temporary Internet Files\*|discard=true
-rule04=*\Microsoft\Windows\INetCache\*|discard=true
-
-[input-file-rules]
-rule01=*\sqmcpp*.log|ignore_transient_errors=true;ignore_unexpected_input=true
-rule02=*\vctoolstelemetry*.dat|ignore_transient_errors=true;ignore_unexpected_input=true
-rule03=*\Microsoft\Windows\Temporary Internet Files\*|ignore_transient_errors=true;ignore_unexpected_input=true
-rule04=*\Microsoft\Windows\INetCache\*|ignore_transient_errors=true;ignore_unexpected_input=true
-
-[system-file-filters]
-filter01=msvcr*.dll
-filter02=msvcp*.dll
-filter03=vcruntime140*.dll
-filter04=appcrt140*.dll
-filter05=desktopcrt140*.dll
-filter06=concrt140*.dll",
-			["cl.exe"] = @"
-[tool]
-family=msvc
-vc_major_version=14
-use_surrogate=true
-force_synchronous_pdb_writes=true
-error_report_mode=prompt
-
-[group]
-server={VC_TOOLCHAIN_DIR}\mspdbsrv.exe
-
-[files]
-main={VC_TOOLCHAIN_DIR}\cl.exe
-file01={VC_TOOLCHAIN_DIR}\c1.dll
-file02={VC_TOOLCHAIN_DIR}\c1xx.dll
-file03={VC_TOOLCHAIN_DIR}\c2.dll
-file04={VC_TOOLCHAIN_DIR}\mspdb140.dll
-file05={VC_TOOLCHAIN_DIR}\mspdbcore.dll
-file06={VC_TOOLCHAIN_DIR}\mspdbsrv.exe
-file07={VC_TOOLCHAIN_DIR}\mspft140.dll
-file08={VC_TOOLCHAIN_DIR}\vcmeta.dll
-file09={VC_TOOLCHAIN_DIR}\*\clui.dll
-file10={VC_TOOLCHAIN_DIR}\*\mspft140ui.dll
-file11={VC_TOOLCHAIN_DIR}\localespc.dll
-file12={VC_TOOLCHAIN_DIR}\cppcorecheck.dll
-file13={VC_TOOLCHAIN_DIR}\experimentalcppcorecheck.dll
-file14={VC_TOOLCHAIN_DIR}\espxengine.dll
-
-[output-file-patterns]
-outputfile01=\s*""([^ "",]+\.cpp\.txt\.json)\""
-
-[output-file-rules]
-rule01=*\sqmcpp*.log|discard=true
-rule02=*\vctoolstelemetry*.dat|discard=true
-rule03=*\Microsoft\Windows\Temporary Internet Files\*|discard=true
-rule04=*\Microsoft\Windows\INetCache\*|discard=true
-
-[input-file-rules]
-rule01=*\sqmcpp*.log|ignore_transient_errors=true;ignore_unexpected_input=true
-rule02=*\vctoolstelemetry*.dat|ignore_transient_errors=true;ignore_unexpected_input=true
-rule03=*\Microsoft\Windows\Temporary Internet Files\*|ignore_transient_errors=true;ignore_unexpected_input=true
-rule04=*\Microsoft\Windows\INetCache\*|ignore_transient_errors=true;ignore_unexpected_input=true
-
-[system-file-filters]
-filter01=msvcr*.dll
-filter02=msvcp*.dll
-filter03=vcruntime140*.dll
-filter04=appcrt140*.dll
-filter05=desktopcrt140*.dll
-filter06=concrt140*.dll",
-			["mspdbsrv.exe"] = @"
-[tool]
-use_cache=no
-
-[files]
-main={VC_TOOLCHAIN_DIR}\mspdbsrv.exe
-
-[openmp]
-omp=true
-omp_min_threads=1",
-			["clang++.exe"] = @"
-[tool]
-family=clang
-extensions=.c;.cc;.cpp;.cxx;.c++;.h;.hpp;.s;.asm
-adjacent_metadata_extensions=.pch;.gch;.pth;.o;.obj
-include_path01=..\include
-include_path02=..\include\c++\v1
-include_path03=..\lib\clang\*\include
-include_path04=%CPATH%;%C_INCLUDE_PATH%;%CPLUS_INCLUDE_PATH%
-
-[files]
-main=clang++.exe
-
-[include-path-patterns]
-include01=-(?:I|F|isystem|idirafter)[ \t]*(\""[^\""]+\""|[^ ]+)
-
-[additional-include-file-patterns]
-includefile01=--?include[ \t]*(\""[^\""]+\""|[^ ]+)
-
-[additional-input-file-patterns]
-inputfile01=--?(?:include-pch|fprofile-instr-use=|fprofile-sample-use=|fprofile-use=|fsanitize-ignorelist=)[ \t]*(\""[^\""]+\""|[^ ]+)
-
-[output-file-patterns]
-outputfile01=-o[ \t]*(\""[^\""]+\""|[^ ]+)
-
-[definition-patterns]
-definition01=-D[ \t]*""?([a-zA-Z_0-9]+)(?:[\s""]|$)()
-definition02=-D[ \t]*""?([a-zA-Z_0-9]+)=(?:\\""|<)(.+?)(?:\\""|>)
-definition03=-D[ \t]*""?([a-zA-Z_0-9]+)=((?!(?:\\""|<))[^ ""]*)
-
-[input-scanners]
-scanner01=c .*
-scanner02=assembler .s;.asm",
-		};
 	}
 }
diff --git a/Engine/Source/Programs/UnrealBuildTool/Platform/Windows/UEBuildWindows.cs b/Engine/Source/Programs/UnrealBuildTool/Platform/Windows/UEBuildWindows.cs
--- a/Engine/Source/Programs/UnrealBuildTool/Platform/Windows/UEBuildWindows.cs
+++ b/Engine/Source/Programs/UnrealBuildTool/Platform/Windows/UEBuildWindows.cs
@@ -1812,6 +1812,11 @@ namespace UnrealBuildTool
 			return true;
 		}
 
+		public override IReadOnlyDictionary<string, string>? GetSNDBSToolTemplates()
+		{
+			return SNDBSToolTemplates;
+		}
+
 		/// <summary>
 		/// If this platform can be compiled with FASTBuild
 		/// </summary>
@@ -2249,6 +2254,125 @@ namespace UnrealBuildTool
 			// Both Win32 and Win64 use Windows headers, so we enforce that here.
 			return GetPlatformName();
 		}
+
+		/// <summary>
+		/// SN-DBS tool templates for Windows
+		/// </summary>
+		private static readonly IReadOnlyDictionary<string, string> SNDBSToolTemplates = new Dictionary<string, string>()
+		{
+			["cl-filter.exe"] = @"
+[tool]
+family=msvc
+vc_major_version=14
+use_surrogate=true
+force_synchronous_pdb_writes=true
+error_report_mode=prompt
+response_file_content_pattern=\s--\s"".*?cl\.exe""\s(.*)
+
+[group]
+server={VC_TOOLCHAIN_DIR}\mspdbsrv.exe
+
+[files]
+main=cl-filter.exe
+file01={VC_TOOLCHAIN_DIR}\c1.dll
+file02={VC_TOOLCHAIN_DIR}\c1xx.dll
+file03={VC_TOOLCHAIN_DIR}\c2.dll
+file04={VC_TOOLCHAIN_DIR}\mspdb140.dll
+file05={VC_TOOLCHAIN_DIR}\mspdbcore.dll
+file06={VC_TOOLCHAIN_DIR}\mspdbsrv.exe
+file07={VC_TOOLCHAIN_DIR}\mspft140.dll
+file08={VC_TOOLCHAIN_DIR}\vcmeta.dll
+file09={VC_TOOLCHAIN_DIR}\*\clui.dll
+file10={VC_TOOLCHAIN_DIR}\*\mspft140ui.dll
+file11={VC_TOOLCHAIN_DIR}\localespc.dll
+file12={VC_TOOLCHAIN_DIR}\cppcorecheck.dll
+file13={VC_TOOLCHAIN_DIR}\experimentalcppcorecheck.dll
+file14={VC_TOOLCHAIN_DIR}\espxengine.dll
+file15={VC_TOOLCHAIN_DIR}\c1.exe
+
+[output-file-patterns]
+outputfile01=\s*""([^ "",]+\.cpp\.txt)\""
+
+[output-file-rules]
+rule01=*\sqmcpp*.log|discard=true
+rule02=*\vctoolstelemetry*.dat|discard=true
+rule03=*\Microsoft\Windows\Temporary Internet Files\*|discard=true
+rule04=*\Microsoft\Windows\INetCache\*|discard=true
+
+[input-file-rules]
+rule01=*\sqmcpp*.log|ignore_transient_errors=true;ignore_unexpected_input=true
+rule02=*\vctoolstelemetry*.dat|ignore_transient_errors=true;ignore_unexpected_input=true
+rule03=*\Microsoft\Windows\Temporary Internet Files\*|ignore_transient_errors=true;ignore_unexpected_input=true
+rule04=*\Microsoft\Windows\INetCache\*|ignore_transient_errors=true;ignore_unexpected_input=true
+
+[system-file-filters]
+filter01=msvcr*.dll
+filter02=msvcp*.dll
+filter03=vcruntime140*.dll
+filter04=appcrt140*.dll
+filter05=desktopcrt140*.dll
+filter06=concrt140*.dll",
+			["cl.exe"] = @"
+[tool]
+family=msvc
+vc_major_version=14
+use_surrogate=true
+force_synchronous_pdb_writes=true
+error_report_mode=prompt
+
+[group]
+server={VC_TOOLCHAIN_DIR}\mspdbsrv.exe
+
+[files]
+main={VC_TOOLCHAIN_DIR}\cl.exe
+file01={VC_TOOLCHAIN_DIR}\c1.dll
+file02={VC_TOOLCHAIN_DIR}\c1xx.dll
+file03={VC_TOOLCHAIN_DIR}\c2.dll
+file04={VC_TOOLCHAIN_DIR}\mspdb140.dll
+file05={VC_TOOLCHAIN_DIR}\mspdbcore.dll
+file06={VC_TOOLCHAIN_DIR}\mspdbsrv.exe
+file07={VC_TOOLCHAIN_DIR}\mspft140.dll
+file08={VC_TOOLCHAIN_DIR}\vcmeta.dll
+file09={VC_TOOLCHAIN_DIR}\*\clui.dll
+file10={VC_TOOLCHAIN_DIR}\*\mspft140ui.dll
+file11={VC_TOOLCHAIN_DIR}\localespc.dll
+file12={VC_TOOLCHAIN_DIR}\cppcorecheck.dll
+file13={VC_TOOLCHAIN_DIR}\experimentalcppcorecheck.dll
+file14={VC_TOOLCHAIN_DIR}\espxengine.dll
+
+[output-file-patterns]
+outputfile01=\s*""([^ "",]+\.cpp\.txt\.json)\""
+
+[output-file-rules]
+rule01=*\sqmcpp*.log|discard=true
+rule02=*\vctoolstelemetry*.dat|discard=true
+rule03=*\Microsoft\Windows\Temporary Internet Files\*|discard=true
+rule04=*\Microsoft\Windows\INetCache\*|discard=true
+
+[input-file-rules]
+rule01=*\sqmcpp*.log|ignore_transient_errors=true;ignore_unexpected_input=true
+rule02=*\vctoolstelemetry*.dat|ignore_transient_errors=true;ignore_unexpected_input=true
+rule03=*\Microsoft\Windows\Temporary Internet Files\*|ignore_transient_errors=true;ignore_unexpected_input=true
+rule04=*\Microsoft\Windows\INetCache\*|ignore_transient_errors=true;ignore_unexpected_input=true
+
+[system-file-filters]
+filter01=msvcr*.dll
+filter02=msvcp*.dll
+filter03=vcruntime140*.dll
+filter04=appcrt140*.dll
+filter05=desktopcrt140*.dll
+filter06=concrt140*.dll",
+			["mspdbsrv.exe"] = @"
+[tool]
+use_cache=no
+
+[files]
+main={VC_TOOLCHAIN_DIR}\mspdbsrv.exe
+
+[openmp]
+omp=true
+omp_min_threads=1",
+		};
 	}
 
 	class UEDeployWindows : UEBuildDeploy
