[Backout] - CL51717742
#rnx
Original CL Desc
-----------------------------------------------------------------
Use MallocBinned3 by default for Windows

#Synced-CL 52054134

diff --git a/Engine/Source/Programs/UnrealBuildTool/Platform/Windows/UEBuildWindows.cs b/Engine/Source/Programs/UnrealBuildTool/Platform/Windows/UEBuildWindows.cs
--- a/Engine/Source/Programs/UnrealBuildTool/Platform/Windows/UEBuildWindows.cs
+++ b/Engine/Source/Programs/UnrealBuildTool/Platform/Windows/UEBuildWindows.cs
@@ -1494,7 +1494,7 @@ namespace UnrealBuildTool
 
 			if (Target.StaticAllocator == StaticAllocatorType.None)
 			{
-				Target.StaticAllocator = StaticAllocatorType.Binned3;
+				Target.StaticAllocator = StaticAllocatorType.Binned2;
 			}
 			if (Target.WindowsPlatform.bEnableAddressSanitizer)
 			{
