UnrealBuildTool: Clean up toolchain info output
Add version string for ISPC compiler

#rnx
#Synced-CL 52145182

diff --git a/Engine/Source/Programs/UnrealBuildTool/Modes/BuildMode.cs b/Engine/Source/Programs/UnrealBuildTool/Modes/BuildMode.cs
--- a/Engine/Source/Programs/UnrealBuildTool/Modes/BuildMode.cs
+++ b/Engine/Source/Programs/UnrealBuildTool/Modes/BuildMode.cs
@@ -1014,32 +1014,40 @@ namespace UnrealBuildTool
 		/// <param name="Logger">Logger for output</param>
 		static void OutputToolchainInfo(List<TargetDescriptor> TargetDescriptors, TargetMakefile[] Makefiles, ILogger Logger)
 		{
-			List<int> OutputIndices = new List<int>();
+			Dictionary<(string, UnrealTargetPlatform, UnrealTargetConfiguration), IEnumerable<string>> DistinctTargets = [];
 			for (int Idx = 0; Idx < TargetDescriptors.Count; Idx++)
 			{
-				if (!TargetDescriptors[Idx].bQuiet)
+				if (!TargetDescriptors[Idx].bQuiet && Makefiles[Idx].Diagnostics.Count > 0)
 				{
-					OutputIndices.Add(Idx);
+					DistinctTargets.TryAdd((TargetDescriptors[Idx].Name, TargetDescriptors[Idx].Platform, TargetDescriptors[Idx].Configuration), Makefiles[Idx].Diagnostics);
 				}
 			}
 
-			if (OutputIndices.Count == 1)
+			if (DistinctTargets.Count == 0)
 			{
-				foreach (string Diagnostic in Makefiles[OutputIndices[0]].Diagnostics)
+				return;
+			}
+
+			Logger.LogInformation("===== Toolchain Information =====");
+			if (DistinctTargets.Count == 1)
+			{
+				foreach (string Diagnostic in DistinctTargets.First().Value)
 				{
 					Logger.LogInformation("{Diagnostic}", Diagnostic);
 				}
 			}
 			else
 			{
-				foreach (int OutputIndex in OutputIndices)
+				foreach (KeyValuePair<(string, UnrealTargetPlatform, UnrealTargetConfiguration), IEnumerable<string>> DistinctTarget in DistinctTargets)
 				{
-					foreach (string Diagnostic in Makefiles[OutputIndex].Diagnostics)
+					Logger.LogInformation("{Name} {Platform} {Configuration}:", DistinctTarget.Key.Item1, DistinctTarget.Key.Item2, DistinctTarget.Key.Item3);
+					foreach (string Diagnostic in DistinctTarget.Value)
 					{
-						Logger.LogInformation("{Name}: {Diagnostic}", TargetDescriptors[OutputIndex].Name, Diagnostic);
+						Logger.LogInformation("  {Diagnostic}", Diagnostic);
 					}
 				}
 			}
+			Logger.LogInformation("=================================");
 		}
 
 		/// <summary>
diff --git a/Engine/Source/Programs/UnrealBuildTool/Platform/Clang/ClangToolChain.cs b/Engine/Source/Programs/UnrealBuildTool/Platform/Clang/ClangToolChain.cs
--- a/Engine/Source/Programs/UnrealBuildTool/Platform/Clang/ClangToolChain.cs
+++ b/Engine/Source/Programs/UnrealBuildTool/Platform/Clang/ClangToolChain.cs
@@ -251,6 +251,65 @@ namespace UnrealBuildTool
 			LazyInfo = new Lazy<ClangToolChainInfo>(() => { return GetToolChainInfo(); }); // Don't change to => GetToolChainInfo().. it doesnt produce correct code
 		}
 
+		public override void GetVersionInfo(List<string> Lines)
+		{
+			base.GetVersionInfo(Lines);
+
+			Lines.Add($"Using Clang compiler {Info.ClangVersion} ({Info.Clang})");
+			Lines.Add($"  {Info.ClangVersionString.ReplaceLineEndings(" ")}");
+			Lines.Add($"Using Clang archiver ({Info.Archiver})");
+			Lines.Add($"  {Info.ArchiverVersionString.ReplaceLineEndings(" ")}");
+
+			if (TargetRules?.bPGOOptimize == true)
+			{
+				if (Options.HasFlag(ClangToolChainOptions.SampleBasedPGO))
+				{
+					Lines.Add("Using Sample-Based PGO (profile guided optimization)");
+				}
+				else
+				{
+					Lines.Add("Using PGO (profile guided optimization)");
+				}
+			}
+			else if (TargetRules?.bPGOProfile == true)
+			{
+				Lines.Add("Using PGI (profile guided instrumentation)");
+			}
+
+			if (Options.HasFlag(ClangToolChainOptions.EnableLinkTimeOptimization))
+			{
+				if (Options.HasFlag(ClangToolChainOptions.EnableThinLTO))
+				{
+					Logger.LogInformation("Using ThinLTO (link-time optimization)");
+				}
+				else
+				{
+					Logger.LogInformation("Using LTO (link-time optimization)");
+				}
+			}
+
+			if (Options.HasFlag(ClangToolChainOptions.EnableAddressSanitizer) ||
+				Options.HasFlag(ClangToolChainOptions.EnableThreadSanitizer) ||
+				Options.HasFlag(ClangToolChainOptions.EnableUndefinedBehaviorSanitizer) ||
+				Options.HasFlag(ClangToolChainOptions.EnableMemorySanitizer) ||
+				Options.HasFlag(ClangToolChainOptions.EnableAutoRTFMSanitizer) ||
+				Options.HasFlag(ClangToolChainOptions.EnableLibFuzzer))
+			{
+				string SanitizerInfo = "Building with:";
+				string StaticOrShared = Options.HasFlag(ClangToolChainOptions.EnableSharedSanitizer) ? " dynamically" : " statically";
+
+				SanitizerInfo += Options.HasFlag(ClangToolChainOptions.EnableAddressSanitizer) ? StaticOrShared + " linked AddressSanitizer" : "";
+				SanitizerInfo += Options.HasFlag(ClangToolChainOptions.EnableThreadSanitizer) ? StaticOrShared + " linked ThreadSanitizer" : "";
+				SanitizerInfo += Options.HasFlag(ClangToolChainOptions.EnableUndefinedBehaviorSanitizer) ? StaticOrShared + " linked UndefinedBehaviorSanitizer" : "";
+				SanitizerInfo += Options.HasFlag(ClangToolChainOptions.EnableMemorySanitizer) ? StaticOrShared + " linked MemorySanitizer" : "";
+				SanitizerInfo += Options.HasFlag(ClangToolChainOptions.EnableAutoRTFMSanitizer) ? " AutoRTFMSanitizer" : "";
+				SanitizerInfo += Options.HasFlag(ClangToolChainOptions.EnableLibFuzzer) ? StaticOrShared + " linked LibFuzzer" : "";
+
+				Lines.Add(SanitizerInfo);
+			}
+		}
+
+		/// <inheritdoc/>
 		public override void SetEnvironmentVariables()
 		{
 			// Don't allow the INCLUDE environment variables to propagate.
diff --git a/Engine/Source/Programs/UnrealBuildTool/Platform/Linux/LinuxToolChain.cs b/Engine/Source/Programs/UnrealBuildTool/Platform/Linux/LinuxToolChain.cs
--- a/Engine/Source/Programs/UnrealBuildTool/Platform/Linux/LinuxToolChain.cs
+++ b/Engine/Source/Programs/UnrealBuildTool/Platform/Linux/LinuxToolChain.cs
@@ -5,7 +5,6 @@ using System.Collections.Generic;
 using System.Diagnostics;
 using System.IO;
 using System.Linq;
-using System.Threading;
 using EpicGames.Core;
 using Microsoft.Extensions.Logging;
 using UnrealBuildBase;
@@ -65,6 +64,41 @@ namespace UnrealBuildTool
 			}
 		}
 
+		/// <inheritdoc/>
+		public override void GetVersionInfo(List<string> Lines)
+		{
+			base.GetVersionInfo(Lines);
+			Lines.Add($"Using {(LinuxInfo.BaseLinuxPath == null ? "system toolchain" : $"toolchain located at '{LinuxInfo.BaseLinuxPath}'")}");
+			Lines.Add($"Using {(ShouldUseLibcxx() ? "bundled libc++" : "compiler default (most likely libstdc++)")} standard C++ library.");
+			if (TargetRules != null)
+			{
+				if (TargetRules.Architecture == UnrealArch.X64 && TargetRules.MinCpuArchX64 > MinimumCpuArchitectureX64.None)
+				{
+					Lines.Add($"Targeted minimum X64 CPU architecture: {TargetRules.MinCpuArchX64}");
+				}
+				else if (TargetRules.Architecture == UnrealArch.Arm64 && TargetRules.MinArm64CpuTarget > Arm64TargetCpuType.Default)
+				{
+					Lines.Add($"Targeted minimum ARM64 CPU architecture: {TargetRules.MinArm64CpuTarget}");
+				}
+				if (TargetRules.bCodeCoverage)
+				{
+					Lines.Add("Using --coverage build flag");
+				}
+				if (TargetRules.bEnablePIE)
+				{
+					Lines.Add("Using position independent executables (PIE)");
+				}
+				if (TargetRules.bEnableStackProtection)
+				{
+					Lines.Add("Using stack protection");
+				}
+			}
+			if (Options.HasFlag(ClangToolChainOptions.CompressDebugFile))
+			{
+				Lines.Add("Compressing debug files");
+			}
+		}
+
 		protected override ClangToolChainInfo GetToolChainInfo()
 		{
 			DirectoryReference? MultiArchRoot = PlatformSDK.GetSDKLocation();
@@ -789,91 +823,6 @@ namespace UnrealBuildTool
 		/// </summary>
 		private readonly Dictionary<FileItem, FileItem> LinkToPrelinkFileList = [];
 
-		/// <summary>
-		/// Tracks that information about used C++ library is only printed once
-		/// </summary>
-		static int s_hasPrintedBuildDetails = 0; // per-process static so this is only printed once per-process
-		protected void PrintBuildDetails(CppCompileEnvironment CompileEnvironment, ILogger Logger)
-		{
-			Logger.LogInformation("------- Build details --------");
-			Logger.LogInformation("Using {ToolchainInfo}.", LinuxInfo.BaseLinuxPath == null ? "system toolchain" : $"toolchain located at '{LinuxInfo.BaseLinuxPath}'");
-			Logger.LogInformation("Using clang ({ClangPath}) version '{ClangVersionString}' (string), {ClangVersionMajor} (major), {ClangVersionMinor} (minor), {ClangVersionPatch} (patch)",
-				Info.Clang, Info.ClangVersionString, Info.ClangVersion.Major, Info.ClangVersion.Minor, Info.ClangVersion.Build);
-
-			// inform the user which C++ library the engine is going to be compiled against - important for compatibility with third party code that uses STL
-			Logger.LogInformation("Using {Lib} standard C++ library.", ShouldUseLibcxx() ? "bundled libc++" : "compiler default (most likely libstdc++)");
-			Logger.LogInformation("Using lld linker");
-			Logger.LogInformation("Using llvm-ar ({LlvmAr}) version '{LlvmArVersionString} (string)'", Info.Archiver, Info.ArchiverVersionString);
-
-			if (Options.HasFlag(ClangToolChainOptions.EnableAddressSanitizer) ||
-				Options.HasFlag(ClangToolChainOptions.EnableThreadSanitizer) ||
-				Options.HasFlag(ClangToolChainOptions.EnableUndefinedBehaviorSanitizer) ||
-				Options.HasFlag(ClangToolChainOptions.EnableMemorySanitizer) ||
-				Options.HasFlag(ClangToolChainOptions.EnableAutoRTFMSanitizer) ||
-				Options.HasFlag(ClangToolChainOptions.EnableLibFuzzer))
-			{
-				string SanitizerInfo = "Building with:";
-				string StaticOrShared = Options.HasFlag(ClangToolChainOptions.EnableSharedSanitizer) ? " dynamically" : " statically";
-
-				SanitizerInfo += Options.HasFlag(ClangToolChainOptions.EnableAddressSanitizer) ? StaticOrShared + " linked AddressSanitizer" : "";
-				SanitizerInfo += Options.HasFlag(ClangToolChainOptions.EnableThreadSanitizer) ? StaticOrShared + " linked ThreadSanitizer" : "";
-				SanitizerInfo += Options.HasFlag(ClangToolChainOptions.EnableUndefinedBehaviorSanitizer) ? StaticOrShared + " linked UndefinedBehaviorSanitizer" : "";
-				SanitizerInfo += Options.HasFlag(ClangToolChainOptions.EnableMemorySanitizer) ? StaticOrShared + " linked MemorySanitizer" : "";
-				SanitizerInfo += Options.HasFlag(ClangToolChainOptions.EnableAutoRTFMSanitizer) ? " AutoRTFMSanitizer" : "";
-				SanitizerInfo += Options.HasFlag(ClangToolChainOptions.EnableLibFuzzer) ? StaticOrShared + " linked LibFuzzer" : "";
-
-				Logger.LogInformation("{SanitizerInfo}", SanitizerInfo);
-			}
-
-			if (Options.HasFlag(ClangToolChainOptions.CompressDebugFile))
-			{
-				Logger.LogInformation("Compressing debug files");
-			}
-
-			Logger.LogInformation("Targeted minimum CPU architecture: {Architecture}", (CompileEnvironment.Architecture == UnrealArch.X64) ? CompileEnvironment.MinCpuArchX64 : CompileEnvironment.MinArm64CpuTarget);
-
-			if (CompileEnvironment.bPGOOptimize)
-			{
-				Logger.LogInformation("Using PGO (profile guided optimization).");
-				Logger.LogInformation("  Directory for PGO data files='{CompileEnvironmentPGODirectory}'", CompileEnvironment.PGODirectory);
-				Logger.LogInformation("  Prefix for PGO data files='{CompileEnvironmentPGOFilenamePrefix}'", CompileEnvironment.PGOFilenamePrefix);
-			}
-
-			if (CompileEnvironment.bCodeCoverage)
-			{
-				Logger.LogInformation("Using --coverage build flag");
-			}
-
-			if (CompileEnvironment.bPGOProfile)
-			{
-				Logger.LogInformation("Using PGI (profile guided instrumentation).");
-			}
-
-			if (Options.HasFlag(ClangToolChainOptions.EnableLinkTimeOptimization))
-			{
-				if (Options.HasFlag(ClangToolChainOptions.EnableThinLTO))
-				{
-					Logger.LogInformation("Using ThinLTO (link-time optimization).");
-				}
-				else
-				{
-					Logger.LogInformation("Using LTO (link-time optimization).");
-				}
-			}
-
-			if (CompileEnvironment.bUsePIE)
-			{
-				Logger.LogInformation("Using position independent executables (PIE)");
-			}
-
-			if (CompileEnvironment.bUseStackProtection)
-			{
-				Logger.LogInformation("Using stack protection");
-			}
-
-			Logger.LogInformation("------------------------------");
-		}
-
 		protected override CPPOutput CompileCPPFiles(CppCompileEnvironment CompileEnvironment, IEnumerable<FileItem> InputFiles, DirectoryReference OutputDir, string ModuleName, IActionGraphBuilder Graph)
 		{
 			if (ShouldSkipCompile(CompileEnvironment))
@@ -886,11 +835,6 @@ namespace UnrealBuildTool
 
 			//UEBuildPlatform BuildPlatform = UEBuildPlatform.GetBuildPlatform(CompileEnvironment.Platform);
 
-			if (Interlocked.Exchange(ref s_hasPrintedBuildDetails, 1) == 0)
-			{
-				PrintBuildDetails(CompileEnvironment, Logger);
-			}
-
 			// Create a compile action for each source file.
 			CPPOutput Result = new();
 			foreach (FileItem SourceFile in InputFiles)
@@ -1039,7 +983,6 @@ namespace UnrealBuildTool
 			}
 
 			// because the logic choosing between lld and ld is somewhat messy atm (lld fails to link .DSO due to bugs), make the name of the linker clear
-			LinkAction.CommandDescription += " (lld)";
 			LinkAction.CommandVersion = Info.ClangVersionString;
 			LinkAction.StatusDescription = OutputFile.Location.GetFileName();
 			LinkAction.CacheBucket = GetCacheBucket(TargetRules, LinkEnvironment);
diff --git a/Engine/Source/Programs/UnrealBuildTool/Platform/Windows/PVSToolChain.cs b/Engine/Source/Programs/UnrealBuildTool/Platform/Windows/PVSToolChain.cs
--- a/Engine/Source/Programs/UnrealBuildTool/Platform/Windows/PVSToolChain.cs
+++ b/Engine/Source/Programs/UnrealBuildTool/Platform/Windows/PVSToolChain.cs
@@ -615,12 +615,13 @@ namespace UnrealBuildTool
 			}
 		}
 
+		/// <inheritdoc/>
 		public override void GetVersionInfo(List<string> Lines)
 		{
 			InnerToolChain.GetVersionInfo(Lines);
 
 			ReadOnlyPVSTargetSettings settings = Target.WindowsPlatform.PVS;
-			Lines.Add(String.Format("Using PVS-Studio {0} at {1} with analysis mode {2} ({3})", AnalyzerVersion, AnalyzerFile, (uint)settings.ModeFlags, settings.ModeFlags.ToString()));
+			Lines.Add($"Using PVS-Studio {AnalyzerVersion} at {AnalyzerFile} with analysis mode {(uint)settings.ModeFlags} ({settings.ModeFlags.ToString()})");
 		}
 
 		public override void GetExternalDependencies(HashSet<FileItem> ExternalDependencies)
diff --git a/Engine/Source/Programs/UnrealBuildTool/Platform/Windows/VCToolChain.cs b/Engine/Source/Programs/UnrealBuildTool/Platform/Windows/VCToolChain.cs
--- a/Engine/Source/Programs/UnrealBuildTool/Platform/Windows/VCToolChain.cs
+++ b/Engine/Source/Programs/UnrealBuildTool/Platform/Windows/VCToolChain.cs
@@ -121,9 +121,7 @@ namespace UnrealBuildTool
 			}
 		}
 
-		/// <summary>
-		/// Prepares the environment for building
-		/// </summary>
+		/// <inheritdoc/>
 		public override void SetEnvironmentVariables()
 		{
 			EnvVars.SetEnvironmentVariables();
@@ -162,12 +160,10 @@ namespace UnrealBuildTool
 			Environment.SetEnvironmentVariable("LIBPATH", null);
 		}
 
-		/// <summary>
-		/// Returns the version info for the toolchain. This will be output before building.
-		/// </summary>
-		/// <returns>String describing the current toolchain</returns>
+		/// <inheritdoc/>
 		public override void GetVersionInfo(List<string> Lines)
 		{
+			base.GetVersionInfo(Lines);
 			if (EnvVars.Compiler == EnvVars.ToolChain)
 			{
 				Lines.Add($"Using {WindowsPlatform.GetCompilerName(EnvVars.Compiler)} {EnvVars.ToolChainVersion} toolchain ({EnvVars.ToolChainDir}) and Windows {EnvVars.WindowsSdkVersion} SDK ({EnvVars.WindowsSdkDir}).");
diff --git a/Engine/Source/Programs/UnrealBuildTool/ToolChain/ISPCToolChain.cs b/Engine/Source/Programs/UnrealBuildTool/ToolChain/ISPCToolChain.cs
--- a/Engine/Source/Programs/UnrealBuildTool/ToolChain/ISPCToolChain.cs
+++ b/Engine/Source/Programs/UnrealBuildTool/ToolChain/ISPCToolChain.cs
@@ -25,6 +25,20 @@ namespace UnrealBuildTool
 		protected bool bAllowUbaCompression = false;
 		protected ReadOnlyTargetRules? TargetRules = null;
 
+		/// <inheritdoc/>
+		public override void GetVersionInfo(List<string> Lines)
+		{
+			base.GetVersionInfo(Lines);
+
+			if (TargetRules?.bCompileISPC == true)
+			{
+				string version = GetISPCHostCompilerVersion(BuildHostPlatform.Current.Platform).ReplaceLineEndings(" ");
+				string path = GetISPCHostCompilerPath(BuildHostPlatform.Current.Platform);
+				Lines.Add($"Using ISPC compiler ({path})");
+				Lines.Add($"  {version}");
+			}
+		}
+
 		public override void SetUpGlobalEnvironment(ReadOnlyTargetRules Target, CppCompileEnvironment GlobalCompileEnvironment, LinkEnvironment GlobalLinkEnvironment)
 		{
 			base.SetUpGlobalEnvironment(Target, GlobalCompileEnvironment, GlobalLinkEnvironment);
diff --git a/Engine/Source/Programs/UnrealBuildTool/ToolChain/UEToolChain.cs b/Engine/Source/Programs/UnrealBuildTool/ToolChain/UEToolChain.cs
--- a/Engine/Source/Programs/UnrealBuildTool/ToolChain/UEToolChain.cs
+++ b/Engine/Source/Programs/UnrealBuildTool/ToolChain/UEToolChain.cs
@@ -23,10 +23,17 @@ namespace UnrealBuildTool
 			Logger = InLogger;
 		}
 
+		/// <summary>
+		/// Prepares the environment for building
+		/// </summary>
 		public virtual void SetEnvironmentVariables()
 		{
 		}
 
+		/// <summary>
+		/// Returns the version info for the toolchain. This will be output before building.
+		/// </summary>
+		/// <returns>String describing the current toolchain</returns>>
 		public virtual void GetVersionInfo(List<string> Lines)
 		{
 		}
