UnrealBuildTool: Add support for arrays of enum values to XmlConfig.

#Synced-CL 52128377

diff --git a/Engine/Source/Programs/Shared/EpicGames.Core/BinaryReaderExtensions.cs b/Engine/Source/Programs/Shared/EpicGames.Core/BinaryReaderExtensions.cs
--- a/Engine/Source/Programs/Shared/EpicGames.Core/BinaryReaderExtensions.cs
+++ b/Engine/Source/Programs/Shared/EpicGames.Core/BinaryReaderExtensions.cs
@@ -201,6 +201,18 @@ namespace EpicGames.Core
 			{
 				return Enum.ToObject(objectType, reader.ReadInt32());
 			}
+			else if (objectType.IsArray && objectType.GetElementType()!.IsEnum)
+			{
+				Type elementType = objectType.GetElementType()!;
+				int[]? enumValuesAsInts = reader.ReadArray(reader.ReadInt32);
+				if (enumValuesAsInts is null)
+				{
+					return null;
+				}
+				Array array = Array.CreateInstance(elementType, enumValuesAsInts.Length);
+				enumValuesAsInts.CopyTo(array, 0);
+				return array;
+			}
 			else
 			{
 				throw new Exception(String.Format("Reading binary objects of type '{0}' is not currently supported.", objectType.Name));
diff --git a/Engine/Source/Programs/Shared/EpicGames.Core/BinaryWriterExtensions.cs b/Engine/Source/Programs/Shared/EpicGames.Core/BinaryWriterExtensions.cs
--- a/Engine/Source/Programs/Shared/EpicGames.Core/BinaryWriterExtensions.cs
+++ b/Engine/Source/Programs/Shared/EpicGames.Core/BinaryWriterExtensions.cs
@@ -216,6 +216,16 @@ namespace EpicGames.Core
 			{
 				writer.Write((FileReference)value);
 			}
+			else if (fieldType.IsArray && fieldType.GetElementType()!.IsEnum)
+			{
+				List<int> arrayEnumValuesAsInts = [];
+				Array arrayValues = (Array)value;
+				foreach (object arrayValue in arrayValues)
+				{
+					arrayEnumValuesAsInts.Add((int)arrayValue);
+				}
+				writer.Write(arrayEnumValuesAsInts, writer.Write);
+			}
 			else
 			{
 				throw new Exception(String.Format("Unsupported type '{0}' for binary serialization", fieldType.Name));
diff --git a/Engine/Source/Programs/UnrealBuildTool/Configuration/Xml/XmlConfig.cs b/Engine/Source/Programs/UnrealBuildTool/Configuration/Xml/XmlConfig.cs
--- a/Engine/Source/Programs/UnrealBuildTool/Configuration/Xml/XmlConfig.cs
+++ b/Engine/Source/Programs/UnrealBuildTool/Configuration/Xml/XmlConfig.cs
@@ -666,12 +666,9 @@ namespace UnrealBuildTool
 				enumType.Content = restriction;
 				element.SchemaType = enumType;
 			}
-			else if (type == typeof(string[]))
+			else if (MaybeGetInnerCollectionType(type) is Type genericType)
 			{
-				XmlSchemaElement itemElement = new XmlSchemaElement();
-				itemElement.Name = "Item";
-				itemElement.SchemaTypeName = XmlSchemaType.GetBuiltInSimpleType(XmlTypeCode.String).QualifiedName;
-				itemElement.MinOccurs = 0;
+				XmlSchemaElement itemElement = CreateSchemaFieldElement("Item", genericType);
 				itemElement.MaxOccursString = "unbounded";
 
 				XmlSchemaSequence sequence = new XmlSchemaSequence();
@@ -688,6 +685,19 @@ namespace UnrealBuildTool
 			return element;
 		}
 
+		private static Type? MaybeGetInnerCollectionType(Type type)
+		{
+			if (type.IsArray)
+			{
+				return type.GetElementType();
+			}
+			if (type.GetInterfaces().FirstOrDefault(x => x.IsGenericType && x.GetGenericTypeDefinition() == typeof(ICollection<>)) is Type collectionType)
+			{
+				return collectionType.GetGenericArguments()[0];
+			}
+			return null;
+		}
+
 		/// <summary>
 		/// Writes a schema to the given location. Avoids writing it if the file is identical.
 		/// </summary>
@@ -767,9 +777,23 @@ namespace UnrealBuildTool
 						if (nameToField.TryGetValue(keyElement.Name, out XmlConfigData.TargetMember? field))
 						{
 							object value;
-							if (field.Type == typeof(string[]))
+							if (MaybeGetInnerCollectionType(field.Type) is Type innerType)
 							{
-								value = keyElement.ChildNodes.OfType<XmlElement>().Where(x => x.Name == "Item").Select(x => x.InnerText).ToArray();
+								if (field.Type.IsArray)
+								{
+									string[] stringValues = [.. keyElement.ChildNodes.OfType<XmlElement>().Where(x => x.Name == "Item").Select(x => x.InnerText)];
+									Array valueArray = Array.CreateInstance(innerType, stringValues.Length);
+									for (int i = 0; i < stringValues.Length; ++i)
+									{
+										string stringValue = stringValues[i];
+										valueArray.SetValue(ParseValue(innerType, stringValue), i);
+									}
+									value = valueArray;
+								}
+								else
+								{
+									throw new NotImplementedException("XmlConfig support for non-array types is not implemented.");
+								}
 							}
 							else if (TryGetNullableStructType(field.Type, out Type? structType))
 							{
