UnrealBuildTool: Allow customizing which configurations disable optimizations for adaptive unity files.

#Synced-CL 52202922

diff --git a/Engine/Source/Programs/UnrealBuildTool/Configuration/Rules/ReadOnlyTargetRules.cs b/Engine/Source/Programs/UnrealBuildTool/Configuration/Rules/ReadOnlyTargetRules.cs
--- a/Engine/Source/Programs/UnrealBuildTool/Configuration/Rules/ReadOnlyTargetRules.cs
+++ b/Engine/Source/Programs/UnrealBuildTool/Configuration/Rules/ReadOnlyTargetRules.cs
@@ -336,6 +336,8 @@ namespace UnrealBuildTool
 
 		public bool bAdaptiveUnityDisablesOptimizations => Inner.bAdaptiveUnityDisablesOptimizations;
 
+		public IReadOnlyList<UnrealTargetConfiguration> AdaptiveUnityDisablesOptimizationsConfigurations => Inner.AdaptiveUnityDisablesOptimizationsConfigurations;
+
 		public bool bAdaptiveUnityDisablesPCH => Inner.bAdaptiveUnityDisablesPCH;
 
 		public bool bAdaptiveUnityDisablesPCHForProject => Inner.bAdaptiveUnityDisablesPCHForProject;
diff --git a/Engine/Source/Programs/UnrealBuildTool/Configuration/Rules/TargetRules.cs b/Engine/Source/Programs/UnrealBuildTool/Configuration/Rules/TargetRules.cs
--- a/Engine/Source/Programs/UnrealBuildTool/Configuration/Rules/TargetRules.cs
+++ b/Engine/Source/Programs/UnrealBuildTool/Configuration/Rules/TargetRules.cs
@@ -1880,6 +1880,12 @@ namespace UnrealBuildTool
 		[XmlConfigFile(Category = "BuildConfiguration")]
 		public bool bAdaptiveUnityDisablesOptimizations { get; set; }
 
+		/// <summary>
+		/// Disable optimizations for files that are in the adaptive non-unity working set for the given configurations.
+		/// </summary>
+		[XmlConfigFile(Category = "BuildConfiguration")]
+		public UnrealTargetConfiguration[] AdaptiveUnityDisablesOptimizationsConfigurations { get; set; } = [];
+
 		/// <summary>
 		/// Disables force-included PCHs for files that are in the adaptive non-unity working set.
 		/// </summary>
diff --git a/Engine/Source/Programs/UnrealBuildTool/Configuration/UEBuildModuleCPP.cs b/Engine/Source/Programs/UnrealBuildTool/Configuration/UEBuildModuleCPP.cs
--- a/Engine/Source/Programs/UnrealBuildTool/Configuration/UEBuildModuleCPP.cs
+++ b/Engine/Source/Programs/UnrealBuildTool/Configuration/UEBuildModuleCPP.cs
@@ -2072,6 +2072,10 @@ namespace UnrealBuildTool
 				{
 					Graph.AddDiagnostic("[Adaptive Build] Disabling optimizations for excluded files. Set bAdaptiveUnityDisablesOptimizations to false in BuildConfiguration.xml to change this behavior.");
 				}
+				else if (Target.AdaptiveUnityDisablesOptimizationsConfigurations.Contains(Target.Configuration))
+				{
+					Graph.AddDiagnostic($"[Adaptive Build] Disabling optimizations for excluded files in {Target.Configuration} configuration. Adjust {nameof(TargetRules.AdaptiveUnityDisablesOptimizationsConfigurations)} in BuildConfiguration.xml to change this behavior.");
+				}
 				if (Target.bAdaptiveUnityEnablesEditAndContinue)
 				{
 					Graph.AddDiagnostic("[Adaptive Build] Enabling Edit & Continue for excluded files. Set bAdaptiveUnityEnablesEditAndContinue to false in BuildConfiguration.xml to change this behavior.");
@@ -2088,6 +2092,10 @@ namespace UnrealBuildTool
 				{
 					AdaptiveUnityEnvironment.bOptimizeCode = false;
 				}
+				if (Target.AdaptiveUnityDisablesOptimizationsConfigurations.Contains(Target.Configuration))
+				{
+					AdaptiveUnityEnvironment.bOptimizeCode = false;
+				}
 				if (Target.bAdaptiveUnityEnablesEditAndContinue)
 				{
 					AdaptiveUnityEnvironment.bSupportEditAndContinue = true;
