[iOS/macOS] Removing aggressive fallback in the portable toolchain path

#rnx
#Synced-CL 52229246

diff --git a/Engine/Source/Programs/UnrealBuildTool/Platform/Mac/ApplePortableToolChain.cs b/Engine/Source/Programs/UnrealBuildTool/Platform/Mac/ApplePortableToolChain.cs
--- a/Engine/Source/Programs/UnrealBuildTool/Platform/Mac/ApplePortableToolChain.cs
+++ b/Engine/Source/Programs/UnrealBuildTool/Platform/Mac/ApplePortableToolChain.cs
@@ -135,33 +135,9 @@ namespace UnrealBuildTool
 				return MappedVersion;
 			}
 
-			// Fallback to auto-discovery of latest v* directory
-			string? LatestVersion = null;
-			int LatestNumericVersion = -1;
+			Logger.LogDebug("No suitable portable toolchain was found.");
 
-			if (DirectoryReference.Exists(AppleToolchainDir))
-			{
-				foreach (DirectoryReference Dir in DirectoryReference.EnumerateDirectories(AppleToolchainDir))
-				{
-					string DirName = Dir.GetDirectoryName();
-					if (DirName.StartsWith("v", StringComparison.OrdinalIgnoreCase))
-					{
-						int NumericVersion = ParseVersionNumber(DirName);
-						if (NumericVersion > LatestNumericVersion)
-						{
-							LatestNumericVersion = NumericVersion;
-							LatestVersion = DirName;
-						}
-					}
-				}
-			}
-
-			if (LatestVersion != null)
-			{
-				Logger.LogInformation("Auto-discovered Apple portable toolchain version: {Version}", LatestVersion);
-			}
-
-			return LatestVersion;
+			return null;
 		}
 
 		/// <summary>
@@ -231,16 +207,6 @@ namespace UnrealBuildTool
 			return null;
 		}
 
-		private static int ParseVersionNumber(string VersionString)
-		{
-			int UnderscoreIndex = VersionString.IndexOf('_', StringComparison.Ordinal);
-			string NumericPart = UnderscoreIndex > 1
-				? VersionString.Substring(1, UnderscoreIndex - 1)
-				: VersionString.Substring(1);
-
-			return Int32.TryParse(NumericPart, out int Result) ? Result : 0;
-		}
-
 		private static DirectoryReference? TryGetAutoSDKDir()
 		{
 			string? SdksRoot = Environment.GetEnvironmentVariable("UE_SDKS_ROOT");
