Enable MallocBinned3 by default for Windows

Also set UE_MB3_USE_CACHED_PAGE_ALLOCATOR_FOR_LARGE_ALLOCS=1 by default for WIndows/Android/IOS for performance.

#Synced-CL 52232959

diff --git a/Engine/Source/Programs/UnrealBuildTool/Platform/Windows/UEBuildWindows.cs b/Engine/Source/Programs/UnrealBuildTool/Platform/Windows/UEBuildWindows.cs
--- a/Engine/Source/Programs/UnrealBuildTool/Platform/Windows/UEBuildWindows.cs
+++ b/Engine/Source/Programs/UnrealBuildTool/Platform/Windows/UEBuildWindows.cs
@@ -1494,7 +1494,7 @@ namespace UnrealBuildTool
 
 			if (Target.StaticAllocator == StaticAllocatorType.None)
 			{
-				Target.StaticAllocator = StaticAllocatorType.Binned2;
+				Target.StaticAllocator = StaticAllocatorType.Binned3;
 			}
 			if (Target.WindowsPlatform.bEnableAddressSanitizer)
 			{
diff --git a/Engine/Source/Runtime/Core/Public/HAL/MallocBinned3.h b/Engine/Source/Runtime/Core/Public/HAL/MallocBinned3.h
--- a/Engine/Source/Runtime/Core/Public/HAL/MallocBinned3.h
+++ b/Engine/Source/Runtime/Core/Public/HAL/MallocBinned3.h
@@ -20,7 +20,11 @@
 
 
 #ifndef UE_MB3_USE_CACHED_PAGE_ALLOCATOR_FOR_LARGE_ALLOCS
-#	define UE_MB3_USE_CACHED_PAGE_ALLOCATOR_FOR_LARGE_ALLOCS (0)
+#	if PLATFORM_ANDROID || PLATFORM_IOS || PLATFORM_WINDOWS
+#		define UE_MB3_USE_CACHED_PAGE_ALLOCATOR_FOR_LARGE_ALLOCS (1)
+#	else
+#		define UE_MB3_USE_CACHED_PAGE_ALLOCATOR_FOR_LARGE_ALLOCS (0)
+#	endif
 #endif
 
 #if PLATFORM_IOS
