UnrealBuildTool: Use separate executor for GatherModuleDependency actions

#rnx
#Synced-CL 52229246

diff --git a/Engine/Source/Programs/UnrealBuildTool/Modes/BuildMode.cs b/Engine/Source/Programs/UnrealBuildTool/Modes/BuildMode.cs
--- a/Engine/Source/Programs/UnrealBuildTool/Modes/BuildMode.cs
+++ b/Engine/Source/Programs/UnrealBuildTool/Modes/BuildMode.cs
@@ -571,7 +571,7 @@ namespace UnrealBuildTool
 					if (PreprocessActions.Count > 0)
 					{
 						Logger.LogInformation("Updating module dependencies...");
-						await ActionGraph.ExecuteActionsAsync(BuildConfiguration, PreprocessActions, TargetDescriptors, Logger, executor);
+						await ActionGraph.ExecuteActionsAsync(BuildConfiguration, PreprocessActions, TargetDescriptors, Logger, null);
 
 						foreach (FileItem ProducedItem in PreprocessActions.SelectMany(x => x.ProducedItems))
 						{
