IWYU Changes
- Intermediate .h.cpp files need to be 'associated' with the header so that we can get correct data from IWYU about that header.
- Adding metadata export code path to IWYU to export data not found in .iwyu files.
- Fixing an issue where ParseSkippedHeaders would count deprecated includes as skipped.
- Various bug fixes.

#Synced-CL 52214067

diff --git a/Engine/Source/Programs/UnrealBuildTool/Configuration/UEBuildModuleCPP.cs b/Engine/Source/Programs/UnrealBuildTool/Configuration/UEBuildModuleCPP.cs
--- a/Engine/Source/Programs/UnrealBuildTool/Configuration/UEBuildModuleCPP.cs
+++ b/Engine/Source/Programs/UnrealBuildTool/Configuration/UEBuildModuleCPP.cs
@@ -3167,7 +3167,7 @@ namespace UnrealBuildTool
 			GeneratedHeaderCppContents.Add("// This file is automatically generated at compile-time to include a user created header file.");
 			GeneratedHeaderCppContents.Add("#define UE_DIRECT_HEADER_COMPILE 1");
 			GeneratedHeaderCppContents.Add($"#define __COMPILING_{HeaderName}");
-			GeneratedHeaderCppContents.Add($"#include \"{IncludeFileString.Replace('\\', '/')}\"");
+			GeneratedHeaderCppContents.Add($"#include \"{IncludeFileString.Replace('\\', '/')}\" // IWYU pragma: associated");
 			return GeneratedHeaderCppContents;
 		}
 
diff --git a/Engine/Source/Programs/UnrealBuildTool/Modes/IWYUMode.cs b/Engine/Source/Programs/UnrealBuildTool/Modes/IWYUMode.cs
--- a/Engine/Source/Programs/UnrealBuildTool/Modes/IWYUMode.cs
+++ b/Engine/Source/Programs/UnrealBuildTool/Modes/IWYUMode.cs
@@ -343,6 +343,12 @@ namespace UnrealBuildTool
 		[CommandLine("-PrintInfo")]
 		public string PrintInfo = "";
 
+		/// <summary>
+		/// Export resolved IWYU metadata to a JSON file for external tools.
+		/// </summary>
+		[CommandLine("-ExportMetadata=")]
+		public string? ExportMetadataPath;
+
 		private string? GetModuleToUpdateName(TargetDescriptor Descriptor)
 		{
 			if (!String.IsNullOrEmpty(ModuleToUpdate))
@@ -408,7 +414,7 @@ namespace UnrealBuildTool
 			if (TargetDescriptors.Count != 1)
 			{
 				Logger.LogError("IWYUMode can only handle command lines that produce one target (Cmdline: {CommandLine})", Arguments);
-				return 0;
+				return -1;
 			}
 
 			TargetDescriptors[0].IntermediateEnvironment = UnrealIntermediateEnvironment.IWYU;
@@ -482,8 +488,14 @@ namespace UnrealBuildTool
 
 				if (PathToUpdateList.Length == 0 && UEModuleToUpdate == null)
 				{
-					Logger.LogError($"Need to provide -ModuleToUpdate=<modulename> or -PathToUpdate=<partofpath> to run IWYU.");
-					return -1;
+					if (String.IsNullOrEmpty(ExportMetadataPath))
+					{
+						Logger.LogError($"Need to provide -ModuleToUpdate=<modulename> or -PathToUpdate=<partofpath> to run IWYU.");
+						return -1;
+					}
+
+					// When exporting metadata without a module/path filter, include all modules
+					Logger.LogInformation("No -ModuleToUpdate or -PathToUpdate specified. Will analyze all modules for metadata export.");
 				}
 
 				int ModulesToUpdateCount = 0;
@@ -1328,6 +1340,19 @@ namespace UnrealBuildTool
 				return 1;
 			}
 
+			if (!String.IsNullOrEmpty(ExportMetadataPath))
+			{
+				using ITimelineEvent exportEvent = Timeline.ScopeEvent("ExportMetadata");
+				ExportMetadata(ExportMetadataPath, Infos, NameToModule, systemPaths, TargetDescriptors[0], HeaderDeprecationTag);
+				Logger.LogInformation("Exported IWYU metadata to {Path}", ExportMetadataPath);
+
+				// If only exporting metadata (no -Write or -Compare), skip the update/compare step
+				if (!bWrite && !bCompare)
+				{
+					return 0;
+				}
+			}
+
 			if (bCompare)
 			{
 				using ITimelineEvent compareFiles = Timeline.ScopeEvent("CompareFiles");
@@ -1437,6 +1462,7 @@ namespace UnrealBuildTool
 
 					bool HasIncludeGuard = false;
 					int IfCount = 0;
+					int SkipDepth = 0; // > 0 means we're inside a known-false #if block
 					string[] Lines = File.ReadAllLines(Info.File);
 					foreach (string Line in Lines)
 					{
@@ -1457,14 +1483,42 @@ namespace UnrealBuildTool
 									continue;
 								}
 							}
+
+							if (SkipDepth > 0)
+							{
+								// Already skipping — just track nesting depth
+								SkipDepth++;
+							}
+							else
+							{
+								// Check if this is a known-false condition
+								ReadOnlySpan<char> Condition = LineSpan.StartsWith("if ") ? LineSpan.Slice(3).TrimStart()
+									: LineSpan.StartsWith("ifdef ") ? LineSpan.Slice(6).TrimStart()
+									: ReadOnlySpan<char>.Empty;
+
+								if (Condition.Length > 0 &&
+									(Condition.SequenceEqual("0") ||
+									 Condition.StartsWith("UE_ENABLE_INCLUDE_ORDER_DEPRECATED_IN_5_")))
+								{
+									SkipDepth = 1;
+								}
+							}
 							++IfCount;
 						}
 						else if (LineSpan.StartsWith("endif"))
 						{
+							if (SkipDepth > 0)
+							{
+								SkipDepth--;
+							}
 							--IfCount;
 						}
 						else if (LineSpan.StartsWith("include"))
 						{
+							if (SkipDepth > 0)
+							{
+								continue; // Inside a known-false #if block — skip this include
+							}
 							ReadOnlySpan<char> IncludeSpan = LineSpan.Slice("include".Length).TrimStart();
 							char LeadingIncludeChar = IncludeSpan[0];
 							if (LeadingIncludeChar == '"' || LeadingIncludeChar == '<')
@@ -2541,5 +2595,122 @@ namespace UnrealBuildTool
 			Logger.LogInformation("Sorted top list written to {ResultPath} (Can be used as input to -PathToUpdateFile=<file>)", resultFile);
 			return 0;
 		}
+
+		static void ExportMetadata(string OutputPath, Dictionary<string, IWYUInfo> Infos, Dictionary<string, UEBuildModule> NameToModule, List<string> SystemPaths, TargetDescriptor Target, string DeprecationTag)
+		{
+			using FileStream Stream = new FileStream(OutputPath, FileMode.Create);
+			using Utf8JsonWriter Writer = new Utf8JsonWriter(Stream, new JsonWriterOptions { Indented = false });
+
+			Writer.WriteStartObject();
+
+			// Header
+			Writer.WriteNumber("Version", 1);
+			Writer.WriteString("TargetName", Target.Name);
+			Writer.WriteString("Platform", Target.Platform.ToString());
+			Writer.WriteString("Configuration", Target.Configuration.ToString());
+			Writer.WriteString("EngineDir", Unreal.EngineDirectory.FullName);
+			Writer.WriteString("Timestamp", DateTime.UtcNow.ToString("o"));
+			Writer.WriteString("DeprecationDefine", DeprecationTag);
+
+			// Modules
+			Writer.WriteStartObject("Modules");
+			foreach (KeyValuePair<string, UEBuildModule> Kvp in NameToModule)
+			{
+				UEBuildModule Module = Kvp.Value;
+				Writer.WriteStartObject(Kvp.Key);
+				Writer.WriteString("Name", Module.Name);
+				Writer.WriteString("Directory", Module.ModuleDirectory.FullName);
+				Writer.WriteString("IWYUSupport", Module.Rules.IWYUSupport.ToString());
+
+				Writer.WriteStartArray("PublicDependencies");
+				foreach (string Dep in Module.Rules.PublicDependencyModuleNames)
+				{
+					Writer.WriteStringValue(Dep);
+				}
+				Writer.WriteEndArray();
+
+				Writer.WriteStartArray("PrivateDependencies");
+				foreach (string Dep in Module.Rules.PrivateDependencyModuleNames)
+				{
+					Writer.WriteStringValue(Dep);
+				}
+				Writer.WriteEndArray();
+
+				Writer.WriteEndObject();
+			}
+			Writer.WriteEndObject();
+
+			// Files
+			Writer.WriteStartObject("Files");
+			foreach (KeyValuePair<string, IWYUInfo> Kvp in Infos)
+			{
+				IWYUInfo Info = Kvp.Value;
+				Writer.WriteStartObject(Kvp.Key);
+				Writer.WriteString("File", Info.File);
+				Writer.WriteString("Module", Info.Module?.Name ?? "");
+				Writer.WriteBoolean("IsCpp", Info.IsCpp);
+
+				// Includes (IWYU recommended)
+				Writer.WriteStartArray("Includes");
+				foreach (IWYUIncludeEntry Include in Info.Includes)
+				{
+					Writer.WriteStartObject();
+					Writer.WriteString("Full", Include.Full);
+					Writer.WriteString("Printable", Include.Printable);
+					Writer.WriteString("Reason", Include.Reason);
+					Writer.WriteBoolean("System", Include.System);
+					Writer.WriteBoolean("InsideDecl", Include.InsideDecl);
+					Writer.WriteString("ResolvedModule", Include.Resolved?.Module?.Name ?? "");
+					Writer.WriteEndObject();
+				}
+				Writer.WriteEndArray();
+
+				// Includes seen in file (on disk)
+				Writer.WriteStartArray("IncludesSeenInFile");
+				foreach (IWYUIncludeEntry Seen in Info.IncludesSeenInFile)
+				{
+					Writer.WriteStartObject();
+					Writer.WriteString("Full", Seen.Full);
+					Writer.WriteString("Printable", Seen.Printable);
+					Writer.WriteBoolean("System", Seen.System);
+					Writer.WriteEndObject();
+				}
+				Writer.WriteEndArray();
+
+				// Forward declarations
+				Writer.WriteStartArray("ForwardDeclarations");
+				foreach (IWYUForwardEntry Fwd in Info.ForwardDeclarations)
+				{
+					Writer.WriteStartObject();
+					Writer.WriteString("Printable", Fwd.Printable);
+					Writer.WriteBoolean("Present", Fwd.Present);
+					Writer.WriteEndObject();
+				}
+				Writer.WriteEndArray();
+
+				// Missing includes
+				Writer.WriteStartArray("MissingIncludes");
+				foreach (IWYUMissingInclude Missing in Info.MissingIncludes)
+				{
+					Writer.WriteStartObject();
+					Writer.WriteString("Full", Missing.Full);
+					Writer.WriteEndObject();
+				}
+				Writer.WriteEndArray();
+
+				Writer.WriteEndObject();
+			}
+			Writer.WriteEndObject();
+
+			// System paths
+			Writer.WriteStartArray("SystemPaths");
+			foreach (string Path in SystemPaths)
+			{
+				Writer.WriteStringValue(Path);
+			}
+			Writer.WriteEndArray();
+
+			Writer.WriteEndObject();
+		}
 	}
 }
diff --git a/Engine/Source/Programs/UnrealBuildTool/ToolChain/IWYUToolChain.cs b/Engine/Source/Programs/UnrealBuildTool/ToolChain/IWYUToolChain.cs
--- a/Engine/Source/Programs/UnrealBuildTool/ToolChain/IWYUToolChain.cs
+++ b/Engine/Source/Programs/UnrealBuildTool/ToolChain/IWYUToolChain.cs
@@ -245,7 +245,7 @@ namespace UnrealBuildTool
 		{
 			if (TargetPlatform == UnrealTargetPlatform.Linux)
 			{
-				return GetForceIncludeFileArgument(ForceIncludeFile, CompileEnvironment);
+				return base.GetForceIncludeFileArgument(ForceIncludeFile, CompileEnvironment);
 			}
 			else
 			{
