EpicGames.Core: Move processor and memory related functions from UnrealBuildTool and add some tests

#rnx
#Synced-CL 52312713

diff --git a/Engine/Source/Programs/Shared/EpicGames.Core.Tests/SystemUtilsTests.cs b/Engine/Source/Programs/Shared/EpicGames.Core.Tests/SystemUtilsTests.cs
new file mode 100644
--- /dev/null
+++ b/Engine/Source/Programs/Shared/EpicGames.Core.Tests/SystemUtilsTests.cs
@@ -0,0 +1,144 @@
+// Copyright Epic Games, Inc. All Rights Reserved.
+
+using System;
+using Microsoft.VisualStudio.TestTools.UnitTesting;
+
+namespace EpicGames.Core.Tests;
+
+[TestClass]
+public class SystemUtilsTests
+{
+	#region GetLogicalProcessorCount
+
+	[TestMethod]
+	public void GetLogicalProcessorCount_ReturnsPositiveValue()
+	{
+		int count = SystemUtils.GetLogicalProcessorCount();
+		Assert.IsGreaterThan(0, count, $"Expected logical processor count > 0, got {count}");
+	}
+
+	[TestMethod]
+	public void GetLogicalProcessorCount_AtLeastEnvironmentProcessorCount()
+	{
+		// SystemUtils handles >64-core machines via group enumeration; it should never return
+		// fewer cores than the clamped Environment.ProcessorCount value.
+		int count = SystemUtils.GetLogicalProcessorCount();
+		Assert.IsGreaterThanOrEqualTo(Environment.ProcessorCount, count, $"Expected >= Environment.ProcessorCount ({Environment.ProcessorCount}), got {count}");
+	}
+
+	#endregion
+
+	#region GetPhysicalProcessorCount
+
+	[TestMethod]
+	public void GetPhysicalProcessorCount_ReturnsValidSentinelOrPositiveValue()
+	{
+		int count = SystemUtils.GetPhysicalProcessorCount();
+		Assert.IsGreaterThan(0, count, $"Expected positive count, got {count}");
+	}
+
+	[TestMethod]
+	public void GetPhysicalProcessorCount_DoesNotExceedLogicalCount()
+	{
+		int physical = SystemUtils.GetPhysicalProcessorCount();
+		int logical = SystemUtils.GetLogicalProcessorCount();
+		Assert.IsLessThanOrEqualTo(logical, physical, $"Physical core count ({physical}) must not exceed logical core count ({logical})");
+	}
+
+	#endregion
+
+	#region IsAsymmetricalProcessor
+
+	[TestMethod]
+	public void IsAsymmetricalProcessor_FalseWhenPhysicalEqualsLogical()
+	{
+		int physical = SystemUtils.GetPhysicalProcessorCount();
+		int logical = SystemUtils.GetLogicalProcessorCount();
+		if (physical == logical)
+		{
+			Assert.IsFalse(SystemUtils.IsAsymmetricalProcessor(), $"Expected false when physical ({physical}) == logical ({logical})");
+		}
+	}
+
+	[TestMethod]
+	public void IsAsymmetricalProcessor_FalseWhenPhysicalIsHalfLogical()
+	{
+		int physical = SystemUtils.GetPhysicalProcessorCount();
+		int logical = SystemUtils.GetLogicalProcessorCount();
+		if (logical == physical * 2)
+		{
+			Assert.IsFalse(SystemUtils.IsAsymmetricalProcessor(), $"Expected false for symmetric HT setup: physical={physical}, logical={logical}");
+		}
+	}
+
+	#endregion
+
+	#region GetTotalSystemMemoryBytes
+
+	[TestMethod]
+	public void GetTotalSystemMemoryBytes_ReturnsPositiveValue()
+	{
+		long bytes = SystemUtils.GetTotalSystemMemoryBytes();
+		Assert.IsGreaterThan(0, bytes, $"Expected total system memory > 0, got {bytes}");
+	}
+
+	#endregion
+
+	#region GetAvailableMemoryBytes
+
+	[TestMethod]
+	public void GetAvailableMemoryBytes_ReturnsValidValue()
+	{
+		GC.Collect();
+		long bytes = SystemUtils.GetAvailableMemoryBytes();
+		Assert.IsGreaterThan(0, bytes, $"Expected available memory > 0 after GC.Collect(), got {bytes}");
+	}
+
+	[TestMethod]
+	public void GetAvailableMemoryBytes_NotGreaterThanTotal()
+	{
+		GC.Collect();
+		long available = SystemUtils.GetAvailableMemoryBytes();
+		long total = SystemUtils.GetTotalSystemMemoryBytes();
+		Assert.IsLessThanOrEqualTo(total, available, $"Available ({available}) must not exceed total ({total})");
+	}
+
+	#endregion
+
+	#region GetFreeMemoryBytes
+
+	[TestMethod]
+	public void GetFreeMemoryBytes_ReturnsValidValue()
+	{
+		GC.Collect();
+		long bytes = SystemUtils.GetFreeMemoryBytes();
+		Assert.IsGreaterThanOrEqualTo(0, bytes, $"Expected free memory >= 0 after GC.Collect(), got {bytes}");
+	}
+
+	[TestMethod]
+	public void GetFreeMemoryBytes_NotGreaterThanTotal()
+	{
+		GC.Collect();
+		long free = SystemUtils.GetFreeMemoryBytes();
+		long total = SystemUtils.GetTotalSystemMemoryBytes();
+		Assert.IsLessThanOrEqualTo(total, free, $"Free ({free}) must not exceed total ({total})");
+	}
+
+	#endregion
+
+	#region GetTotalCpuUtilization
+
+	[TestMethod]
+	public void GetTotalCpuUtilization_ReturnsValidValue()
+	{
+		if (OperatingSystem.IsWindows())
+		{
+			bool success = SystemUtils.GetTotalCpuUtilization(out float utilization);
+			Assert.IsTrue(success, "Expected GetTotalCpuUtilization to return true on Windows");
+			Assert.IsGreaterThanOrEqualTo(0.0f, utilization, $"Expected utilization in [0, 100], got {utilization}");
+			Assert.IsLessThanOrEqualTo(100.0f, utilization, $"Expected utilization in [0, 100], got {utilization}");
+		}
+	}
+
+	#endregion
+}
diff --git a/Engine/Source/Programs/Shared/EpicGames.Core/SystemUtils.cs b/Engine/Source/Programs/Shared/EpicGames.Core/SystemUtils.cs
new file mode 100644
--- /dev/null
+++ b/Engine/Source/Programs/Shared/EpicGames.Core/SystemUtils.cs
@@ -0,0 +1,425 @@
+// Copyright Epic Games, Inc. All Rights Reserved.
+
+using System;
+using System.Collections.Generic;
+using System.Diagnostics;
+using System.Linq;
+using System.Runtime.InteropServices;
+using System.Text;
+
+namespace EpicGames.Core;
+
+/// <summary>
+/// Utilities for querying system hardware information such as processor counts, memory availability, and CPU utilization.
+/// </summary>
+public static class SystemUtils
+{
+	#region Processor Count
+
+	enum LogicalProcessorRelationship
+	{
+		RelationProcessorCore,
+		RelationNumaNode,
+		RelationCache,
+		RelationProcessorPackage,
+		RelationGroup,
+		RelationAll = 0xffff
+	}
+
+	[DllImport("kernel32.dll", SetLastError = true)]
+	static extern bool GetLogicalProcessorInformationEx(LogicalProcessorRelationship relationshipType, IntPtr buffer, ref uint returnedLength);
+
+	// int sysctlbyname(const char *name, void *oldp, size_t *oldlenp, void *newp, size_t newlen); // from man page
+	[DllImport("libc")]
+	static extern int sysctlbyname(string name, out int oldp, ref ulong oldlenp, IntPtr newp, ulong newlen);
+
+	/// <summary>
+	/// Gets the number of logical cores. We use this rather than Environment.ProcessorCount when possible to handle machines with > 64 cores (the single group limit available to the .NET framework).
+	/// </summary>
+	/// <returns>The number of logical cores.</returns>
+	public static int GetLogicalProcessorCount()
+	{
+		// This function uses Windows P/Invoke calls; if we're not running on Windows, just return the default.
+		if (OperatingSystem.IsWindows())
+		{
+			const int ERROR_INSUFFICIENT_BUFFER = 122;
+
+			// Determine the required buffer size to store the processor information
+			uint returnLength = 0;
+			if (!GetLogicalProcessorInformationEx(LogicalProcessorRelationship.RelationGroup, IntPtr.Zero, ref returnLength) && Marshal.GetLastWin32Error() == ERROR_INSUFFICIENT_BUFFER)
+			{
+				// Allocate a buffer for it
+				IntPtr ptr = Marshal.AllocHGlobal((int)returnLength);
+				try
+				{
+					if (GetLogicalProcessorInformationEx(LogicalProcessorRelationship.RelationGroup, ptr, ref returnLength))
+					{
+						int count = 0;
+						for (int pos = 0; pos < returnLength;)
+						{
+							LogicalProcessorRelationship type = (LogicalProcessorRelationship)Marshal.ReadInt16(ptr, pos);
+							if (type == LogicalProcessorRelationship.RelationGroup)
+							{
+								// Read the values from the embedded GROUP_RELATIONSHIP structure
+								int groupRelationshipPos = pos + 8;
+								int activeGroupCount = Marshal.ReadInt16(ptr, groupRelationshipPos + 2);
+
+								// Read the processor counts from the embedded PROCESSOR_GROUP_INFO structures
+								int groupInfoPos = groupRelationshipPos + 24;
+								for (int groupIdx = 0; groupIdx < activeGroupCount; groupIdx++)
+								{
+									count += Marshal.ReadByte(ptr, groupInfoPos + 1);
+									groupInfoPos += 40 + IntPtr.Size;
+								}
+							}
+							pos += Marshal.ReadInt32(ptr, pos + 4);
+						}
+						return count;
+					}
+				}
+				finally
+				{
+					Marshal.FreeHGlobal(ptr);
+				}
+			}
+		}
+		else if (OperatingSystem.IsLinux())
+		{
+			// query socket/logical core pairings.  There should not be duplicates in this list since each hyperthread
+			// will show up as it's own logical "cpu".  Including the socket number allows us to count multi-processor
+			// system cores correctly
+			string output = RunProcessAndReturnStdOut("lscpu", "-p='SOCKET,CPU'");
+			List<string> cpus = [.. output.Split("\n", StringSplitOptions.None).Where(x => !x.StartsWith('#'))];
+
+			return cpus.Count;
+		}
+		return Environment.ProcessorCount;
+	}
+
+	/// <summary>
+	/// Gets the number of physical cores, excluding hyper threading.
+	/// </summary>
+	/// <returns>The number of physical cores, or -1 if it could not be obtained</returns>
+	public static int GetPhysicalProcessorCount()
+	{
+		if (OperatingSystem.IsWindows())
+		{
+			const int ERROR_INSUFFICIENT_BUFFER = 122;
+
+			// Determine the required buffer size to store the processor information
+			uint returnLength = 0;
+			if (!GetLogicalProcessorInformationEx(LogicalProcessorRelationship.RelationProcessorCore, IntPtr.Zero,
+				ref returnLength) && Marshal.GetLastWin32Error() == ERROR_INSUFFICIENT_BUFFER)
+			{
+				// Allocate a buffer for it
+				IntPtr ptr = Marshal.AllocHGlobal((int)returnLength);
+				try
+				{
+					if (GetLogicalProcessorInformationEx(LogicalProcessorRelationship.RelationProcessorCore, ptr,
+						ref returnLength))
+					{
+						// As per-MSDN, this will return one structure per physical processor. Each SYSTEM_LOGICAL_PROCESSOR_INFORMATION_EX structure is of a variable size, so just skip
+						// through the list and count the number of entries.
+						int count = 0;
+						for (int pos = 0; pos < returnLength;)
+						{
+							LogicalProcessorRelationship type =
+								(LogicalProcessorRelationship)Marshal.ReadInt16(ptr, pos);
+							if (type == LogicalProcessorRelationship.RelationProcessorCore)
+							{
+								count++;
+							}
+
+							pos += Marshal.ReadInt32(ptr, pos + 4);
+						}
+
+						return count;
+					}
+				}
+				finally
+				{
+					Marshal.FreeHGlobal(ptr);
+				}
+			}
+		}
+		else if (OperatingSystem.IsMacOS())
+		{
+			ulong size = 4;
+			if (0 == sysctlbyname("hw.physicalcpu", out int value, ref size, IntPtr.Zero, 0))
+			{
+				return value;
+			}
+		}
+		else if (OperatingSystem.IsLinux())
+		{
+			// query socket/physical core pairings.  There will be duplicates in this if there are hyperthreads
+			// using the HashSet ensures that those duplicates are removed and we only count the first "cpu" found
+			// for each "core".  Including the socket number allows us to count multi-processor system cores correctly
+			string output = RunProcessAndReturnStdOut("lscpu", "-p='SOCKET,CORE'");
+			HashSet<string> cpus = [.. output.Split("\n", StringSplitOptions.None).Where(x => !x.StartsWith('#'))];
+
+			return cpus.Count;
+		}
+
+		return -1;
+	}
+
+	/// <summary>
+	/// Gets if the processor has asymmetrical cores (Windows only)
+	/// </summary>
+	/// <returns></returns>
+	public static bool IsAsymmetricalProcessor()
+	{
+		int logicalCores = GetLogicalProcessorCount();
+		int physicalCores = GetPhysicalProcessorCount();
+
+		if (physicalCores <= 0 || physicalCores == logicalCores)
+		{
+			return false;
+		}
+
+		return logicalCores != physicalCores * 2;
+	}
+
+	#endregion
+
+	#region Memory
+
+	/// <summary>
+	/// Gets the total memory bytes available, based on what is known to the garbage collector.
+	/// </summary>
+	/// <remarks>This will return a max of 2GB for a 32bit application.</remarks>
+	/// <returns>The total memory available, in bytes.</returns>
+	public static long GetAvailableMemoryBytes()
+	{
+		GCMemoryInfo memoryInfo = GC.GetGCMemoryInfo();
+		// TotalAvailableMemoryBytes will be 0 if garbage collection has not run yet
+		return memoryInfo.TotalAvailableMemoryBytes != 0 ? memoryInfo.TotalAvailableMemoryBytes : -1;
+	}
+
+	// vm_statistics64, based on the definition in <mach/vm_statistics.h>
+	[StructLayout(LayoutKind.Sequential)]
+	struct VmStatistics64
+	{
+		/*natural_t*/
+		public int _freeCount;              /* # of pages free */
+		/*natural_t*/
+		public int _activeCount;            /* # of pages active */
+		/*natural_t*/
+		public int _inactiveCount;          /* # of pages inactive */
+		/*natural_t*/
+		public int _wireCount;              /* # of pages wired down */
+		/*uint64_t */
+		public ulong _zeroFillCount;        /* # of zero fill pages */
+		/*uint64_t */
+		public ulong _reactivations;        /* # of pages reactivated */
+		/*uint64_t */
+		public ulong _pageins;              /* # of pageins */
+		/*uint64_t */
+		public ulong _pageouts;             /* # of pageouts */
+		/*uint64_t */
+		public ulong _faults;               /* # of faults */
+		/*uint64_t */
+		public ulong _cowFaults;            /* # of copy-on-writes */
+		/*uint64_t */
+		public ulong _lookups;              /* object cache lookups */
+		/*uint64_t */
+		public ulong _hits;                 /* object cache hits */
+		/*uint64_t */
+		public ulong _purges;               /* # of pages purged */
+		/*natural_t*/
+		public int _purgeableCount;         /* # of pages purgeable */
+		/*
+          	 * NB: speculative pages are already accounted for in "_freeCount",
+          	 * so "_speculativeCount" is the number of "free" pages that are
+          	 * used to hold data that was read speculatively from disk but
+          	 * haven't actually been used by anyone so far.
+          	 */
+		/*natural_t*/
+		public int _speculativeCount;       /* # of pages speculative */
+
+		/* added for rev1 */
+		/*uint64_t */
+		public ulong _decompressions;       /* # of pages decompressed */
+		/*uint64_t */
+		public ulong _compressions;         /* # of pages compressed */
+		/*uint64_t */
+		public ulong _swapins;              /* # of pages swapped in (via compression segments) */
+		/*uint64_t */
+		public ulong _swapouts;             /* # of pages swapped out (via compression segments) */
+		/*natural_t*/
+		public int _compressorPageCount;    /* # of pages used by the compressed pager to hold all the compressed data */
+		/*natural_t*/
+		public int _throttledCount;         /* # of pages throttled */
+		/*natural_t*/
+		public int _externalPageCount;      /* # of pages that are file-backed (non-swap) */
+		/*natural_t*/
+		public int _internalPageCount;      /* # of pages that are anonymous */
+		/*uint64_t */
+		public ulong _totalUncompressedPagesInCompressor; /* # of pages (uncompressed) held within the compressor. */
+	} // __attribute__((aligned(8)));
+
+	// kern_return_t host_statistics64(host_t host_priv, host_flavor_t flavor, host_info64_t host_info64_out, mach_msg_type_number_t *host_info64_outCnt); // from <mach/mach_host.h>
+	[DllImport("libc")]
+	static extern int host_statistics64(IntPtr hostPriv, int flavor, out VmStatistics64 hostInfo64Out, ref uint hostInfoCount);
+
+	// mach_port_t mach_host_self() // from <mach/mach_init.h>
+	[DllImport("libc")]
+	static extern IntPtr mach_host_self();
+
+	/// <summary>
+	/// Gets the total system memory in bytes based on what is known to the garbage collector.
+	/// </summary>
+	/// <remarks>This will return a max of 2GB for a 32bit application.</remarks>
+	/// <returns>The total system memory, in bytes.</returns>
+	public static long GetTotalSystemMemoryBytes()
+	{
+		GCMemoryInfo memoryInfo = GC.GetGCMemoryInfo();
+		return memoryInfo.TotalAvailableMemoryBytes;
+	}
+
+	/// <summary>
+	/// Gets the total memory bytes free, based on what is known to the garbage collector.
+	/// </summary>
+	/// <remarks>This will return a max of 2GB for a 32bit application.</remarks>
+	/// <returns>The total memory free, in bytes.</returns>
+	public static long GetFreeMemoryBytes()
+	{
+		long freeMemoryBytes = -1;
+
+		GCMemoryInfo memoryInfo = GC.GetGCMemoryInfo();
+		// TotalAvailableMemoryBytes will be 0 if garbage collection has not run yet
+		if (memoryInfo.TotalAvailableMemoryBytes != 0)
+		{
+			freeMemoryBytes = memoryInfo.TotalAvailableMemoryBytes - memoryInfo.MemoryLoadBytes;
+		}
+
+		// On Mac, MemoryInfo.MemoryLoadBytes includes memory used to cache disk-backed files ("Cached Files" in
+		// Activity Monitor), which can result in a significant over-estimate of memory pressure.
+		// We treat memory used for caching of disk-backed files as free for use in compilation tasks.
+		if (OperatingSystem.IsMacOS())
+		{
+			// host_statistics64() flavor, from <mach/host_info.h>
+			const int HOST_VM_INFO64 = 4;
+			// host_statistics64() count of 32bit values in output struct, from <mach/host_info.h>
+			int hostVmInfo64Count = Marshal.SizeOf(typeof(VmStatistics64)) / 4;
+
+			VmStatistics64 vmStats;
+			uint structSize = (uint)hostVmInfo64Count;
+			IntPtr host = mach_host_self();
+			if (host_statistics64(host, HOST_VM_INFO64, out vmStats, ref structSize) != 0)
+			{
+				// Return the dotnet memory estimate if the OS call fails
+				return freeMemoryBytes;
+			}
+
+			ulong outSize = 4;
+			if (0 != sysctlbyname("hw.pagesize", out int pageSize, ref outSize, IntPtr.Zero, 0))
+			{
+				pageSize = 4096; // likely result
+			}
+
+			freeMemoryBytes += (long)pageSize * (long)vmStats._externalPageCount;
+		}
+		return freeMemoryBytes;
+	}
+
+	#endregion
+
+	#region CPU Utilization
+
+	[DllImport("pdh.dll", SetLastError = true, CharSet = CharSet.Auto)]
+	static extern int PdhOpenQueryW([MarshalAs(UnmanagedType.LPWStr)] string? szDataSource, UIntPtr dwUserData, out IntPtr phQuery);
+
+	[DllImport("pdh.dll", SetLastError = true, CharSet = CharSet.Auto)]
+	static extern uint PdhAddCounter(IntPtr hQuery, string szFullCounterPath, IntPtr dwUserData, out IntPtr phCounter);
+
+	[DllImport("pdh.dll", SetLastError = true)]
+	static extern uint PdhCollectQueryData(IntPtr phQuery);
+
+	struct PdhFmtCounterValue
+	{
+		public uint _cStatus;
+		public double _doubleValue;
+	};
+
+	[DllImport("pdh.dll", SetLastError = true)]
+	static extern uint PdhGetFormattedCounterValue(IntPtr phCounter, uint dwFormat, IntPtr lpdwType, out PdhFmtCounterValue pValue);
+
+	static IntPtr s_cpuQuery;
+	static IntPtr s_cpuTotal;
+	static bool s_cpuInitialized = false;
+
+	/// <summary>
+	/// Gives the current CPU utilization.
+	/// </summary>
+	/// <param name="utilization">Percentage of CPU utilization currently.</param>
+	/// <returns>Whether or not it was successful in getting the CPU utilization.</returns>
+	public static bool GetTotalCpuUtilization(out float utilization)
+	{
+		utilization = 0.0f;
+		if (OperatingSystem.IsWindows())
+		{
+			const uint ERROR_SUCCESS = 0;
+			if (!s_cpuInitialized)
+			{
+				if (PdhOpenQueryW(null, UIntPtr.Zero, out s_cpuQuery) == ERROR_SUCCESS)
+				{
+					if (PdhAddCounter(s_cpuQuery, "\\Processor(_Total)\\% Processor Time", IntPtr.Zero, out s_cpuTotal) != ERROR_SUCCESS)
+					{
+						return false;
+					}
+
+					if (PdhCollectQueryData(s_cpuQuery) != ERROR_SUCCESS)
+					{
+						return false;
+					}
+
+					s_cpuInitialized = true;
+				}
+			}
+
+			if (s_cpuInitialized)
+			{
+				const uint PDH_FMT_DOUBLE = 0x00000200;
+
+				if (PdhCollectQueryData(s_cpuQuery) != ERROR_SUCCESS)
+				{
+					return false;
+				}
+
+				if (PdhGetFormattedCounterValue(s_cpuTotal, PDH_FMT_DOUBLE, IntPtr.Zero, out PdhFmtCounterValue counterVal) != ERROR_SUCCESS)
+				{
+					return false;
+				}
+
+				utilization = (float)counterVal._doubleValue;
+				return true;
+			}
+		}
+
+		return false;
+	}
+
+	#endregion
+
+	#region Helpers
+
+	static string RunProcessAndReturnStdOut(string command, string args)
+	{
+		ProcessStartInfo startInfo = new(command, args)
+		{
+			UseShellExecute = false,
+			RedirectStandardOutput = true,
+			RedirectStandardError = true,
+			CreateNoWindow = true,
+			StandardOutputEncoding = Encoding.UTF8,
+		};
+		using Process localProcess = Process.Start(startInfo)!;
+		string output = localProcess.StandardOutput.ReadToEnd().Trim();
+		localProcess.WaitForExit();
+		return output;
+	}
+
+	#endregion
+}
diff --git a/Engine/Source/Programs/UnrealBuildTool/Executors/UnrealBuildAccelerator/UnrealBuildAcceleratorConfig.cs b/Engine/Source/Programs/UnrealBuildTool/Executors/UnrealBuildAccelerator/UnrealBuildAcceleratorConfig.cs
--- a/Engine/Source/Programs/UnrealBuildTool/Executors/UnrealBuildAccelerator/UnrealBuildAcceleratorConfig.cs
+++ b/Engine/Source/Programs/UnrealBuildTool/Executors/UnrealBuildAccelerator/UnrealBuildAcceleratorConfig.cs
@@ -399,7 +399,7 @@ namespace UnrealBuildTool
 		/// Default: BelowNormal or Normal for an Asymmetrical processor as BelowNormal can cause scheduling issues.
 		/// </summary>
 		[XmlConfigFile(Category = "ParallelExecutor")]
-		public ProcessPriorityClass ProcessPriority { get; set; } = Utils.IsAsymmetricalProcessor() ? ProcessPriorityClass.Normal : ProcessPriorityClass.BelowNormal;
+		public ProcessPriorityClass ProcessPriority { get; set; } = SystemUtils.IsAsymmetricalProcessor() ? ProcessPriorityClass.Normal : ProcessPriorityClass.BelowNormal;
 
 		/// <summary>
 		/// When enabled, will stop compiling targets after a compile error occurs.
diff --git a/Engine/Source/Programs/UnrealBuildTool/Platform/Android/AndroidToolChain.cs b/Engine/Source/Programs/UnrealBuildTool/Platform/Android/AndroidToolChain.cs
--- a/Engine/Source/Programs/UnrealBuildTool/Platform/Android/AndroidToolChain.cs
+++ b/Engine/Source/Programs/UnrealBuildTool/Platform/Android/AndroidToolChain.cs
@@ -1206,7 +1206,7 @@ namespace UnrealBuildTool
 			if (LinkEnvironment.bAllowLTCG && Options.HasFlag(ClangToolChainOptions.EnableThinLTO))
 			{
 				// Set the weight to number of logical cores as lld can max out the available cores
-				LinkAction.Weight = Utils.GetLogicalProcessorCount();
+				LinkAction.Weight = SystemUtils.GetLogicalProcessorCount();
 
 				// Disallow remote to prevent this long running action from running on an agent if remote linking is enabled
 				LinkAction.bCanExecuteRemotely = false;
diff --git a/Engine/Source/Programs/UnrealBuildTool/Platform/Linux/LinuxToolChain.cs b/Engine/Source/Programs/UnrealBuildTool/Platform/Linux/LinuxToolChain.cs
--- a/Engine/Source/Programs/UnrealBuildTool/Platform/Linux/LinuxToolChain.cs
+++ b/Engine/Source/Programs/UnrealBuildTool/Platform/Linux/LinuxToolChain.cs
@@ -976,7 +976,7 @@ namespace UnrealBuildTool
 			if (Options.HasFlag(ClangToolChainOptions.EnableLinkTimeOptimization) && Options.HasFlag(ClangToolChainOptions.EnableThinLTO))
 			{
 				// Set the weight to number of logical cores as lld can max out the available cores
-				LinkAction.Weight = Utils.GetLogicalProcessorCount();
+				LinkAction.Weight = SystemUtils.GetLogicalProcessorCount();
 
 				// Disallow remote to prevent this long running action from running on an agent if remote linking is enabled
 				LinkAction.bCanExecuteRemotely = false;
diff --git a/Engine/Source/Programs/UnrealBuildTool/Platform/Windows/VCToolChain.cs b/Engine/Source/Programs/UnrealBuildTool/Platform/Windows/VCToolChain.cs
--- a/Engine/Source/Programs/UnrealBuildTool/Platform/Windows/VCToolChain.cs
+++ b/Engine/Source/Programs/UnrealBuildTool/Platform/Windows/VCToolChain.cs
@@ -2422,6 +2422,7 @@ namespace UnrealBuildTool
 				if (SourceFile.HasExtension(".ixx"))
 				{
 					FileItem IfcFile = FileItem.GetItemByFileReference(FileReference.Combine(GetModuleInterfaceDir(OutputDir), SourceFile.Location.ChangeExtension(".ifc").GetFileName()));
+					DirectoryReference.CreateDirectory(IfcFile.Directory.Location);
 					CompileAction.Arguments.Add("/interface");
 					CompileAction.Arguments.Add($"/ifcOutput \"{IfcFile.Location}\"");
 					CompileAction.CompiledModuleInterfaceFile = IfcFile;
@@ -3583,7 +3584,7 @@ namespace UnrealBuildTool
 				(Unreal.IsBuildMachine() || LinkEnvironment.bAllowLTCG || LinkEnvironment.bPGOOptimize || LinkEnvironment.bPGOProfile))
 			{
 				// Set the weight to number of logical cores as lld can max out the available cores
-				LinkAction.Weight = Utils.GetLogicalProcessorCount();
+				LinkAction.Weight = SystemUtils.GetLogicalProcessorCount();
 
 				// Disallow remote to prevent this long running action from running on an agent if remote linking is enabled
 				LinkAction.bCanExecuteRemotely = false;
diff --git a/Engine/Source/Programs/UnrealBuildTool/System/Utils.cs b/Engine/Source/Programs/UnrealBuildTool/System/Utils.cs
--- a/Engine/Source/Programs/UnrealBuildTool/System/Utils.cs
+++ b/Engine/Source/Programs/UnrealBuildTool/System/Utils.cs
@@ -953,383 +953,58 @@ namespace UnrealBuildTool
 		[Obsolete("Deprecated in UE5.8. Replace with Unreal.UserSettingDirectory")]
 		public static DirectoryReference? GetUserSettingDirectory() => Unreal.UserSettingDirectory;
 
-		enum LOGICAL_PROCESSOR_RELATIONSHIP
-		{
-			RelationProcessorCore,
-			RelationNumaNode,
-			RelationCache,
-			RelationProcessorPackage,
-			RelationGroup,
-			RelationAll = 0xffff
-		}
-
-		[DllImport("kernel32.dll", SetLastError = true)]
-		static extern bool GetLogicalProcessorInformationEx(LOGICAL_PROCESSOR_RELATIONSHIP RelationshipType, IntPtr Buffer, ref uint ReturnedLength);
-
 		/// <summary>
 		/// Gets the number of logical cores. We use this rather than Environment.ProcessorCount when possible to handle machines with > 64 cores (the single group limit available to the .NET framework).
 		/// </summary>
 		/// <returns>The number of logical cores.</returns>
-		public static int GetLogicalProcessorCount()
-		{
-			// This function uses Windows P/Invoke calls; if we're not running on Windows, just return the default.
-			if (OperatingSystem.IsWindows())
-			{
-				const int ERROR_INSUFFICIENT_BUFFER = 122;
-
-				// Determine the required buffer size to store the processor information
-				uint ReturnLength = 0;
-				if (!GetLogicalProcessorInformationEx(LOGICAL_PROCESSOR_RELATIONSHIP.RelationGroup, IntPtr.Zero, ref ReturnLength) && Marshal.GetLastWin32Error() == ERROR_INSUFFICIENT_BUFFER)
-				{
-					// Allocate a buffer for it
-					IntPtr Ptr = Marshal.AllocHGlobal((int)ReturnLength);
-					try
-					{
-						if (GetLogicalProcessorInformationEx(LOGICAL_PROCESSOR_RELATIONSHIP.RelationGroup, Ptr, ref ReturnLength))
-						{
-							int Count = 0;
-							for (int Pos = 0; Pos < ReturnLength;)
-							{
-								LOGICAL_PROCESSOR_RELATIONSHIP Type = (LOGICAL_PROCESSOR_RELATIONSHIP)Marshal.ReadInt16(Ptr, Pos);
-								if (Type == LOGICAL_PROCESSOR_RELATIONSHIP.RelationGroup)
-								{
-									// Read the values from the embedded GROUP_RELATIONSHIP structure
-									int GroupRelationshipPos = Pos + 8;
-									int ActiveGroupCount = Marshal.ReadInt16(Ptr, GroupRelationshipPos + 2);
-
-									// Read the processor counts from the embedded PROCESSOR_GROUP_INFO structures
-									int GroupInfoPos = GroupRelationshipPos + 24;
-									for (int GroupIdx = 0; GroupIdx < ActiveGroupCount; GroupIdx++)
-									{
-										Count += Marshal.ReadByte(Ptr, GroupInfoPos + 1);
-										GroupInfoPos += 40 + IntPtr.Size;
-									}
-								}
-								Pos += Marshal.ReadInt32(Ptr, Pos + 4);
-							}
-							return Count;
-						}
-					}
-					finally
-					{
-						Marshal.FreeHGlobal(Ptr);
-					}
-				}
-			}
-			else if (OperatingSystem.IsLinux())
-			{
-				// query socket/logical core pairings.  There should not be duplicates in this list since each hyperthread
-				// will show up as it's own logical "cpu".  Including the socket number allows us to count multi-processor
-				// system cores correctly
-				string Output = RunLocalProcessAndReturnStdOut("lscpu", "-p='SOCKET,CPU'");
-				List<string> CPUs = Output.Split("\n").Where(x => !x.StartsWith('#')).ToList();
-
-				return CPUs.Count;
-			}
-			return Environment.ProcessorCount;
-		}
-
-		// int sysctlbyname(const char *name, void *oldp, size_t *oldlenp, void *newp, size_t newlen); // from man page
-		[DllImport("libc")]
-		static extern int sysctlbyname(string name, out int oldp, ref ulong oldlenp, IntPtr newp, ulong newlen);
+		[Obsolete("Deprecated in UE5.8. Use SystemUtils.GetLogicalProcessorCount() from EpicGames.Core instead.")]
+		public static int GetLogicalProcessorCount() => SystemUtils.GetLogicalProcessorCount();
 
 		/// <summary>
 		/// Gets the number of physical cores, excluding hyper threading.
 		/// </summary>
 		/// <returns>The number of physical cores, or -1 if it could not be obtained</returns>
-		public static int GetPhysicalProcessorCount()
-		{
-			if (OperatingSystem.IsWindows())
-			{
-				const int ERROR_INSUFFICIENT_BUFFER = 122;
-
-				// Determine the required buffer size to store the processor information
-				uint ReturnLength = 0;
-				if (!GetLogicalProcessorInformationEx(LOGICAL_PROCESSOR_RELATIONSHIP.RelationProcessorCore, IntPtr.Zero,
-					ref ReturnLength) && Marshal.GetLastWin32Error() == ERROR_INSUFFICIENT_BUFFER)
-				{
-					// Allocate a buffer for it
-					IntPtr Ptr = Marshal.AllocHGlobal((int)ReturnLength);
-					try
-					{
-						if (GetLogicalProcessorInformationEx(LOGICAL_PROCESSOR_RELATIONSHIP.RelationProcessorCore, Ptr,
-							ref ReturnLength))
-						{
-							// As per-MSDN, this will return one structure per physical processor. Each SYSTEM_LOGICAL_PROCESSOR_INFORMATION_EX structure is of a variable size, so just skip 
-							// through the list and count the number of entries.
-							int Count = 0;
-							for (int Pos = 0; Pos < ReturnLength;)
-							{
-								LOGICAL_PROCESSOR_RELATIONSHIP Type =
-									(LOGICAL_PROCESSOR_RELATIONSHIP)Marshal.ReadInt16(Ptr, Pos);
-								if (Type == LOGICAL_PROCESSOR_RELATIONSHIP.RelationProcessorCore)
-								{
-									Count++;
-								}
-
-								Pos += Marshal.ReadInt32(Ptr, Pos + 4);
-							}
-
-							return Count;
-						}
-					}
-					finally
-					{
-						Marshal.FreeHGlobal(Ptr);
-					}
-				}
-			}
-			else if (OperatingSystem.IsMacOS())
-			{
-				ulong Size = 4;
-				if (0 == sysctlbyname("hw.physicalcpu", out int Value, ref Size, IntPtr.Zero, 0))
-				{
-					return Value;
-				}
-			}
-			else if (OperatingSystem.IsLinux())
-			{
-				// query socket/physical core pairings.  There will be duplicates in this if there are hyperthreads
-				// using the HashSet ensures that those duplicates are removed and we only count the first "cpu" found
-				// for each "core".  Including the socket number allows us to count multi-processor system cores correctly
-				string Output = RunLocalProcessAndReturnStdOut("lscpu", "-p='SOCKET,CORE'");
-				HashSet<string> CPUs = Output.Split("\n").Where(x => !x.StartsWith('#')).ToHashSet();
-
-				return CPUs.Count;
-			}
-
-			return -1;
-		}
+		[Obsolete("Deprecated in UE5.8. Use SystemUtils.GetPhysicalProcessorCount() from EpicGames.Core instead.")]
+		public static int GetPhysicalProcessorCount() => SystemUtils.GetPhysicalProcessorCount();
 
 		/// <summary>
-		/// Gets if the processos has asymmetrical cores (Windows only)
+		/// Gets if the processor has asymmetrical cores (Windows only)
 		/// </summary>
 		/// <returns></returns>
-		public static bool IsAsymmetricalProcessor()
-		{
-			int LogicalCores = GetLogicalProcessorCount();
-			int PhysicalCores = GetPhysicalProcessorCount();
-
-			if (PhysicalCores <= 0 || PhysicalCores == LogicalCores)
-			{
-				return false;
-			}
-
-			return LogicalCores != PhysicalCores * 2;
-		}
+		[Obsolete("Deprecated in UE5.8. Use SystemUtils.IsAsymmetricalProcessor() from EpicGames.Core instead.")]
+		public static bool IsAsymmetricalProcessor() => SystemUtils.IsAsymmetricalProcessor();
 
 		/// <summary>
 		/// Gets the total memory bytes available, based on what is known to the garbage collector.
 		/// </summary>
 		/// <remarks>This will return a max of 2GB for a 32bit application.</remarks>
 		/// <returns>The total memory available, in bytes.</returns>
-		public static long GetAvailableMemoryBytes()
-		{
-			GCMemoryInfo MemoryInfo = GC.GetGCMemoryInfo();
-			// TotalAvailableMemoryBytes will be 0 if garbage collection has not run yet
-			return MemoryInfo.TotalAvailableMemoryBytes != 0 ? MemoryInfo.TotalAvailableMemoryBytes : -1;
-		}
-
-		// vm_statistics64, based on the definition in <mach/vm_statistics.h>
-		[StructLayout(LayoutKind.Sequential)]
-		struct vm_statistics64
-		{
-			/*natural_t*/
-			public int free_count;              /* # of pages free */
-			/*natural_t*/
-			public int active_count;            /* # of pages active */
-			/*natural_t*/
-			public int inactive_count;          /* # of pages inactive */
-			/*natural_t*/
-			public int wire_count;              /* # of pages wired down */
-			/*uint64_t */
-			public ulong zero_fill_count;      /* # of zero fill pages */
-			/*uint64_t */
-			public ulong reactivations;            /* # of pages reactivated */
-			/*uint64_t */
-			public ulong pageins;              /* # of pageins */
-			/*uint64_t */
-			public ulong pageouts;             /* # of pageouts */
-			/*uint64_t */
-			public ulong faults;                   /* # of faults */
-			/*uint64_t */
-			public ulong cow_faults;               /* # of copy-on-writes */
-			/*uint64_t */
-			public ulong lookups;              /* object cache lookups */
-			/*uint64_t */
-			public ulong hits;                 /* object cache hits */
-			/*uint64_t */
-			public ulong purges;                   /* # of pages purged */
-			/*natural_t*/
-			public int purgeable_count;     /* # of pages purgeable */
-			/*
-          	 * NB: speculative pages are already accounted for in "free_count",
-          	 * so "speculative_count" is the number of "free" pages that are
-          	 * used to hold data that was read speculatively from disk but
-          	 * haven't actually been used by anyone so far.
-          	 */
-			/*natural_t*/
-			public int speculative_count;       /* # of pages speculative */
-
-			/* added for rev1 */
-			/*uint64_t */
-			public ulong decompressions;           /* # of pages decompressed */
-			/*uint64_t */
-			public ulong compressions;         /* # of pages compressed */
-			/*uint64_t */
-			public ulong swapins;              /* # of pages swapped in (via compression segments) */
-			/*uint64_t */
-			public ulong swapouts;             /* # of pages swapped out (via compression segments) */
-			/*natural_t*/
-			public int compressor_page_count;   /* # of pages used by the compressed pager to hold all the compressed data */
-			/*natural_t*/
-			public int throttled_count;     /* # of pages throttled */
-			/*natural_t*/
-			public int external_page_count; /* # of pages that are file-backed (non-swap) */
-			/*natural_t*/
-			public int internal_page_count; /* # of pages that are anonymous */
-			/*uint64_t */
-			public ulong total_uncompressed_pages_in_compressor; /* # of pages (uncompressed) held within the compressor. */
-		} // __attribute__((aligned(8))); 
-
-		// kern_return_t host_statistics64(host_t host_priv, host_flavor_t flavor, host_info64_t host_info64_out, mach_msg_type_number_t *host_info64_outCnt); // from <mach/mach_host.h>
-		[DllImport("libc")]
-		static extern int host_statistics64(IntPtr host_priv, int flavor, out vm_statistics64 host_info64_out, ref uint host_info_count);
-
-		// mach_port_t mach_host_self() // from <mach/mach_init.h>
-		[DllImport("libc")]
-		static extern IntPtr mach_host_self();
+		[Obsolete("Deprecated in UE5.8. Use SystemUtils.GetAvailableMemoryBytes() from EpicGames.Core instead.")]
+		public static long GetAvailableMemoryBytes() => SystemUtils.GetAvailableMemoryBytes();
 
 		/// <summary>
 		/// Gets the total system memory in bytes based on what is known to the garbage collector.
 		/// </summary>
 		/// <remarks>This will return a max of 2GB for a 32bit application.</remarks>
-		/// <returns>The total system memory free, in bytes.</returns>
-		public static long GetTotalSystemMemoryBytes()
-		{
-			GCMemoryInfo MemoryInfo = GC.GetGCMemoryInfo();
-			return MemoryInfo.TotalAvailableMemoryBytes;
-		}
+		/// <returns>The total system memory, in bytes.</returns>
+		[Obsolete("Deprecated in UE5.8. Use SystemUtils.GetTotalSystemMemoryBytes() from EpicGames.Core instead.")]
+		public static long GetTotalSystemMemoryBytes() => SystemUtils.GetTotalSystemMemoryBytes();
 
 		/// <summary>
 		/// Gets the total memory bytes free, based on what is known to the garbage collector.
 		/// </summary>
 		/// <remarks>This will return a max of 2GB for a 32bit application.</remarks>
 		/// <returns>The total memory free, in bytes.</returns>
-		public static long GetFreeMemoryBytes()
-		{
-			long FreeMemoryBytes = -1;
-
-			GCMemoryInfo MemoryInfo = GC.GetGCMemoryInfo();
-			// TotalAvailableMemoryBytes will be 0 if garbage collection has not run yet
-			if (MemoryInfo.TotalAvailableMemoryBytes != 0)
-			{
-				FreeMemoryBytes = MemoryInfo.TotalAvailableMemoryBytes - MemoryInfo.MemoryLoadBytes;
-			}
-
-			// On Mac, MemoryInfo.MemoryLoadBytes includes memory used to cache disk-backed files ("Cached Files" in
-			// Activity Monitor), which can result in a significant over-estimate of memory pressure.
-			// We treat memory used for caching of disk-backed files as free for use in compilation tasks.
-			if (OperatingSystem.IsMacOS())
-			{
-				// host_statistics64() flavor, from <mach/host_info.h>
-				int HOST_VM_INFO64 = 4;
-				// host_statistics64() count of 32bit values in output struct, from <mach/host_info.h>
-				int HOST_VM_INFO64_COUNT = Marshal.SizeOf(typeof(vm_statistics64)) / 4;
-
-				vm_statistics64 VMStats;
-				uint StructSize = (uint)HOST_VM_INFO64_COUNT;
-				IntPtr Host = mach_host_self();
-				if (host_statistics64(Host, HOST_VM_INFO64, out VMStats, ref StructSize) != 0)
-				{
-					// Return the dotnet memory estimate if the OS call fails
-					return FreeMemoryBytes;
-				}
-
-				ulong OutSize = 4;
-				if (0 != sysctlbyname("hw.pagesize", out int PageSize, ref OutSize, IntPtr.Zero, 0))
-				{
-					PageSize = 4096; // likely result
-				}
-
-				FreeMemoryBytes += (long)PageSize * (long)VMStats.external_page_count;
-			}
-			return FreeMemoryBytes;
-		}
-
-		[DllImport("pdh.dll", SetLastError = true, CharSet = CharSet.Auto)]
-		static extern int PdhOpenQueryW([MarshalAs(UnmanagedType.LPWStr)] string? szDataSource, UIntPtr dwUserData, out IntPtr phQuery);
-
-		[DllImport("pdh.dll", SetLastError = true, CharSet = CharSet.Auto)]
-		static extern uint PdhAddCounter(IntPtr hQuery, string szFullCounterPath, IntPtr dwUserData, out IntPtr phCounter);
-
-		[DllImport("pdh.dll", SetLastError = true)]
-		static extern uint PdhCollectQueryData(IntPtr phQuery);
-
-		struct PDH_FMT_COUNTERVALUE
-		{
-			public uint CStatus;
-			public double doubleValue;
-		};
-
-		[DllImport("pdh.dll", SetLastError = true)]
-		static extern uint PdhGetFormattedCounterValue(IntPtr phCounter, uint dwFormat, IntPtr lpdwType, out PDH_FMT_COUNTERVALUE pValue);
-
-		static IntPtr CpuQuery;
-		static IntPtr CpuTotal;
-		static bool CpuInitialized = false;
+		[Obsolete("Deprecated in UE5.8. Use SystemUtils.GetFreeMemoryBytes() from EpicGames.Core instead.")]
+		public static long GetFreeMemoryBytes() => SystemUtils.GetFreeMemoryBytes();
 
 		/// <summary>
 		/// Gives the current CPU utilization.
 		/// </summary>
 		/// <param name="Utilization">Percentage of CPU utilization currently.</param>
 		/// <returns>Whether or not it was successful in getting the CPU utilization.</returns>
-		public static bool GetTotalCpuUtilization(out float Utilization)
-		{
-			Utilization = 0.0f;
-			if (OperatingSystem.IsWindows())
-			{
-				const uint ERROR_SUCCESS = 0;
-				if (!CpuInitialized)
-				{
-					if (PdhOpenQueryW(null, UIntPtr.Zero, out CpuQuery) == ERROR_SUCCESS)
-					{
-						if (PdhAddCounter(CpuQuery, "\\Processor(_Total)\\% Processor Time", IntPtr.Zero, out CpuTotal) != ERROR_SUCCESS)
-						{
-							return false;
-						}
-
-						if (PdhCollectQueryData(CpuQuery) != ERROR_SUCCESS)
-						{
-							return false;
-						}
-
-						CpuInitialized = true;
-					}
-				}
-
-				if (CpuInitialized)
-				{
-					const uint PDH_FMT_DOUBLE = 0x00000200;
-
-					PDH_FMT_COUNTERVALUE counterVal;
-					if (PdhCollectQueryData(CpuQuery) != ERROR_SUCCESS)
-					{
-						return false;
-					}
-
-					if (PdhGetFormattedCounterValue(CpuTotal, PDH_FMT_DOUBLE, IntPtr.Zero, out counterVal) != ERROR_SUCCESS)
-					{
-						return false;
-					}
-
-					Utilization = (float)counterVal.doubleValue;
-					return true;
-				}
-			}
-
-			return false;
-		}
+		[Obsolete("Deprecated in UE5.8. Use SystemUtils.GetTotalCpuUtilization() from EpicGames.Core instead.")]
+		public static bool GetTotalCpuUtilization(out float Utilization) => SystemUtils.GetTotalCpuUtilization(out Utilization);
 
 		/// <summary>
 		/// Determines the maximum number of actions to execute in parallel, taking into account the resources available on this machine.
@@ -1350,10 +1025,10 @@ namespace UnrealBuildTool
 			}
 
 			// Get the number of logical processors
-			int NumLogicalCores = Utils.GetLogicalProcessorCount();
+			int NumLogicalCores = SystemUtils.GetLogicalProcessorCount();
 
 			// Use WMI to figure out physical cores, excluding hyper threading.
-			int NumPhysicalCores = Utils.GetPhysicalProcessorCount();
+			int NumPhysicalCores = SystemUtils.GetPhysicalProcessorCount();
 			if (NumPhysicalCores == -1)
 			{
 				NumPhysicalCores = NumLogicalCores;
@@ -1394,7 +1069,7 @@ namespace UnrealBuildTool
 				// limiting based on total system memory and relying on the OS to swap out other processes to make room is fine.
 				// That said only Mac has been extensively tested here in this scenario, and does tend to have faster disk access
 				// than other platforms. So for now we'll change Mac and leave Windows/Linux using the old behavior
-				long AvailableMemoryBytes = OperatingSystem.IsMacOS() ? GetTotalSystemMemoryBytes() : GetFreeMemoryBytes();
+				long AvailableMemoryBytes = OperatingSystem.IsMacOS() ? SystemUtils.GetTotalSystemMemoryBytes() : SystemUtils.GetFreeMemoryBytes();
 
 				if (AvailableMemoryBytes != -1)
 				{
