UnrealBuildTool: Change how we retrieve mode properties to use static interface methods, so that if the IToolMode<> compiles, it's a valid mode and can't throw at runtime due to missing an attribute.

#Synced-CL 52340182

diff --git a/Engine/Source/Programs/UnrealBuildTool/Actions/ActionGraphBuilder.cs b/Engine/Source/Programs/UnrealBuildTool/Actions/ActionGraphBuilder.cs
--- a/Engine/Source/Programs/UnrealBuildTool/Actions/ActionGraphBuilder.cs
+++ b/Engine/Source/Programs/UnrealBuildTool/Actions/ActionGraphBuilder.cs
@@ -2,7 +2,6 @@
 
 using System;
 using System.Collections.Generic;
-using System.Reflection;
 using EpicGames.Core;
 using Microsoft.Extensions.Logging;
 using UnrealBuildBase;
@@ -318,19 +317,13 @@ namespace UnrealBuildTool
 		/// <param name="Type">Type of the action</param>
 		/// <param name="Arguments">Arguments for the action</param>
 		/// <returns>New action instance</returns>
-		public static Action CreateRecursiveAction<T>(this IActionGraphBuilder Graph, ActionType Type, string Arguments) where T : ToolMode
+		public static Action CreateRecursiveAction<T>(this IActionGraphBuilder Graph, ActionType Type, string Arguments) where T : IToolMode<T>, new()
 		{
-			ToolModeAttribute? Attribute = typeof(T).GetCustomAttribute<ToolModeAttribute>();
-			if (Attribute == null)
-			{
-				throw new BuildException("Missing ToolModeAttribute on {0}", typeof(T).Name);
-			}
-
 			Action NewAction = Graph.CreateAction(Type);
 			NewAction.CommandPath = Unreal.DotnetPath;
 			NewAction.WorkingDirectory = Unreal.EngineSourceDirectory;
-			NewAction.CommandArguments = $"\"{Unreal.UnrealBuildToolDllPath}\" -Session=\"{UnrealBuildTool.SessionIdentifier}\" -Mode={Attribute.Name} {Arguments}";
-			NewAction.CommandDescription = Attribute.Name;
+			NewAction.CommandArguments = $"\"{Unreal.UnrealBuildToolDllPath}\" -Session=\"{UnrealBuildTool.SessionIdentifier}\" -Mode={T.Name} {Arguments}";
+			NewAction.CommandDescription = T.Name;
 			NewAction.bCanExecuteRemotely = false;
 			NewAction.bCanExecuteRemotelyWithSNDBS = false;
 			NewAction.bCanExecuteInUBA = false;
diff --git a/Engine/Source/Programs/UnrealBuildTool/Modes/AggregateClangTimingInfo.cs b/Engine/Source/Programs/UnrealBuildTool/Modes/AggregateClangTimingInfo.cs
--- a/Engine/Source/Programs/UnrealBuildTool/Modes/AggregateClangTimingInfo.cs
+++ b/Engine/Source/Programs/UnrealBuildTool/Modes/AggregateClangTimingInfo.cs
@@ -16,9 +16,11 @@ namespace UnrealBuildTool
 	/// <summary>
 	/// Aggregates clang timing information files together into one monolithic breakdown file.
 	/// </summary>
-	[ToolMode("AggregateClangTimingInfo", ToolModeOptions.None)]
-	class AggregateClangTimingInfo : ToolMode
+	internal sealed class AggregateClangTimingInfo : IToolMode<AggregateClangTimingInfo>
 	{
+		public static string Name => "AggregateClangTimingInfo";
+		public static ToolModeOptions Options => ToolModeOptions.None;
+
 		public class ClangTrace
 		{
 			public class TraceEvent
@@ -169,7 +171,7 @@ namespace UnrealBuildTool
 
 		private readonly ConcurrentDictionary<FileReference, ClangTrace> ClangTraceCache = new();
 
-		public override async Task<int> ExecuteAsync(CommandLineArguments Arguments, ILogger Logger)
+		public async Task<int> ExecuteAsync(CommandLineArguments Arguments, ILogger Logger)
 		{
 			FileReference ManifestFile = Arguments.GetFileReference("-ManifestFile=");
 			IEnumerable<FileReference> SourceFiles = (await FileReference.ReadAllLinesAsync(ManifestFile)).Select(x => new FileReference(x));
diff --git a/Engine/Source/Programs/UnrealBuildTool/Modes/AggregateParsedTimingInfo.cs b/Engine/Source/Programs/UnrealBuildTool/Modes/AggregateParsedTimingInfo.cs
--- a/Engine/Source/Programs/UnrealBuildTool/Modes/AggregateParsedTimingInfo.cs
+++ b/Engine/Source/Programs/UnrealBuildTool/Modes/AggregateParsedTimingInfo.cs
@@ -18,9 +18,11 @@ namespace UnrealBuildTool
 	/// <summary>
 	/// Aggregates parsed Visual C++ timing information files together into one monolithic breakdown file.
 	/// </summary>
-	[ToolMode("AggregateParsedTimingInfo", ToolModeOptions.None)]
-	class AggregateParsedTimingInfo : ToolMode
+	internal sealed class AggregateParsedTimingInfo : IToolMode<AggregateParsedTimingInfo>
 	{
+		public static string Name => "AggregateParsedTimingInfo";
+		public static ToolModeOptions Options => ToolModeOptions.None;
+
 		private TimingData? FileTimingData;
 		private readonly ConcurrentDictionary<string, int> DecompressedFileSizes = new ConcurrentDictionary<string, int>();
 		private readonly ConcurrentDictionary<string, byte[]> CompressedFiles = new ConcurrentDictionary<string, byte[]>();
@@ -29,7 +31,7 @@ namespace UnrealBuildTool
 		private readonly ConcurrentBag<TimingData> AggregateFunctions = new ConcurrentBag<TimingData>();
 		private readonly object AddFileLock = new object();
 
-		public override async Task<int> ExecuteAsync(CommandLineArguments Arguments, ILogger Logger)
+		public async Task<int> ExecuteAsync(CommandLineArguments Arguments, ILogger Logger)
 		{
 			FileReference ManifestFile = Arguments.GetFileReference("-ManifestFile=");
 			string[] ParsedFileNames = await FileReference.ReadAllLinesAsync(ManifestFile);
diff --git a/Engine/Source/Programs/UnrealBuildTool/Modes/AnalyzeMode.cs b/Engine/Source/Programs/UnrealBuildTool/Modes/AnalyzeMode.cs
--- a/Engine/Source/Programs/UnrealBuildTool/Modes/AnalyzeMode.cs
+++ b/Engine/Source/Programs/UnrealBuildTool/Modes/AnalyzeMode.cs
@@ -14,11 +14,13 @@ namespace UnrealBuildTool.Modes
 	/// <summary>
 	/// Outputs information about the given target, including a module dependecy graph (in .gefx format and list of module references)
 	/// </summary>
-	[ToolMode("Analyze", ToolModeOptions.XmlConfig | ToolModeOptions.BuildPlatforms | ToolModeOptions.SingleInstance | ToolModeOptions.StartPrefetchingEngine | ToolModeOptions.ShowExecutionTime)]
-	internal class AnalyzeMode : ToolMode
+	internal sealed class AnalyzeMode : IToolMode<AnalyzeMode>
 	{
+		public static string Name => "Analyze";
+		public static ToolModeOptions Options => ToolModeOptions.XmlConfig | ToolModeOptions.BuildPlatforms | ToolModeOptions.SingleInstance | ToolModeOptions.StartPrefetchingEngine | ToolModeOptions.ShowExecutionTime;
+
 		/// <inheritdoc/>
-		public override Task<int> ExecuteAsync(CommandLineArguments arguments, ILogger logger)
+		public Task<int> ExecuteAsync(CommandLineArguments arguments, ILogger logger)
 		{
 			arguments.ApplyTo(this);
 
diff --git a/Engine/Source/Programs/UnrealBuildTool/Modes/BuildMode.cs b/Engine/Source/Programs/UnrealBuildTool/Modes/BuildMode.cs
--- a/Engine/Source/Programs/UnrealBuildTool/Modes/BuildMode.cs
+++ b/Engine/Source/Programs/UnrealBuildTool/Modes/BuildMode.cs
@@ -45,9 +45,11 @@ namespace UnrealBuildTool
 	/// <summary>
 	/// Builds a target
 	/// </summary>
-	[ToolMode("Build", ToolModeOptions.XmlConfig | ToolModeOptions.BuildPlatforms | ToolModeOptions.SingleInstance | ToolModeOptions.StartPrefetchingEngine | ToolModeOptions.ShowExecutionTime | ToolModeOptions.UseStartupTraceListener)]
-	class BuildMode : ToolMode
+	internal sealed class BuildMode : IToolMode<BuildMode>
 	{
+		public static string Name => "Build";
+		public static ToolModeOptions Options => ToolModeOptions.XmlConfig | ToolModeOptions.BuildPlatforms | ToolModeOptions.SingleInstance | ToolModeOptions.StartPrefetchingEngine | ToolModeOptions.ShowExecutionTime | ToolModeOptions.UseStartupTraceListener;
+
 		/// <summary>
 		/// Specifies the file to use for logging.
 		/// </summary>
@@ -109,7 +111,7 @@ namespace UnrealBuildTool
 		/// <param name="Arguments">Command-line arguments</param>
 		/// <returns>One of the values of ECompilationResult</returns>
 		/// <param name="Logger"></param>
-		public override async Task<int> ExecuteAsync(CommandLineArguments Arguments, ILogger Logger)
+		public async Task<int> ExecuteAsync(CommandLineArguments Arguments, ILogger Logger)
 		{
 			Arguments.ApplyTo(this);
 
diff --git a/Engine/Source/Programs/UnrealBuildTool/Modes/ClReproMode.cs b/Engine/Source/Programs/UnrealBuildTool/Modes/ClReproMode.cs
--- a/Engine/Source/Programs/UnrealBuildTool/Modes/ClReproMode.cs
+++ b/Engine/Source/Programs/UnrealBuildTool/Modes/ClReproMode.cs
@@ -12,9 +12,11 @@ using UnrealBuildBase;
 
 namespace UnrealBuildTool.Modes
 {
-	[ToolMode("ClRepro", ToolModeOptions.UseStartupTraceListener)]
-	class ClReproMode : ToolMode
+	internal sealed class ClReproMode : IToolMode<ClReproMode>
 	{
+		public static string Name => "ClRepro";
+		public static ToolModeOptions Options => ToolModeOptions.UseStartupTraceListener;
+
 		[CommandLine("-InputCpp=", Required = true)]
 		public FileReference? InputCpp { get; set; } = null;
 
@@ -151,7 +153,7 @@ namespace UnrealBuildTool.Modes
 			return 0;
 		}
 
-		public override Task<int> ExecuteAsync(CommandLineArguments arguments, ILogger logger)
+		public Task<int> ExecuteAsync(CommandLineArguments arguments, ILogger logger)
 		{
 			arguments.ApplyTo(this);
 
diff --git a/Engine/Source/Programs/UnrealBuildTool/Modes/CleanMode.cs b/Engine/Source/Programs/UnrealBuildTool/Modes/CleanMode.cs
--- a/Engine/Source/Programs/UnrealBuildTool/Modes/CleanMode.cs
+++ b/Engine/Source/Programs/UnrealBuildTool/Modes/CleanMode.cs
@@ -18,9 +18,11 @@ namespace UnrealBuildTool
 	/// Cleans build products and intermediates for the target. This deletes files which are named consistently with the target being built
 	/// (e.g. UnrealEditor-Foo-Win64-Debug.dll) rather than an actual record of previous build products.
 	/// </summary>
-	[ToolMode("Clean", ToolModeOptions.XmlConfig | ToolModeOptions.BuildPlatforms | ToolModeOptions.SingleInstance)]
-	class CleanMode : ToolMode
+	internal sealed class CleanMode : IToolMode<CleanMode>
 	{
+		public static string Name => "Clean";
+		public static ToolModeOptions Options => ToolModeOptions.XmlConfig | ToolModeOptions.BuildPlatforms | ToolModeOptions.SingleInstance;
+
 #pragma warning disable IDE0044 // Make field readonly - these private static fields are set by command-line parsing.
 		/// <summary>
 		/// Whether to avoid cleaning targets
@@ -47,7 +49,7 @@ namespace UnrealBuildTool
 		/// <param name="Arguments">Command-line arguments</param>
 		/// <returns>One of the values of ECompilationResult</returns>
 		/// <param name="Logger"></param>
-		public override Task<int> ExecuteAsync(CommandLineArguments Arguments, ILogger Logger)
+		public Task<int> ExecuteAsync(CommandLineArguments Arguments, ILogger Logger)
 		{
 			Arguments.ApplyTo(this);
 
diff --git a/Engine/Source/Programs/UnrealBuildTool/Modes/DeployMode.cs b/Engine/Source/Programs/UnrealBuildTool/Modes/DeployMode.cs
--- a/Engine/Source/Programs/UnrealBuildTool/Modes/DeployMode.cs
+++ b/Engine/Source/Programs/UnrealBuildTool/Modes/DeployMode.cs
@@ -9,9 +9,11 @@ namespace UnrealBuildTool
 	/// <summary>
 	/// Invokes the deployment handler for a target.
 	/// </summary>
-	[ToolMode("Deploy", ToolModeOptions.BuildPlatforms)]
-	class DeployMode : ToolMode
+	internal sealed class DeployMode : IToolMode<DeployMode>
 	{
+		public static string Name => "Deploy";
+		public static ToolModeOptions Options => ToolModeOptions.BuildPlatforms;
+
 		/// <summary>
 		/// If we are just running the deployment step, specifies the path to the given deployment settings
 		/// </summary>
@@ -24,7 +26,7 @@ namespace UnrealBuildTool
 		/// <param name="Arguments">Command line arguments</param>
 		/// <returns>Exit code</returns>
 		/// <param name="Logger"></param>
-		public override Task<int> ExecuteAsync(CommandLineArguments Arguments, ILogger Logger)
+		public Task<int> ExecuteAsync(CommandLineArguments Arguments, ILogger Logger)
 		{
 			// Apply the arguments
 			Arguments.ApplyTo(this);
diff --git a/Engine/Source/Programs/UnrealBuildTool/Modes/ExecuteMode.cs b/Engine/Source/Programs/UnrealBuildTool/Modes/ExecuteMode.cs
--- a/Engine/Source/Programs/UnrealBuildTool/Modes/ExecuteMode.cs
+++ b/Engine/Source/Programs/UnrealBuildTool/Modes/ExecuteMode.cs
@@ -11,9 +11,11 @@ namespace UnrealBuildTool
 	/// <summary>
 	/// Builds a target
 	/// </summary>
-	[ToolMode("Execute", ToolModeOptions.XmlConfig | ToolModeOptions.BuildPlatforms | ToolModeOptions.SingleInstance | ToolModeOptions.StartPrefetchingEngine | ToolModeOptions.ShowExecutionTime)]
-	class ExecuteMode : ToolMode
+	internal sealed class ExecuteMode : IToolMode<ExecuteMode>
 	{
+		public static string Name => "Execute";
+		public static ToolModeOptions Options => ToolModeOptions.XmlConfig | ToolModeOptions.BuildPlatforms | ToolModeOptions.SingleInstance | ToolModeOptions.StartPrefetchingEngine | ToolModeOptions.ShowExecutionTime;
+
 		/// <summary>
 		/// Whether we should just export the outdated actions list
 		/// </summary>
@@ -26,7 +28,7 @@ namespace UnrealBuildTool
 		/// <param name="Arguments">Command-line arguments</param>
 		/// <returns>One of the values of ECompilationResult</returns>
 		/// <param name="Logger"></param>
-		public override async Task<int> ExecuteAsync(CommandLineArguments Arguments, ILogger Logger)
+		public async Task<int> ExecuteAsync(CommandLineArguments Arguments, ILogger Logger)
 		{
 			Arguments.ApplyTo(this);
 
diff --git a/Engine/Source/Programs/UnrealBuildTool/Modes/FixIncludePathsMode.cs b/Engine/Source/Programs/UnrealBuildTool/Modes/FixIncludePathsMode.cs
--- a/Engine/Source/Programs/UnrealBuildTool/Modes/FixIncludePathsMode.cs
+++ b/Engine/Source/Programs/UnrealBuildTool/Modes/FixIncludePathsMode.cs
@@ -15,9 +15,11 @@ namespace UnrealBuildTool
 	/// <summary>
 	/// Fixes the include paths found in a header and source file
 	/// </summary>
-	[ToolMode("FixIncludePaths", ToolModeOptions.XmlConfig | ToolModeOptions.BuildPlatforms | ToolModeOptions.SingleInstance | ToolModeOptions.StartPrefetchingEngine | ToolModeOptions.ShowExecutionTime)]
-	class FixIncludePathsMode : ToolMode
+	internal sealed class FixIncludePathsMode : IToolMode<FixIncludePathsMode>
 	{
+		public static string Name => "FixIncludePaths";
+		public static ToolModeOptions Options => ToolModeOptions.XmlConfig | ToolModeOptions.BuildPlatforms | ToolModeOptions.SingleInstance | ToolModeOptions.StartPrefetchingEngine | ToolModeOptions.ShowExecutionTime;
+
 		/// <summary>
 		/// Regex that matches #include statements.
 		/// </summary>
@@ -53,7 +55,7 @@ namespace UnrealBuildTool
 		/// <param name="Arguments">Command line arguments</param>
 		/// <returns>Exit code</returns>
 		/// <param name="Logger"></param>
-		public override async Task<int> ExecuteAsync(CommandLineArguments Arguments, ILogger Logger)
+		public async Task<int> ExecuteAsync(CommandLineArguments Arguments, ILogger Logger)
 		{
 			Arguments.ApplyTo(this);
 
diff --git a/Engine/Source/Programs/UnrealBuildTool/Modes/GenerateClangDatabase.cs b/Engine/Source/Programs/UnrealBuildTool/Modes/GenerateClangDatabase.cs
--- a/Engine/Source/Programs/UnrealBuildTool/Modes/GenerateClangDatabase.cs
+++ b/Engine/Source/Programs/UnrealBuildTool/Modes/GenerateClangDatabase.cs
@@ -14,9 +14,11 @@ namespace UnrealBuildTool
 	/// <summary>
 	/// Generate a clang compile_commands file for a target
 	/// </summary>
-	[ToolMode("GenerateClangDatabase", ToolModeOptions.XmlConfig | ToolModeOptions.BuildPlatforms | ToolModeOptions.SingleInstance | ToolModeOptions.StartPrefetchingEngine | ToolModeOptions.ShowExecutionTime)]
-	class GenerateClangDatabase : ToolMode
+	internal sealed class GenerateClangDatabase : IToolMode<GenerateClangDatabase>
 	{
+		public static string Name => "GenerateClangDatabase";
+		public static ToolModeOptions Options => ToolModeOptions.XmlConfig | ToolModeOptions.BuildPlatforms | ToolModeOptions.SingleInstance | ToolModeOptions.StartPrefetchingEngine | ToolModeOptions.ShowExecutionTime;
+
 #pragma warning disable IDE0044 // Make field readonly - these private static fields are set by command-line parsing.
 		/// <summary>
 		/// Optional set of source files to include in the compile database. If this is empty, all files will be included by default and -Exclude can be used to exclude some.
@@ -58,7 +60,7 @@ namespace UnrealBuildTool
 		/// <param name="Arguments">Command line arguments</param>
 		/// <returns>Exit code</returns>
 		/// <param name="Logger"></param>
-		public override async Task<int> ExecuteAsync(CommandLineArguments Arguments, ILogger Logger)
+		public async Task<int> ExecuteAsync(CommandLineArguments Arguments, ILogger Logger)
 		{
 			Arguments.ApplyTo(this);
 
diff --git a/Engine/Source/Programs/UnrealBuildTool/Modes/GenerateProjectFilesMode.cs b/Engine/Source/Programs/UnrealBuildTool/Modes/GenerateProjectFilesMode.cs
--- a/Engine/Source/Programs/UnrealBuildTool/Modes/GenerateProjectFilesMode.cs
+++ b/Engine/Source/Programs/UnrealBuildTool/Modes/GenerateProjectFilesMode.cs
@@ -15,9 +15,11 @@ namespace UnrealBuildTool
 	/// <summary>
 	/// Generates project files for one or more projects
 	/// </summary>
-	[ToolMode("GenerateProjectFiles", ToolModeOptions.XmlConfig | ToolModeOptions.BuildPlatforms | ToolModeOptions.SingleInstance | ToolModeOptions.UseStartupTraceListener | ToolModeOptions.StartPrefetchingEngine | ToolModeOptions.ShowExecutionTime)]
-	class GenerateProjectFilesMode : ToolMode
+	internal sealed class GenerateProjectFilesMode : IToolMode<GenerateProjectFilesMode>
 	{
+		public static string Name => "GenerateProjectFiles";
+		public static ToolModeOptions Options => ToolModeOptions.XmlConfig | ToolModeOptions.BuildPlatforms | ToolModeOptions.SingleInstance | ToolModeOptions.UseStartupTraceListener | ToolModeOptions.StartPrefetchingEngine | ToolModeOptions.ShowExecutionTime;
+
 #pragma warning disable IDE0044 // Make field readonly - these private static fields are set by command-line parsing.
 		/// <summary>
 		/// Types of project files to generate
@@ -62,7 +64,7 @@ namespace UnrealBuildTool
 		/// <param name="Arguments">Command line arguments</param>
 		/// <returns>Exit code</returns>
 		/// <param name="Logger"></param>
-		public override Task<int> ExecuteAsync(CommandLineArguments Arguments, ILogger Logger)
+		public Task<int> ExecuteAsync(CommandLineArguments Arguments, ILogger Logger)
 		{
 			// Apply any command line arguments to this class
 			Arguments.ApplyTo(this);
diff --git a/Engine/Source/Programs/UnrealBuildTool/Modes/IToolMode.cs b/Engine/Source/Programs/UnrealBuildTool/Modes/IToolMode.cs
new file mode 100644
--- /dev/null
+++ b/Engine/Source/Programs/UnrealBuildTool/Modes/IToolMode.cs
@@ -0,0 +1,90 @@
+// Copyright Epic Games, Inc. All Rights Reserved.
+
+using System;
+using System.Threading.Tasks;
+using EpicGames.Core;
+using Microsoft.Extensions.Logging;
+
+namespace UnrealBuildTool
+{
+	/// <summary>
+	/// Systems that need to be configured to execute a tool mode
+	/// </summary>
+	[Flags]
+	internal enum ToolModeOptions
+	{
+		/// <summary>
+		/// Do not initialize anything
+		/// </summary>
+		None = 0,
+
+		/// <summary>
+		/// Start prefetching metadata for the engine folder as early as possible
+		/// </summary>
+		StartPrefetchingEngine = 1,
+
+		/// <summary>
+		/// Initializes the XmlConfig system
+		/// </summary>
+		XmlConfig = 2,
+
+		/// <summary>
+		/// Registers build platforms
+		/// </summary>
+		BuildPlatforms = 4,
+
+		/// <summary>
+		/// Registers build platforms
+		/// </summary>
+		BuildPlatformsHostOnly = 8,
+
+		/// <summary>
+		/// Registers build platforms for validation
+		/// </summary>
+		BuildPlatformsForValidation = 16,
+
+		/// <summary>
+		/// Only allow a single instance running in the branch at once
+		/// </summary>
+		SingleInstance = 32,
+
+		/// <summary>
+		/// Print out the total time taken to execute
+		/// </summary>
+		ShowExecutionTime = 64,
+
+		/// <summary>
+		/// Capture logs as early as possible in a StartupTraceListener object
+		/// </summary>
+		UseStartupTraceListener = 128,
+	}
+
+	internal interface IToolMode
+	{
+		/// <summary>
+		/// Entry point for this command.
+		/// </summary>
+		/// <param name="Arguments">List of command line arguments</param>
+		/// <param name="Logger"></param>
+		/// <returns>Exit code for the process</returns>
+		abstract Task<int> ExecuteAsync(CommandLineArguments Arguments, ILogger Logger);
+	}
+
+	/// <summary>
+	/// Base interface for standalone UBT modes. Different modes can be invoked using the -Mode=[Name] argument on the command line.
+	/// The log system will be initialized before calling the mode, but little else.
+	/// </summary>
+	internal interface IToolMode<TToolMode> : IToolMode where TToolMode : IToolMode<TToolMode>, new()
+	{
+
+		/// <summary>
+		/// The name of this tool mode.
+		/// </summary>
+		static abstract string Name { get; }
+
+		/// <summary>
+		/// The options for this tool mode.
+		/// </summary>
+		static abstract ToolModeOptions Options { get; }
+	}
+}
diff --git a/Engine/Source/Programs/UnrealBuildTool/Modes/IWYUMode.cs b/Engine/Source/Programs/UnrealBuildTool/Modes/IWYUMode.cs
--- a/Engine/Source/Programs/UnrealBuildTool/Modes/IWYUMode.cs
+++ b/Engine/Source/Programs/UnrealBuildTool/Modes/IWYUMode.cs
@@ -220,9 +220,11 @@ namespace UnrealBuildTool
 	/// <summary>
 	/// Profiles different unity sizes and prints out the different size and its timings
 	/// </summary>
-	[ToolMode("IWYU", ToolModeOptions.XmlConfig | ToolModeOptions.BuildPlatforms | ToolModeOptions.SingleInstance | ToolModeOptions.StartPrefetchingEngine | ToolModeOptions.ShowExecutionTime)]
-	class IWYUMode : ToolMode
+	internal sealed class IWYUMode : IToolMode<IWYUMode>
 	{
+		public static string Name => "IWYU";
+		public static ToolModeOptions Options => ToolModeOptions.XmlConfig | ToolModeOptions.BuildPlatforms | ToolModeOptions.SingleInstance | ToolModeOptions.StartPrefetchingEngine | ToolModeOptions.ShowExecutionTime;
+
 		/// <summary>
 		/// Specifies the file to use for logging.
 		/// </summary>
@@ -380,7 +382,7 @@ namespace UnrealBuildTool
 		/// <param name="Arguments">Command line arguments</param>
 		/// <returns>Exit code</returns>
 		/// <param name="Logger"></param>
-		public override async Task<int> ExecuteAsync(CommandLineArguments Arguments, ILogger Logger)
+		public async Task<int> ExecuteAsync(CommandLineArguments Arguments, ILogger Logger)
 		{
 			Arguments.ApplyTo(this);
 
@@ -2618,7 +2620,7 @@ namespace UnrealBuildTool
 			{
 				UEBuildModule Module = Kvp.Value;
 				Writer.WriteStartObject(Kvp.Key);
-				Writer.WriteString("Name", Module.Name);
+				Writer.WriteString(nameof(Module.Name), Module.Name);
 				Writer.WriteString("Directory", Module.ModuleDirectory.FullName);
 				Writer.WriteString("IWYUSupport", Module.Rules.IWYUSupport.ToString());
 
diff --git a/Engine/Source/Programs/UnrealBuildTool/Modes/InlineGeneratedCpps.cs b/Engine/Source/Programs/UnrealBuildTool/Modes/InlineGeneratedCpps.cs
--- a/Engine/Source/Programs/UnrealBuildTool/Modes/InlineGeneratedCpps.cs
+++ b/Engine/Source/Programs/UnrealBuildTool/Modes/InlineGeneratedCpps.cs
@@ -14,9 +14,11 @@ namespace UnrealBuildTool
 	/// <summary>
 	/// Generate a clang compile_commands file for a target
 	/// </summary>
-	[ToolMode("InlineGeneratedCpps", ToolModeOptions.XmlConfig | ToolModeOptions.BuildPlatforms | ToolModeOptions.SingleInstance | ToolModeOptions.StartPrefetchingEngine | ToolModeOptions.ShowExecutionTime)]
-	class InlineGeneratedCpps : ToolMode
+	internal sealed class InlineGeneratedCpps : IToolMode<InlineGeneratedCpps>
 	{
+		public static string Name => "InlineGeneratedCpps";
+		public static ToolModeOptions Options => ToolModeOptions.XmlConfig | ToolModeOptions.BuildPlatforms | ToolModeOptions.SingleInstance | ToolModeOptions.StartPrefetchingEngine | ToolModeOptions.ShowExecutionTime;
+
 		/// <summary>
 		/// Regex that matches #include statements.
 		/// </summary>
@@ -54,7 +56,7 @@ namespace UnrealBuildTool
 		/// <param name="Arguments">Command line arguments</param>
 		/// <returns>Exit code</returns>
 		/// <param name="Logger"></param>
-		public override async Task<int> ExecuteAsync(CommandLineArguments Arguments, ILogger Logger)
+		public async Task<int> ExecuteAsync(CommandLineArguments Arguments, ILogger Logger)
 		{
 			Arguments.ApplyTo(this);
 
diff --git a/Engine/Source/Programs/UnrealBuildTool/Modes/JsonExportMode.cs b/Engine/Source/Programs/UnrealBuildTool/Modes/JsonExportMode.cs
--- a/Engine/Source/Programs/UnrealBuildTool/Modes/JsonExportMode.cs
+++ b/Engine/Source/Programs/UnrealBuildTool/Modes/JsonExportMode.cs
@@ -12,9 +12,11 @@ namespace UnrealBuildTool
 	/// <summary>
 	/// Exports a target as a JSON file
 	/// </summary>
-	[ToolMode("JsonExport", ToolModeOptions.XmlConfig | ToolModeOptions.BuildPlatforms | ToolModeOptions.SingleInstance)]
-	class JsonExportMode : ToolMode
+	internal sealed class JsonExportMode : IToolMode<JsonExportMode>
 	{
+		public static string Name => "JsonExport";
+		public static ToolModeOptions Options => ToolModeOptions.XmlConfig | ToolModeOptions.BuildPlatforms | ToolModeOptions.SingleInstance;
+
 		/// <summary>
 		/// Execute any actions which result in code generation (eg. ISPC compilation)
 		/// </summary>
@@ -27,7 +29,7 @@ namespace UnrealBuildTool
 		/// <param name="Arguments">Command line arguments</param>
 		/// <returns>Exit code (always zero)</returns>
 		/// <param name="Logger"></param>
-		public override async Task<int> ExecuteAsync(CommandLineArguments Arguments, ILogger Logger)
+		public async Task<int> ExecuteAsync(CommandLineArguments Arguments, ILogger Logger)
 		{
 			Arguments.ApplyTo(this);
 
diff --git a/Engine/Source/Programs/UnrealBuildTool/Modes/ParseMsvcTimingInfo.cs b/Engine/Source/Programs/UnrealBuildTool/Modes/ParseMsvcTimingInfo.cs
--- a/Engine/Source/Programs/UnrealBuildTool/Modes/ParseMsvcTimingInfo.cs
+++ b/Engine/Source/Programs/UnrealBuildTool/Modes/ParseMsvcTimingInfo.cs
@@ -15,12 +15,14 @@ namespace UnrealBuildTool
 	/// Parses an MSVC timing info file generated from cl-filter to turn it into a form that can be used by other tooling.
 	/// This is implemented as a separate mode to allow it to be done as part of the action graph.
 	/// </summary>
-	[ToolMode("ParseMsvcTimingInfo", ToolModeOptions.None)]
-	class ParseMsvcTimingInfoMode : ToolMode
+	internal sealed class ParseMsvcTimingInfoMode : IToolMode<ParseMsvcTimingInfoMode>
 	{
+		public static string Name => "ParseMsvcTimingInfo";
+		public static ToolModeOptions Options => ToolModeOptions.None;
+
 		const string TimingDataRegex = @"^\t\t(?<Indent>\t*)(?<Name>[^\t]+):\s*(?<Duration>[0-9\.]+)s$";
 
-		public override async Task<int> ExecuteAsync(CommandLineArguments Arguments, ILogger Logger)
+		public async Task<int> ExecuteAsync(CommandLineArguments Arguments, ILogger Logger)
 		{
 			FileReference InputFile = Arguments.GetFileReference("-TimingFile=");
 
diff --git a/Engine/Source/Programs/UnrealBuildTool/Modes/PipInstallMode.cs b/Engine/Source/Programs/UnrealBuildTool/Modes/PipInstallMode.cs
--- a/Engine/Source/Programs/UnrealBuildTool/Modes/PipInstallMode.cs
+++ b/Engine/Source/Programs/UnrealBuildTool/Modes/PipInstallMode.cs
@@ -16,9 +16,11 @@ namespace UnrealBuildTool.Modes
 	/// <summary>
 	/// Identifies plugins with python requirements and attempts to install all dependencies using pip.
 	/// </summary>
-	[ToolMode("PipInstall", ToolModeOptions.XmlConfig | ToolModeOptions.BuildPlatformsHostOnly | ToolModeOptions.SingleInstance | ToolModeOptions.StartPrefetchingEngine | ToolModeOptions.ShowExecutionTime)]
-	class PipInstallMode : ToolMode
+	internal sealed class PipInstallMode : IToolMode<PipInstallMode>
 	{
+		public static string Name => "PipInstall";
+		public static ToolModeOptions Options => ToolModeOptions.XmlConfig | ToolModeOptions.BuildPlatformsHostOnly | ToolModeOptions.SingleInstance | ToolModeOptions.StartPrefetchingEngine | ToolModeOptions.ShowExecutionTime;
+
 		private enum ActionBits : byte
 		{
 			NoOp = 0,
@@ -77,7 +79,7 @@ namespace UnrealBuildTool.Modes
 		/// <param name="Arguments">Command line arguments</param>
 		/// <param name="Logger"></param>
 		/// <returns>Exit code</returns>
-		public override Task<int> ExecuteAsync(CommandLineArguments Arguments, ILogger Logger)
+		public Task<int> ExecuteAsync(CommandLineArguments Arguments, ILogger Logger)
 		{
 			// HACK: This env var must be cleared or it carries into python subprocesses and python sys.executable detection breaking venvs
 			Environment.SetEnvironmentVariable("PYTHONEXECUTABLE", null);
diff --git a/Engine/Source/Programs/UnrealBuildTool/Modes/PrintBuildGraphInfo.cs b/Engine/Source/Programs/UnrealBuildTool/Modes/PrintBuildGraphInfo.cs
--- a/Engine/Source/Programs/UnrealBuildTool/Modes/PrintBuildGraphInfo.cs
+++ b/Engine/Source/Programs/UnrealBuildTool/Modes/PrintBuildGraphInfo.cs
@@ -13,16 +13,18 @@ namespace UnrealBuildTool
 	/// <summary>
 	/// Generate a clang compile_commands file for a target
 	/// </summary>
-	[ToolMode("PrintBuildGraphInfo", ToolModeOptions.XmlConfig | ToolModeOptions.BuildPlatforms | ToolModeOptions.SingleInstance | ToolModeOptions.StartPrefetchingEngine | ToolModeOptions.ShowExecutionTime)]
-	class PrintBuildGraphInfo : ToolMode
+	internal sealed class PrintBuildGraphInfo : IToolMode<PrintBuildGraphInfo>
 	{
+		public static string Name => "PrintBuildGraphInfo";
+		public static ToolModeOptions Options => ToolModeOptions.XmlConfig | ToolModeOptions.BuildPlatforms | ToolModeOptions.SingleInstance | ToolModeOptions.StartPrefetchingEngine | ToolModeOptions.ShowExecutionTime;
+
 		/// <summary>
 		/// Execute the command
 		/// </summary>
 		/// <param name="Arguments">Command line arguments</param>
 		/// <returns>Exit code</returns>
 		/// <param name="Logger"></param>
-		public override Task<int> ExecuteAsync(CommandLineArguments Arguments, ILogger Logger)
+		public Task<int> ExecuteAsync(CommandLineArguments Arguments, ILogger Logger)
 		{
 			Arguments.ApplyTo(this);
 
diff --git a/Engine/Source/Programs/UnrealBuildTool/Modes/ProfileUnitySizesMode.cs b/Engine/Source/Programs/UnrealBuildTool/Modes/ProfileUnitySizesMode.cs
--- a/Engine/Source/Programs/UnrealBuildTool/Modes/ProfileUnitySizesMode.cs
+++ b/Engine/Source/Programs/UnrealBuildTool/Modes/ProfileUnitySizesMode.cs
@@ -14,9 +14,11 @@ namespace UnrealBuildTool
 	/// <summary>
 	/// Profiles different unity sizes and prints out the different size and its timings
 	/// </summary>
-	[ToolMode("ProfileUnitySizes", ToolModeOptions.XmlConfig | ToolModeOptions.BuildPlatforms | ToolModeOptions.SingleInstance | ToolModeOptions.StartPrefetchingEngine | ToolModeOptions.ShowExecutionTime)]
-	class ProfileUnitySizesMode : ToolMode
+	internal sealed class ProfileUnitySizesMode : IToolMode<ProfileUnitySizesMode>
 	{
+		public static string Name => "ProfileUnitySizes";
+		public static ToolModeOptions Options => ToolModeOptions.XmlConfig | ToolModeOptions.BuildPlatforms | ToolModeOptions.SingleInstance | ToolModeOptions.StartPrefetchingEngine | ToolModeOptions.ShowExecutionTime;
+
 #pragma warning disable IDE0044 // Make field readonly - these private static fields are set by command-line parsing.
 		/// <summary>
 		/// Set of filters for files to include in the database. Relative to the root directory, or to the project file.
@@ -112,7 +114,7 @@ namespace UnrealBuildTool
 		/// <param name="Arguments">Command line arguments</param>
 		/// <returns>Exit code</returns>
 		/// <param name="Logger"></param>
-		public override async Task<int> ExecuteAsync(CommandLineArguments Arguments, ILogger Logger)
+		public async Task<int> ExecuteAsync(CommandLineArguments Arguments, ILogger Logger)
 		{
 			Arguments.ApplyTo(this);
 
diff --git a/Engine/Source/Programs/UnrealBuildTool/Modes/QueryMode.cs b/Engine/Source/Programs/UnrealBuildTool/Modes/QueryMode.cs
--- a/Engine/Source/Programs/UnrealBuildTool/Modes/QueryMode.cs
+++ b/Engine/Source/Programs/UnrealBuildTool/Modes/QueryMode.cs
@@ -87,9 +87,11 @@ namespace UnrealBuildTool
 		public List<string> Configurations { get; set; } = [];
 	}
 
-	[ToolMode("Query", ToolModeOptions.BuildPlatforms | ToolModeOptions.XmlConfig | ToolModeOptions.UseStartupTraceListener)]
-	class QueryMode : ToolMode
+	internal sealed class QueryMode : IToolMode<QueryMode>
 	{
+		public static string Name => "Query";
+		public static ToolModeOptions Options => ToolModeOptions.BuildPlatforms | ToolModeOptions.XmlConfig | ToolModeOptions.UseStartupTraceListener;
+
 		[CommandLine("-LogDirectory=")]
 		public DirectoryReference? LogDirectory = null;
 
@@ -122,7 +124,7 @@ namespace UnrealBuildTool
 
 		private readonly BuildConfiguration BuildConfiguration = new();
 
-		public override Task<int> ExecuteAsync(CommandLineArguments Arguments, ILogger Logger)
+		public Task<int> ExecuteAsync(CommandLineArguments Arguments, ILogger Logger)
 		{
 			Arguments.ApplyTo(this);
 
diff --git a/Engine/Source/Programs/UnrealBuildTool/Modes/QueryTargetsMode.cs b/Engine/Source/Programs/UnrealBuildTool/Modes/QueryTargetsMode.cs
--- a/Engine/Source/Programs/UnrealBuildTool/Modes/QueryTargetsMode.cs
+++ b/Engine/Source/Programs/UnrealBuildTool/Modes/QueryTargetsMode.cs
@@ -14,9 +14,11 @@ namespace UnrealBuildTool
 	/// <summary>
 	/// Queries information about the targets supported by a project
 	/// </summary>
-	[ToolMode("QueryTargets", ToolModeOptions.XmlConfig | ToolModeOptions.BuildPlatformsHostOnly | ToolModeOptions.SingleInstance)]
-	class QueryTargetsMode : ToolMode
+	internal sealed class QueryTargetsMode : IToolMode<QueryTargetsMode>
 	{
+		public static string Name => "QueryTargets";
+		public static ToolModeOptions Options => ToolModeOptions.XmlConfig | ToolModeOptions.BuildPlatformsHostOnly | ToolModeOptions.SingleInstance;
+
 #pragma warning disable IDE0044 // Make field readonly - these private static fields are set by command-line parsing.
 		/// <summary>
 		/// Path to the project file to query
@@ -56,7 +58,7 @@ namespace UnrealBuildTool
 		/// <param name="arguments">Command line arguments</param>
 		/// <returns></returns>
 		/// <param name="logger"></param>
-		public override Task<int> ExecuteAsync(CommandLineArguments arguments, ILogger logger)
+		public Task<int> ExecuteAsync(CommandLineArguments arguments, ILogger logger)
 		{
 			arguments.ApplyTo(this);
 
diff --git a/Engine/Source/Programs/UnrealBuildTool/Modes/ReportBuildResultsMode.cs b/Engine/Source/Programs/UnrealBuildTool/Modes/ReportBuildResultsMode.cs
--- a/Engine/Source/Programs/UnrealBuildTool/Modes/ReportBuildResultsMode.cs
+++ b/Engine/Source/Programs/UnrealBuildTool/Modes/ReportBuildResultsMode.cs
@@ -17,10 +17,12 @@ using UnrealBuildTool.Storage.Impl;
 
 namespace UnrealBuildTool.Modes;
 
-[ToolMode("ReportBuildResults", ToolModeOptions.XmlConfig | ToolModeOptions.BuildPlatforms)]
-internal class ReportBuildResultsMode : ToolMode
+internal sealed class ReportBuildResultsMode : IToolMode<ReportBuildResultsMode>
 {
-	public override async Task<int> ExecuteAsync(CommandLineArguments arguments, ILogger logger)
+	public static string Name => "ReportBuildResults";
+	public static ToolModeOptions Options => ToolModeOptions.XmlConfig | ToolModeOptions.BuildPlatforms;
+
+	public async Task<int> ExecuteAsync(CommandLineArguments arguments, ILogger logger)
 	{
 		BuildConfiguration buildConfiguration = new();
 		XmlConfig.ApplyTo(buildConfiguration);
diff --git a/Engine/Source/Programs/UnrealBuildTool/Modes/SetupPlatformsMode.cs b/Engine/Source/Programs/UnrealBuildTool/Modes/SetupPlatformsMode.cs
--- a/Engine/Source/Programs/UnrealBuildTool/Modes/SetupPlatformsMode.cs
+++ b/Engine/Source/Programs/UnrealBuildTool/Modes/SetupPlatformsMode.cs
@@ -9,16 +9,18 @@ namespace UnrealBuildTool
 	/// <summary>
 	/// Register all platforms (and in the process, configure all autosdks)
 	/// </summary>
-	[ToolMode("SetupPlatforms", ToolModeOptions.XmlConfig | ToolModeOptions.BuildPlatforms | ToolModeOptions.SingleInstance)]
-	class SetupPlatforms : ToolMode
+	internal sealed class SetupPlatforms : IToolMode<SetupPlatforms>
 	{
+		public static string Name => "SetupPlatforms";
+		public static ToolModeOptions Options => ToolModeOptions.XmlConfig | ToolModeOptions.BuildPlatforms | ToolModeOptions.SingleInstance;
+
 		/// <summary>
 		/// Execute the tool mode
 		/// </summary>
 		/// <param name="Arguments">Command line arguments</param>
 		/// <returns>Exit code</returns>
 		/// <param name="Logger"></param>
-		public override Task<int> ExecuteAsync(CommandLineArguments Arguments, ILogger Logger)
+		public Task<int> ExecuteAsync(CommandLineArguments Arguments, ILogger Logger)
 		{
 			return Task.FromResult(0);
 		}
diff --git a/Engine/Source/Programs/UnrealBuildTool/Modes/TestMode.cs b/Engine/Source/Programs/UnrealBuildTool/Modes/TestMode.cs
--- a/Engine/Source/Programs/UnrealBuildTool/Modes/TestMode.cs
+++ b/Engine/Source/Programs/UnrealBuildTool/Modes/TestMode.cs
@@ -15,9 +15,11 @@ namespace UnrealBuildTool
 	/// <summary>
 	/// Builds low level tests on one or more targets.
 	/// </summary>
-	[ToolMode("Test", ToolModeOptions.XmlConfig | ToolModeOptions.BuildPlatforms | ToolModeOptions.SingleInstance | ToolModeOptions.ShowExecutionTime)]
-	class TestMode : ToolMode
+	internal sealed class TestMode : IToolMode<TestMode>
 	{
+		public static string Name => "Test";
+		public static ToolModeOptions Options => ToolModeOptions.XmlConfig | ToolModeOptions.BuildPlatforms | ToolModeOptions.SingleInstance | ToolModeOptions.ShowExecutionTime;
+
 		private readonly string ENGINE_MODULE = "Engine";
 		private readonly string APPLICATION_CORE_MODULE = "ApplicationCore";
 		private readonly string COREUOBJECT_MODULE = "CoreUObject";
@@ -66,7 +68,7 @@ namespace UnrealBuildTool
 		/// </summary>
 		/// <param name="Arguments">Command-line arguments</param>
 		/// <param name="Logger"></param>
-		public override async Task<int> ExecuteAsync(CommandLineArguments Arguments, ILogger Logger)
+		public async Task<int> ExecuteAsync(CommandLineArguments Arguments, ILogger Logger)
 		{
 			Arguments.ApplyTo(this);
 
diff --git a/Engine/Source/Programs/UnrealBuildTool/Modes/ToolMode.cs b/Engine/Source/Programs/UnrealBuildTool/Modes/ToolMode.cs
deleted file mode 100644
diff --git a/Engine/Source/Programs/UnrealBuildTool/Modes/UnrealHeaderToolMode.cs b/Engine/Source/Programs/UnrealBuildTool/Modes/UnrealHeaderToolMode.cs
--- a/Engine/Source/Programs/UnrealBuildTool/Modes/UnrealHeaderToolMode.cs
+++ b/Engine/Source/Programs/UnrealBuildTool/Modes/UnrealHeaderToolMode.cs
@@ -615,9 +615,11 @@ namespace UnrealBuildTool.Modes
 	/// <summary>
 	/// Invoke UHT
 	/// </summary>
-	[ToolMode("UnrealHeaderTool", ToolModeOptions.XmlConfig | ToolModeOptions.BuildPlatforms | ToolModeOptions.ShowExecutionTime)]
-	class UnrealHeaderToolMode : ToolMode
+	internal sealed class UnrealHeaderToolMode : IToolMode<UnrealHeaderToolMode>
 	{
+		public static string Name => "UnrealHeaderTool";
+		public static ToolModeOptions Options => ToolModeOptions.XmlConfig | ToolModeOptions.BuildPlatforms | ToolModeOptions.ShowExecutionTime;
+
 		/// <summary>
 		/// Directory for saved application settings (typically Engine/Programs)
 		/// </summary>
@@ -699,7 +701,7 @@ namespace UnrealBuildTool.Modes
 		/// <param name="Arguments">Command line arguments</param>
 		/// <returns>Exit code</returns>
 		/// <param name="Logger"></param>
-		public override async Task<int> ExecuteAsync(CommandLineArguments Arguments, ILogger Logger)
+		public async Task<int> ExecuteAsync(CommandLineArguments Arguments, ILogger Logger)
 		{
 			try
 			{
diff --git a/Engine/Source/Programs/UnrealBuildTool/Modes/ValidatePlatformsMode.cs b/Engine/Source/Programs/UnrealBuildTool/Modes/ValidatePlatformsMode.cs
--- a/Engine/Source/Programs/UnrealBuildTool/Modes/ValidatePlatformsMode.cs
+++ b/Engine/Source/Programs/UnrealBuildTool/Modes/ValidatePlatformsMode.cs
@@ -11,9 +11,11 @@ namespace UnrealBuildTool
 	/// <summary>
 	/// Validates the various platforms to determine if they are ready for building
 	/// </summary>
-	[ToolMode("ValidatePlatforms", ToolModeOptions.XmlConfig | ToolModeOptions.BuildPlatformsForValidation | ToolModeOptions.SingleInstance)]
-	class ValidatePlatformsMode : ToolMode
+	internal sealed class ValidatePlatformsMode : IToolMode<ValidatePlatformsMode>
 	{
+		public static string Name => "ValidatePlatforms";
+		public static ToolModeOptions Options => ToolModeOptions.XmlConfig | ToolModeOptions.BuildPlatformsForValidation | ToolModeOptions.SingleInstance;
+
 #pragma warning disable IDE0044 // Make field readonly - these private static fields are set by command-line parsing.
 		/// <summary>
 		/// Platforms to validate
@@ -40,7 +42,7 @@ namespace UnrealBuildTool
 		/// <param name="Arguments">Command line arguments</param>
 		/// <returns>Exit code</returns>
 		/// <param name="Logger"></param>
-		public override Task<int> ExecuteAsync(CommandLineArguments Arguments, ILogger Logger)
+		public Task<int> ExecuteAsync(CommandLineArguments Arguments, ILogger Logger)
 		{
 			// Output a message if there are any arguments that are still unused
 			Arguments.ApplyTo(this);
diff --git a/Engine/Source/Programs/UnrealBuildTool/Modes/WriteDocumentationMode.cs b/Engine/Source/Programs/UnrealBuildTool/Modes/WriteDocumentationMode.cs
--- a/Engine/Source/Programs/UnrealBuildTool/Modes/WriteDocumentationMode.cs
+++ b/Engine/Source/Programs/UnrealBuildTool/Modes/WriteDocumentationMode.cs
@@ -9,9 +9,11 @@ namespace UnrealBuildTool
 	/// <summary>
 	/// Generates documentation from reflection data
 	/// </summary>
-	[ToolMode("WriteDocumentation", ToolModeOptions.None)]
-	class WriteDocumentationMode : ToolMode
+	internal sealed class WriteDocumentationMode : IToolMode<WriteDocumentationMode>
 	{
+		public static string Name => "WriteDocumentation";
+		public static ToolModeOptions Options => ToolModeOptions.None;
+
 		/// <summary>
 		/// Enum for the type of documentation to generate
 		/// </summary>
@@ -40,7 +42,7 @@ namespace UnrealBuildTool
 		/// Entry point for this command
 		/// </summary>
 		/// <returns></returns>
-		public override Task<int> ExecuteAsync(CommandLineArguments Arguments, ILogger Logger)
+		public Task<int> ExecuteAsync(CommandLineArguments Arguments, ILogger Logger)
 		{
 			Arguments.ApplyTo(this);
 			Arguments.CheckAllArgumentsUsed();
diff --git a/Engine/Source/Programs/UnrealBuildTool/Modes/WriteMetadataMode.cs b/Engine/Source/Programs/UnrealBuildTool/Modes/WriteMetadataMode.cs
--- a/Engine/Source/Programs/UnrealBuildTool/Modes/WriteMetadataMode.cs
+++ b/Engine/Source/Programs/UnrealBuildTool/Modes/WriteMetadataMode.cs
@@ -35,9 +35,11 @@ namespace UnrealBuildTool
 	/// <summary>
 	/// Writes all metadata files at the end of a build (receipts, version files, etc...). This is implemented as a separate mode to allow it to be done as part of the action graph.
 	/// </summary>
-	[ToolMode("WriteMetadata", ToolModeOptions.None)]
-	class WriteMetadataMode : ToolMode
+	internal sealed class WriteMetadataMode : IToolMode<WriteMetadataMode>
 	{
+		public static string Name => "WriteMetadata";
+		public static ToolModeOptions Options => ToolModeOptions.None;
+
 		/// <summary>
 		/// Version number for output files. This is not used directly, but can be appended to command-line invocations of the tool to ensure that actions to generate metadata are updated if the output format changes. 
 		/// The action graph is regenerated whenever UBT is rebuilt, so this should always match.
@@ -58,7 +60,7 @@ namespace UnrealBuildTool
 		/// <param name="Arguments">Command line arguments</param>
 		/// <returns>Exit code</returns>
 		/// <param name="Logger"></param>
-		public override async Task<int> ExecuteAsync(CommandLineArguments Arguments, ILogger Logger)
+		public async Task<int> ExecuteAsync(CommandLineArguments Arguments, ILogger Logger)
 		{
 			// Acquire a different mutex to the regular UBT instance, since this mode will be called as part of a build. We need the mutex to ensure that building two modular configurations 
 			// in parallel don't clash over writing shared *.modules files (eg. DebugGame and Development editors).
diff --git a/Engine/Source/Programs/UnrealBuildTool/Platform/Apple/AppleToolChain.cs b/Engine/Source/Programs/UnrealBuildTool/Platform/Apple/AppleToolChain.cs
--- a/Engine/Source/Programs/UnrealBuildTool/Platform/Apple/AppleToolChain.cs
+++ b/Engine/Source/Programs/UnrealBuildTool/Platform/Apple/AppleToolChain.cs
@@ -1228,9 +1228,11 @@ namespace UnrealBuildTool
 		}
 	}
 
-	[ToolMode("ApplePostBuildSync", ToolModeOptions.XmlConfig | ToolModeOptions.BuildPlatforms)]
-	class ApplePostBuildSyncMode : ToolMode
+	internal sealed class ApplePostBuildSyncMode : IToolMode<ApplePostBuildSyncMode>
 	{
+		public static string Name => "ApplePostBuildSync";
+		public static ToolModeOptions Options => ToolModeOptions.XmlConfig | ToolModeOptions.BuildPlatforms;
+
 		[CommandLine("-Input=", Required = true)]
 		public FileReference? InputFile = null;
 
@@ -1244,7 +1246,7 @@ namespace UnrealBuildTool
 		[CommandLine("-noEntitlements")]
 		public bool bNoEntitlements;
 
-		public override async Task<int> ExecuteAsync(CommandLineArguments Arguments, ILogger Logger)
+		public async Task<int> ExecuteAsync(CommandLineArguments Arguments, ILogger Logger)
 		{
 			Arguments.ApplyTo(this);
 			Arguments.CheckAllArgumentsUsed();
diff --git a/Engine/Source/Programs/UnrealBuildTool/Platform/Microsoft/GDKGameConfigGenerator.cs b/Engine/Source/Programs/UnrealBuildTool/Platform/Microsoft/GDKGameConfigGenerator.cs
--- a/Engine/Source/Programs/UnrealBuildTool/Platform/Microsoft/GDKGameConfigGenerator.cs
+++ b/Engine/Source/Programs/UnrealBuildTool/Platform/Microsoft/GDKGameConfigGenerator.cs
@@ -862,9 +862,11 @@ namespace UnrealBuildTool.GDK
 	/// <summary>
 	/// ToolMode to create a MicrosoftGame.config for a single executable pair
 	/// </summary>
-	[ToolMode("GDKGameConfigGenerator", ToolModeOptions.BuildPlatforms)]
-	internal class GDKGameConfigGeneratorMode : ToolMode
+	internal class GDKGameConfigGeneratorMode : IToolMode<GDKGameConfigGeneratorMode>
 	{
+		public static string Name => "GDKGameConfigGenerator";
+		public static ToolModeOptions Options => ToolModeOptions.BuildPlatforms;
+
 #pragma warning disable IDE0044 // Make field readonly - these private fields are set by the command line.
 		[CommandLine("-Executable", Required = false)]
 		string? Executable = null;
@@ -892,7 +894,7 @@ namespace UnrealBuildTool.GDK
 #pragma warning restore IDE0044
 
 		/// <inheritdoc/>
-		public override Task<int> ExecuteAsync(CommandLineArguments Arguments, ILogger Logger)
+		public Task<int> ExecuteAsync(CommandLineArguments Arguments, ILogger Logger)
 		{
 			Arguments.ApplyTo(this);
 			Arguments.CheckAllArgumentsUsed();
diff --git a/Engine/Source/Programs/UnrealBuildTool/Platform/Windows/PVSToolChain.cs b/Engine/Source/Programs/UnrealBuildTool/Platform/Windows/PVSToolChain.cs
--- a/Engine/Source/Programs/UnrealBuildTool/Platform/Windows/PVSToolChain.cs
+++ b/Engine/Source/Programs/UnrealBuildTool/Platform/Windows/PVSToolChain.cs
@@ -370,9 +370,11 @@ namespace UnrealBuildTool
 	/// <summary>
 	/// Special mode for gathering all the messages into a single output file
 	/// </summary>
-	[ToolMode("PVSGather", ToolModeOptions.None)]
-	class PVSGatherMode : ToolMode
+	internal sealed class PVSGatherMode : IToolMode<PVSGatherMode>
 	{
+		public static string Name => "PVSGather";
+		public static ToolModeOptions Options => ToolModeOptions.None;
+
 #pragma warning disable IDE0044 // Make field readonly - these private static fields are set by command-line parsing.
 		/// <summary>
 		/// Path to the input file list
@@ -429,7 +431,7 @@ namespace UnrealBuildTool
 		/// <param name="arguments">List of command line arguments</param>
 		/// <returns>Always zero, or throws an exception</returns>
 		/// <param name="logger"></param>
-		public override async Task<int> ExecuteAsync(CommandLineArguments arguments, ILogger logger)
+		public async Task<int> ExecuteAsync(CommandLineArguments arguments, ILogger logger)
 		{
 			arguments.ApplyTo(this);
 			arguments.CheckAllArgumentsUsed();
diff --git a/Engine/Source/Programs/UnrealBuildTool/System/Telemetry.cs b/Engine/Source/Programs/UnrealBuildTool/System/Telemetry.cs
--- a/Engine/Source/Programs/UnrealBuildTool/System/Telemetry.cs
+++ b/Engine/Source/Programs/UnrealBuildTool/System/Telemetry.cs
@@ -324,7 +324,7 @@ namespace UnrealBuildTool
 				"eteventstream",
 				Environment.UserName,
 				UnrealBuildTool.SessionIdentifier,
-				UnrealBuildTool.BuildMode,
+				UnrealBuildTool.CurrentMode,
 				Get()._descriptor,
 				buildBranchName,
 				Unreal.IsBuildMachine(),
diff --git a/Engine/Source/Programs/UnrealBuildTool/UnrealBuildTool.cs b/Engine/Source/Programs/UnrealBuildTool/UnrealBuildTool.cs
--- a/Engine/Source/Programs/UnrealBuildTool/UnrealBuildTool.cs
+++ b/Engine/Source/Programs/UnrealBuildTool/UnrealBuildTool.cs
@@ -37,9 +37,9 @@ namespace UnrealBuildTool
 		public static string SessionIdentifier = Guid.NewGuid().ToString("B");
 
 		/// <summary>
-		///  The mode of this instance
+		/// The mode of this instance
 		/// </summary>
-		public static string BuildMode = "BuildMode";
+		public static string CurrentMode = "<uninitialized>";
 
 		/// <summary>
 		/// The result of running the application
@@ -72,16 +72,22 @@ namespace UnrealBuildTool
 		/// Get all the valid Modes
 		/// </summary>
 		/// <returns></returns>
-		private static Dictionary<string, Type> GetModes()
+		private static Dictionary<string, ModeDescriptor> GetModes()
 		{
-			Dictionary<string, Type> ModeNameToType = new(StringComparer.OrdinalIgnoreCase);
-			foreach (Type Type in Assembly.GetExecutingAssembly().GetTypes().Where(x => x.IsClass && !x.IsAbstract && x.IsSubclassOf(typeof(ToolMode))))
+			Dictionary<string, ModeDescriptor> modeNameToDescriptor = new(StringComparer.OrdinalIgnoreCase);
+			foreach (Type type in Assembly.GetExecutingAssembly().GetTypes().Where(x => x.IsClass && !x.IsAbstract))
 			{
-				ToolModeAttribute? Attribute = Type.GetCustomAttribute<ToolModeAttribute>()
-					?? throw new BuildLogEventException("Class '{Name}' should have a ToolModeAttribute", Type.Name);
-				ModeNameToType.Add(Attribute.Name, Type);
+				bool isToolMode = type.GetInterfaces().FirstOrDefault(x => x.IsGenericType && x.GetGenericTypeDefinition() == typeof(IToolMode<>)) is Type ToolModeType;
+				if (!isToolMode)
+				{
+					continue;
+				}
+
+				string name = (string)type.GetProperty("Name", BindingFlags.Public | BindingFlags.Static)!.GetValue(null)!;
+				ToolModeOptions options = (ToolModeOptions)type.GetProperty("Options", BindingFlags.Public | BindingFlags.Static)!.GetValue(null)!;
+				modeNameToDescriptor.Add(name, new(type, options));
 			}
-			return ModeNameToType;
+			return modeNameToDescriptor;
 		}
 
 		/// <summary>
@@ -108,10 +114,10 @@ namespace UnrealBuildTool
 				}
 			}
 
-			Dictionary<string, Type> modes = GetModes();
+			Dictionary<string, ModeDescriptor> modes = GetModes();
 			if (modes.Count == 0)
 			{
-				throw new BuildLogEventException("No valid subclasses of {Type} found", typeof(ToolMode));
+				throw new BuildLogEventException("No valid subclasses of {Type} found", typeof(IToolMode<>));
 			}
 
 			// Load tool mode summary from xml documentation if available
@@ -129,10 +135,10 @@ namespace UnrealBuildTool
 			longestPrefix = modes.Select(x => x.Key.Length).Max();
 			Console.WriteLine();
 			Console.WriteLine("Available Tool Modes:");
-			foreach ((string name, Type type) in modes)
+			foreach ((string name, ModeDescriptor descriptor) in modes)
 			{
 				string summary = doc?.Descendants("member")
-					.FirstOrDefault(x => x.Attribute("name")?.Value == $"T:{type.FullName}")?
+					.FirstOrDefault(x => x.Attribute("name")?.Value == $"T:{descriptor.ModeType.FullName}")?
 					.Element("summary")?.Value.Trim()
 					?? "Summary unavailable";
 				Console.WriteLine($" {name.PadRight(longestPrefix)} : {Regex.Replace(summary, @"\s+", " ")}");
@@ -298,21 +304,21 @@ namespace UnrealBuildTool
 				Encoding.RegisterProvider(CodePagesEncodingProvider.Instance);
 
 				// Get the type of the mode to execute, using a fast-path for the build mode.
-				Type? ModeType = typeof(BuildMode);
+				ModeDescriptor? ModeDescriptor = new(typeof(BuildMode), BuildMode.Options);
 				if (Options.Mode != null)
 				{
 					// Try to get the correct mode
-					IReadOnlyDictionary<string, Type> ModeNameToType = GetModes();
-					if (!ModeNameToType.TryGetValue(Options.Mode, out ModeType))
+					IReadOnlyDictionary<string, ModeDescriptor> ModeNameToDesciptor = GetModes();
+					if (!ModeNameToDesciptor.TryGetValue(Options.Mode, out ModeDescriptor))
 					{
-						Logger.LogError("No mode named '{Name}'. Available modes are:\n  {ModeList}", Options.Mode, String.Join("\n  ", ModeNameToType.Keys.Order()));
+						Logger.LogError("No mode named '{Name}'. Available modes are:\n  {ModeList}", Options.Mode, String.Join("\n  ", ModeNameToDesciptor.Keys.Order()));
 						return 1;
 					}
 				}
-				BuildMode = ModeType.Name;
+				CurrentMode = ModeDescriptor.ModeType.Name;
 
 				// Get the options for which systems have to be initialized for this mode
-				ModeOptions = ModeType.GetCustomAttribute<ToolModeAttribute>()!.Options;
+				ModeOptions = ModeDescriptor.Options;
 
 				// if we don't care about the trace listener, toss it now
 				if ((ModeOptions & ToolModeOptions.UseStartupTraceListener) == 0)
@@ -455,27 +461,27 @@ namespace UnrealBuildTool
 					{
 						using (GlobalTracer.Instance.BuildSpan("UEBuildPlatform.RegisterPlatforms()").StartActive())
 						{
-							UEBuildPlatform.RegisterPlatforms(false, false, ModeType, ArgumentsArray, Logger);
+							UEBuildPlatform.RegisterPlatforms(false, false, ModeDescriptor.ModeType, ArgumentsArray, Logger);
 						}
 					}
 					if ((ModeOptions & ToolModeOptions.BuildPlatformsHostOnly) != 0)
 					{
 						using (GlobalTracer.Instance.BuildSpan("UEBuildPlatform.RegisterPlatforms()").StartActive())
 						{
-							UEBuildPlatform.RegisterPlatforms(false, true, ModeType, ArgumentsArray, Logger);
+							UEBuildPlatform.RegisterPlatforms(false, true, ModeDescriptor.ModeType, ArgumentsArray, Logger);
 						}
 					}
 					if ((ModeOptions & ToolModeOptions.BuildPlatformsForValidation) != 0)
 					{
 						using (GlobalTracer.Instance.BuildSpan("UEBuildPlatform.RegisterPlatforms()").StartActive())
 						{
-							UEBuildPlatform.RegisterPlatforms(true, false, ModeType, ArgumentsArray, Logger);
+							UEBuildPlatform.RegisterPlatforms(true, false, ModeDescriptor.ModeType, ArgumentsArray, Logger);
 						}
 					}
 				}
 
 				// Create the appropriate handler
-				ToolMode Mode = (ToolMode)Activator.CreateInstance(ModeType)!;
+				IToolMode Mode = (IToolMode)Activator.CreateInstance(ModeDescriptor.ModeType)!;
 
 				// Execute the mode
 				MainScope.Finish();
@@ -580,5 +586,7 @@ namespace UnrealBuildTool
 				Mutex?.Dispose();
 			}
 		}
+
+		private record class ModeDescriptor(Type ModeType, ToolModeOptions Options);
 	}
 }
