UnrealBuildTool: Print CPU and memory info when starting the UBA exectutor

#rnx

#Synced-CL 52364399

diff --git a/Engine/Source/Programs/UnrealBuildTool/Executors/UnrealBuildAccelerator/UBAExecutor.cs b/Engine/Source/Programs/UnrealBuildTool/Executors/UnrealBuildAccelerator/UBAExecutor.cs
--- a/Engine/Source/Programs/UnrealBuildTool/Executors/UnrealBuildAccelerator/UBAExecutor.cs
+++ b/Engine/Source/Programs/UnrealBuildTool/Executors/UnrealBuildAccelerator/UBAExecutor.cs
@@ -479,12 +479,18 @@ namespace UnrealBuildTool
 
 		private void PrintConfiguration(IEnumerable<LinkedAction> inputActions, ILogger logger)
 		{
-			logger.LogInformation("  Storage capacity {StoreCapacity}", StringUtils.FormatBytesString((ulong)UBAConfig.StoreCapacityGb * 1024 * 1024 * 1024));
+			long totalPhysicalBytes = SystemUtils.GetTotalSystemMemoryBytes();
+			long totalCommittedMemory = SystemUtils.GetMaxCommittedMemoryBytes();
+			long currentCommittedMemory = SystemUtils.GetCurrentCommittedMemoryBytes();
+
+			logger.LogInformation("  CPU {Physical} physical cores, {Logical} logical cores", SystemUtils.GetPhysicalProcessorCount(), SystemUtils.GetLogicalProcessorCount());
+			logger.LogInformation("  Memory {PhysicalTotal} physical, {CommittedCurrent}/{CommittedTotal} committed", StringUtils.FormatBytesString(totalPhysicalBytes), StringUtils.FormatBytesString(currentCommittedMemory), StringUtils.FormatBytesString(totalCommittedMemory));
+			logger.LogInformation("  UBA Storage capacity {StoreCapacity}", StringUtils.FormatBytesString((ulong)UBAConfig.StoreCapacityGb * 1024 * 1024 * 1024));
 			if (UBAConfig.bStoreObjFilesCompressed)
 			{
-				logger.LogInformation("  Object Compression Allowed");
+				logger.LogInformation("  UBA Object Compression Allowed");
 			}
-			logger.LogDebug("  RootDir {RootDir}", _rootDirRef);
+			logger.LogDebug("  UBA RootDir {RootDir}", _rootDirRef);
 
 			CaptureLogger recommendationLogger = new();
 			if (_cacheableActionsLeft > 0 && !inputActions.Any(x => x.RootPaths.bUseVfs))
