TargetFramework net10.0

#rnx

#Synced-CL 52396139

#Synced-CL 52396139

diff --git a/Engine/Source/Programs/AutomationTool.Tests/AutomationTool.Tests.csproj b/Engine/Source/Programs/AutomationTool.Tests/AutomationTool.Tests.csproj
--- a/Engine/Source/Programs/AutomationTool.Tests/AutomationTool.Tests.csproj
+++ b/Engine/Source/Programs/AutomationTool.Tests/AutomationTool.Tests.csproj
@@ -17,7 +17,6 @@
       <PrivateAssets>all</PrivateAssets>
       <IncludeAssets>runtime; build; native; contentfiles; analyzers; buildtransitive</IncludeAssets>
     </PackageReference>
-    <PackageReference Include="System.Text.Json" Version="8.0.6" />
   </ItemGroup>
 
   <ItemGroup>
diff --git a/Engine/Source/Programs/AutomationTool/AutomationTool.csproj b/Engine/Source/Programs/AutomationTool/AutomationTool.csproj
--- a/Engine/Source/Programs/AutomationTool/AutomationTool.csproj
+++ b/Engine/Source/Programs/AutomationTool/AutomationTool.csproj
@@ -4,13 +4,13 @@
   <Choose>
     <When Condition="'$(IsWindows)' == 'true'">
       <PropertyGroup>
-        <TargetFramework>net8.0-windows10.0.19041</TargetFramework>
+        <TargetFramework>net10.0-windows</TargetFramework>
         <UseWindowsForms>true</UseWindowsForms>
       </PropertyGroup>
     </When>
     <Otherwise>
       <PropertyGroup>
-        <TargetFramework>net8.0</TargetFramework>
+        <TargetFramework>net10.0</TargetFramework>
       </PropertyGroup>
     </Otherwise>
   </Choose>
diff --git a/Engine/Source/Programs/AutomationTool/AutomationUtils/AutomationUtils.Automation.csproj b/Engine/Source/Programs/AutomationTool/AutomationUtils/AutomationUtils.Automation.csproj
--- a/Engine/Source/Programs/AutomationTool/AutomationUtils/AutomationUtils.Automation.csproj
+++ b/Engine/Source/Programs/AutomationTool/AutomationUtils/AutomationUtils.Automation.csproj
@@ -39,14 +39,10 @@
   <ItemGroup>
     <PackageReference Include="Magick.NET-Q16-HDRI-AnyCPU" Version="14.9.1" PrivateAssets="all" />
     <PackageReference Include="System.Drawing.Common" Version="8.0.20" />
-    <PackageReference Include="System.Net.Http" Version="4.3.4" />
-
     <!--Required by Ionic -->
     <PackageReference Include="System.Security.Permissions" Version="8.0.0" />
-    <PackageReference Include="System.Text.Encoding.CodePages" Version="8.0.0" />
 
     <PackageReference Include="Google.Apis.Drive.v3" Version="1.33.1.1232" PrivateAssets="all" />
     <PackageReference Include="Google.Apis.Sheets.v4" Version="1.33.1.1240" PrivateAssets="all" />
-    <PackageReference Include="System.Text.RegularExpressions" Version="4.3.1" />
   </ItemGroup>
 </Project>
\ No newline at end of file
diff --git a/Engine/Source/Programs/Shared/EpicGames.AspNet/EpicGames.AspNet.csproj b/Engine/Source/Programs/Shared/EpicGames.AspNet/EpicGames.AspNet.csproj
--- a/Engine/Source/Programs/Shared/EpicGames.AspNet/EpicGames.AspNet.csproj
+++ b/Engine/Source/Programs/Shared/EpicGames.AspNet/EpicGames.AspNet.csproj
@@ -17,8 +17,4 @@
     <ProjectReference Include="..\..\Shared\EpicGames.Core\EpicGames.Core.csproj" />
     <ProjectReference Include="..\EpicGames.Serialization\EpicGames.Serialization.csproj" />
   </ItemGroup>
-
-  <ItemGroup>
-    <PackageReference Include="System.Text.Json" Version="8.0.6" />
-  </ItemGroup>
 </Project>
diff --git a/Engine/Source/Programs/Shared/EpicGames.BuildGraph.Tests/EpicGames.BuildGraph.Tests.csproj b/Engine/Source/Programs/Shared/EpicGames.BuildGraph.Tests/EpicGames.BuildGraph.Tests.csproj
--- a/Engine/Source/Programs/Shared/EpicGames.BuildGraph.Tests/EpicGames.BuildGraph.Tests.csproj
+++ b/Engine/Source/Programs/Shared/EpicGames.BuildGraph.Tests/EpicGames.BuildGraph.Tests.csproj
@@ -1,8 +1,4 @@
 ﻿<Project Sdk="Microsoft.NET.Sdk">
-  <ItemGroup>
-    <PackageReference Include="System.Text.Json" Version="8.0.6" />
-  </ItemGroup>
-
   <ItemGroup>
     <ProjectReference Include="..\EpicGames.BuildGraph\EpicGames.BuildGraph.csproj" />
   </ItemGroup>
diff --git a/Engine/Source/Programs/Shared/EpicGames.Core/EpicGames.Core.csproj b/Engine/Source/Programs/Shared/EpicGames.Core/EpicGames.Core.csproj
--- a/Engine/Source/Programs/Shared/EpicGames.Core/EpicGames.Core.csproj
+++ b/Engine/Source/Programs/Shared/EpicGames.Core/EpicGames.Core.csproj
@@ -10,6 +10,5 @@
     <PackageReference Include="Microsoft.Extensions.Logging.Console" Version="8.0.1" />
     <PackageReference Include="Microsoft.Extensions.ObjectPool" Version="8.0.20" />
     <PackageReference Include="OpenTracing" Version="0.12.1" />
-    <PackageReference Include="System.Text.Json" Version="8.0.6" />
   </ItemGroup>
 </Project>
diff --git a/Engine/Source/Programs/Shared/EpicGames.Horde.Tests/EpicGames.Horde.Tests.csproj b/Engine/Source/Programs/Shared/EpicGames.Horde.Tests/EpicGames.Horde.Tests.csproj
--- a/Engine/Source/Programs/Shared/EpicGames.Horde.Tests/EpicGames.Horde.Tests.csproj
+++ b/Engine/Source/Programs/Shared/EpicGames.Horde.Tests/EpicGames.Horde.Tests.csproj
@@ -1,8 +1,4 @@
 <Project Sdk="Microsoft.NET.Sdk">
-  <ItemGroup>
-    <PackageReference Include="System.Linq.Async" Version="6.0.3" />
-  </ItemGroup>
-
   <ItemGroup>
     <ProjectReference Include="..\EpicGames.Core\EpicGames.Core.csproj" />
     <ProjectReference Include="..\EpicGames.Horde\EpicGames.Horde.csproj" />
diff --git a/Engine/Source/Programs/Shared/EpicGames.Horde/Compute/Transports/TcpSslTransport.cs b/Engine/Source/Programs/Shared/EpicGames.Horde/Compute/Transports/TcpSslTransport.cs
--- a/Engine/Source/Programs/Shared/EpicGames.Horde/Compute/Transports/TcpSslTransport.cs
+++ b/Engine/Source/Programs/Shared/EpicGames.Horde/Compute/Transports/TcpSslTransport.cs
@@ -28,7 +28,7 @@ public sealed class TcpSslTransport : ComputeTransport
 	/// <param name="socket">Socket to communicate over</param>
 	/// <param name="certData">Certificate used for auth on both server and client</param>
 	/// <param name="isServer">Whether socket is acting as a server or client</param>
-	public TcpSslTransport(Socket socket, byte[] certData, bool isServer) : this(socket, new X509Certificate2(certData), isServer)
+	public TcpSslTransport(Socket socket, byte[] certData, bool isServer) : this(socket, X509CertificateLoader.LoadPkcs12(certData, null), isServer)
 	{
 	}
 
diff --git a/Engine/Source/Programs/Shared/EpicGames.Horde/EpicGames.Horde.csproj b/Engine/Source/Programs/Shared/EpicGames.Horde/EpicGames.Horde.csproj
--- a/Engine/Source/Programs/Shared/EpicGames.Horde/EpicGames.Horde.csproj
+++ b/Engine/Source/Programs/Shared/EpicGames.Horde/EpicGames.Horde.csproj
@@ -14,14 +14,12 @@
       <IncludeAssets>runtime; build; native; contentfiles; analyzers; buildtransitive</IncludeAssets>
     </PackageReference>
     <PackageReference Include="K4os.Compression.LZ4" Version="1.2.16" />
-    <PackageReference Include="System.IO.Pipelines" Version="8.0.0" />
     <PackageReference Include="Microsoft.Data.Sqlite" Version="8.0.20" />
     <PackageReference Include="Microsoft.Extensions.Caching.Memory" Version="8.0.1" />
     <PackageReference Include="Microsoft.VisualStudio.Threading.Analyzers" Version="17.14.15">
       <PrivateAssets>all</PrivateAssets>
       <IncludeAssets>runtime; build; native; contentfiles; analyzers; buildtransitive</IncludeAssets>
     </PackageReference>
-    <PackageReference Include="System.Linq.Async" Version="6.0.3" />
     <PackageReference Include="ZstdSharp.Port" Version="0.8.6" />
   </ItemGroup>
 
diff --git a/Engine/Source/Programs/Shared/EpicGames.IoHash/EpicGames.IoHash.csproj b/Engine/Source/Programs/Shared/EpicGames.IoHash/EpicGames.IoHash.csproj
--- a/Engine/Source/Programs/Shared/EpicGames.IoHash/EpicGames.IoHash.csproj
+++ b/Engine/Source/Programs/Shared/EpicGames.IoHash/EpicGames.IoHash.csproj
@@ -6,7 +6,6 @@
   <ItemGroup>
     <!-- Blake3 2.0.0 has a massive performance regression on Windows x64 - takes 2.25x as long for large (1MB) inputs -->
     <PackageReference Include="Blake3" Version="1.1.0" />
-    <PackageReference Include="System.Memory" Version="4.6.3" />
   </ItemGroup>
 
   <ItemGroup>
diff --git a/Engine/Source/Programs/Shared/EpicGames.OIDC/EpicGames.OIDC.csproj b/Engine/Source/Programs/Shared/EpicGames.OIDC/EpicGames.OIDC.csproj
--- a/Engine/Source/Programs/Shared/EpicGames.OIDC/EpicGames.OIDC.csproj
+++ b/Engine/Source/Programs/Shared/EpicGames.OIDC/EpicGames.OIDC.csproj
@@ -1,6 +1,6 @@
 ﻿<Project Sdk="Microsoft.NET.Sdk">
   <PropertyGroup>
-    <EnableConfigurationBindingGenerator>true</EnableConfigurationBindingGenerator>
+    <EnableConfigurationBindingGenerator>false</EnableConfigurationBindingGenerator>
   </PropertyGroup>
 
   <ItemGroup>
@@ -8,7 +8,6 @@
     <PackageReference Include="Microsoft.Extensions.Configuration.Binder" Version="8.0.2" />
     <PackageReference Include="Microsoft.Extensions.Configuration.Json" Version="8.0.1" />
     <PackageReference Include="Microsoft.Extensions.Options" Version="8.0.2" />
-    <PackageReference Include="System.Text.Json" Version="8.0.6" />
   </ItemGroup>
 
   <ItemGroup>
diff --git a/Engine/Source/Programs/Shared/EpicGames.Perforce.Fixture/EpicGames.Perforce.Fixture.csproj b/Engine/Source/Programs/Shared/EpicGames.Perforce.Fixture/EpicGames.Perforce.Fixture.csproj
--- a/Engine/Source/Programs/Shared/EpicGames.Perforce.Fixture/EpicGames.Perforce.Fixture.csproj
+++ b/Engine/Source/Programs/Shared/EpicGames.Perforce.Fixture/EpicGames.Perforce.Fixture.csproj
@@ -20,7 +20,6 @@
       <PrivateAssets>all</PrivateAssets>
       <IncludeAssets>runtime; build; native; contentfiles; analyzers; buildtransitive</IncludeAssets>
     </PackageReference>
-    <PackageReference Include="System.Text.Json" Version="8.0.6" />
   </ItemGroup>
 
   <ItemGroup>
diff --git a/Engine/Source/Programs/Shared/EpicGames.Perforce/EpicGames.Perforce.csproj b/Engine/Source/Programs/Shared/EpicGames.Perforce/EpicGames.Perforce.csproj
--- a/Engine/Source/Programs/Shared/EpicGames.Perforce/EpicGames.Perforce.csproj
+++ b/Engine/Source/Programs/Shared/EpicGames.Perforce/EpicGames.Perforce.csproj
@@ -35,7 +35,6 @@
   </ItemGroup>
 
   <ItemGroup>
-    <PackageReference Include="System.Linq.Async" Version="6.0.3" />
     <PackageReference Include="OpenTelemetry.Api" Version="1.8.1" />
   </ItemGroup>
 
diff --git a/Engine/Source/Programs/UnrealBuildTool/Platform/Apple/AppleToolChain.cs b/Engine/Source/Programs/UnrealBuildTool/Platform/Apple/AppleToolChain.cs
--- a/Engine/Source/Programs/UnrealBuildTool/Platform/Apple/AppleToolChain.cs
+++ b/Engine/Source/Programs/UnrealBuildTool/Platform/Apple/AppleToolChain.cs
@@ -1352,7 +1352,7 @@ namespace UnrealBuildTool
 				X509Certificate2 Certificate;
 				try
 				{
-					Certificate = new X509Certificate2(Target.RemoteImportCertificate, Target.RemoteImportCertificatePassword ?? "");
+					Certificate = X509CertificateLoader.LoadPkcs12FromFile(Target.RemoteImportCertificate, Target.RemoteImportCertificatePassword ?? "");
 				}
 				catch (Exception Ex)
 				{
diff --git a/Engine/Source/Programs/UnrealBuildTool/Platform/IOS/IOSCodeSigning.cs b/Engine/Source/Programs/UnrealBuildTool/Platform/IOS/IOSCodeSigning.cs
--- a/Engine/Source/Programs/UnrealBuildTool/Platform/IOS/IOSCodeSigning.cs
+++ b/Engine/Source/Programs/UnrealBuildTool/Platform/IOS/IOSCodeSigning.cs
@@ -153,7 +153,7 @@ namespace UnrealBuildTool
 							int end = CertToolData.IndexOf(footer, start, StringComparison.Ordinal);
 							string base64 = CertToolData.Substring(start, (end - start));
 							byte[] certData = Convert.FromBase64String(base64);
-							X509Certificate2 cert = new X509Certificate2(certData);
+							X509Certificate2 cert = X509CertificateLoader.LoadCertificate(certData);
 							FoundCerts.Add(cert);
 							start = CertToolData.IndexOf(header, start, StringComparison.Ordinal);
 						}
@@ -1212,7 +1212,7 @@ namespace UnrealBuildTool
 			foreach (string EncodedCert in CertificateList)
 			{
 				byte[] RawCert = Convert.FromBase64String(EncodedCert);
-				DeveloperCertificates.Add(new X509Certificate2(RawCert, CertificatePassword));
+				DeveloperCertificates.Add(X509CertificateLoader.LoadPkcs12(RawCert, CertificatePassword));
 			}
 
 			// Key: Name, String
@@ -1362,7 +1362,7 @@ namespace UnrealBuildTool
 			// Read in the entire file
 			int NumBytes = (int)InputStream.Length;
 			byte[] RawData = new byte[NumBytes];
-			InputStream.Read(RawData, 0, NumBytes);
+			InputStream.ReadExactly(RawData, 0, NumBytes);
 
 			return ParseFile(RawData);
 		}
diff --git a/Engine/Source/Programs/UnrealBuildTool/UnrealBuildTool.csproj b/Engine/Source/Programs/UnrealBuildTool/UnrealBuildTool.csproj
--- a/Engine/Source/Programs/UnrealBuildTool/UnrealBuildTool.csproj
+++ b/Engine/Source/Programs/UnrealBuildTool/UnrealBuildTool.csproj
@@ -72,16 +72,13 @@
     <PackageReference Include="Microsoft.Extensions.Logging.Console" Version="8.0.1" />
     <PackageReference Include="Microsoft.VisualStudio.Setup.Configuration.Interop" Version="3.14.2075" />
     <PackageReference Include="Microsoft.VisualStudio.SolutionPersistence" Version="1.0.52" />
-    <PackageReference Include="Microsoft.Win32.Registry" Version="5.0.0" />
     <PackageReference Include="OpenTracing" Version="0.12.1" />
     <PackageReference Include="System.CodeDom" Version="8.0.0" />
     <PackageReference Include="System.Management" Version="8.0.0" />
-    <PackageReference Include="System.Security.Cryptography.Csp" Version="4.3.0" />
     <PackageReference Include="System.ServiceProcess.ServiceController" Version="8.0.1" />
 
     <!-- Required by Ionic -->
     <PackageReference Include="System.Security.Permissions" Version="8.0.0" />
-    <PackageReference Include="System.Text.Encoding.CodePages" Version="8.0.0" />
   </ItemGroup>
 
 </Project>
