UnrealBuildTool: Cleanly allow cancellation of the request to get a UBA cache server.

#Synced-CL 52466930

diff --git a/Engine/Source/Programs/UnrealBuildTool/Executors/UnrealBuildAccelerator/UBAHordeSession.cs b/Engine/Source/Programs/UnrealBuildTool/Executors/UnrealBuildAccelerator/UBAHordeSession.cs
--- a/Engine/Source/Programs/UnrealBuildTool/Executors/UnrealBuildAccelerator/UBAHordeSession.cs
+++ b/Engine/Source/Programs/UnrealBuildTool/Executors/UnrealBuildAccelerator/UBAHordeSession.cs
@@ -168,7 +168,16 @@ namespace UnrealBuildTool
 		{
 			ClusterId clusterId = new(UnrealBuildAcceleratorHordeConfig.ClusterDefault);
 			using ITimelineEvent timelineEvent = Timeline.ScopeEvent($"HordeGetCache");
-			UbaConfig ubaConfig = _client.AllocateUbaCacheServerAsync(clusterId, cancellationToken).Result;
+			UbaConfig ubaConfig;
+			try
+			{
+				ubaConfig = _client.AllocateUbaCacheServerAsync(clusterId, cancellationToken).GetAwaiter().GetResult();
+			}
+			catch (OperationCanceledException)
+			{
+				// Might get cancelled if the actions finish faster than the call to allocate a cache server. This is fine.
+				return null;
+			}
 			timelineEvent.Finish();
 			if (String.IsNullOrEmpty(ubaConfig.CacheEndpoint))
 			{
