Expose -allmodules and -nolink to BuildConfiguration.xml
#rnx

diff --git a/Engine/Source/Programs/UnrealBuildTool/Configuration/Rules/TargetRules.cs b/Engine/Source/Programs/UnrealBuildTool/Configuration/Rules/TargetRules.cs
--- a/Engine/Source/Programs/UnrealBuildTool/Configuration/Rules/TargetRules.cs
+++ b/Engine/Source/Programs/UnrealBuildTool/Configuration/Rules/TargetRules.cs
@@ -831,6 +831,7 @@ namespace UnrealBuildTool
 		/// Build all the modules that are valid for this target type. Used for CIS and making installed engine builds.
 		/// </summary>
 		[CommandLine("-AllModules")]
+		[XmlConfigFile(Category = "BuildConfiguration")]
 		public bool bBuildAllModules { get; set; }
 
 		/// <summary>
@@ -2483,6 +2484,7 @@ namespace UnrealBuildTool
 		/// Whether to disable linking for this target.
 		/// </summary>
 		[CommandLine("-NoLink")]
+		[XmlConfigFile(Category = "BuildConfiguration")]
 		public bool bDisableLinking { get; set; }
 
 		/// <summary>
