Optimize UBT's C++ include parser for adaptive non-unity builds to reduce allocations.

#Synced-CL 52502778

diff --git a/Engine/Source/Programs/UnrealBuildTool/Configuration/UEBuildModuleCPP.cs b/Engine/Source/Programs/UnrealBuildTool/Configuration/UEBuildModuleCPP.cs
--- a/Engine/Source/Programs/UnrealBuildTool/Configuration/UEBuildModuleCPP.cs
+++ b/Engine/Source/Programs/UnrealBuildTool/Configuration/UEBuildModuleCPP.cs
@@ -2151,10 +2151,15 @@ namespace UnrealBuildTool
 		static CPPOutput CompileAdaptiveNonUnityFilesWithDedicatedPCH(UEToolChain ToolChain, CppCompileEnvironment CompileEnvironment, List<FileItem> Files, DirectoryReference IntermediateDirectory, string ModuleName, IActionGraphBuilder Graph)
 		{
 			CPPOutput Output = new CPPOutput();
+			// Pre-allocate as much space as we're likely to need for a single include directive, to save reallocations.
+			// The default capacity of 16 chars is too small for many includes.
+			StringBuilder TempStorageForCopyIncludes = new StringBuilder(512);
+			// 16KB of headers is a reasonable estimate to save reallocation.
+			StringBuilder WrapperContents = new StringBuilder(1024 * 16);
 			foreach (FileItem File in Files)
 			{
+				WrapperContents.Clear();
 				// Build the contents of the wrapper file
-				StringBuilder WrapperContents = new StringBuilder();
 				{
 					string FileString = File.AbsolutePath;
 					if (File.Location.IsUnderDirectory(Unreal.RootDirectory))
@@ -2162,11 +2167,11 @@ namespace UnrealBuildTool
 						FileString = File.Location.MakeRelativeTo(Unreal.EngineSourceDirectory);
 					}
 					FileString = FileString.Replace('\\', '/');
-					WriteDefinitions($"Dedicated PCH for {FileString}", Array.Empty<string>(), CompileEnvironment.Definitions, Array.Empty<string>(), WrapperContents);
+					WriteDefinitions($"Dedicated PCH for {FileString}", [], CompileEnvironment.Definitions, [], WrapperContents);
 					WrapperContents.AppendLine();
 					using (StreamReader Reader = new StreamReader(File.Location.FullName))
 					{
-						CppIncludeParser.CopyIncludeDirectives(Reader, WrapperContents);
+						CppIncludeParser.CopyIncludeDirectives(Reader, WrapperContents, TempStorageForCopyIncludes);
 					}
 				}
 
diff --git a/Engine/Source/Programs/UnrealBuildTool/System/CppIncludeParser.cs b/Engine/Source/Programs/UnrealBuildTool/System/CppIncludeParser.cs
--- a/Engine/Source/Programs/UnrealBuildTool/System/CppIncludeParser.cs
+++ b/Engine/Source/Programs/UnrealBuildTool/System/CppIncludeParser.cs
@@ -9,64 +9,69 @@ namespace UnrealBuildTool
 	/// <summary>
 	/// Parses include directives from cpp source files, to make dedicated PCHs
 	/// </summary>
-	static class CppIncludeParser
+	internal static class CppIncludeParser
 	{
 		/// <summary>
 		/// Copies include directives from one file to another, until an unsafe directive (or non-#include line) is found.
 		/// </summary>
-		/// <param name="Reader">Stream to read from</param>
-		/// <param name="Writer">Stream to write directives to</param>
-		public static void CopyIncludeDirectives(StreamReader Reader, StringBuilder Writer)
+		/// <param name="reader">Stream to read from</param>
+		/// <param name="writer">Stream to write directives to</param>
+		/// <param name="tempStorage">Buffer to use for parser token storage</param>
+		internal static void CopyIncludeDirectives(StreamReader reader, StringBuilder writer, StringBuilder tempStorage)
 		{
-			StringBuilder Token = new StringBuilder();
-			while (TryReadToken(Reader, Token))
+			while (TryReadToken(reader, tempStorage))
 			{
-				if (Token.Length > 1 || Token[0] != '\n')
+				if (tempStorage.Length > 1 || tempStorage[0] != '\n')
 				{
-					if (Token.Length != 1 || Token[0] != '#')
+					if (tempStorage.Length != 1 || tempStorage[0] != '#')
 					{
 						break;
 					}
-					if (!TryReadToken(Reader, Token))
+					if (!TryReadToken(reader, tempStorage))
 					{
 						break;
 					}
 
-					string Directive = Token.ToString();
-					if (Directive == "pragma")
+					ReadOnlySpan<char> directiveSpan = GetSpanFromStringBuilder(tempStorage);
+					if (directiveSpan.Equals("pragma", StringComparison.Ordinal))
 					{
-						if (!TryReadToken(Reader, Token) || Token.ToString() != "once")
+						if (!TryReadToken(reader, tempStorage) || !GetSpanFromStringBuilder(tempStorage).Equals("once", StringComparison.Ordinal))
 						{
 							break;
 						}
-						if (!TryReadToken(Reader, Token) || Token.ToString() != "\n")
+						if (!TryReadToken(reader, tempStorage) || tempStorage[0] != '\n')
 						{
 							break;
 						}
 					}
-					else if (Directive == "include")
+					else if (directiveSpan.Equals("include", StringComparison.Ordinal))
 					{
-						if (!TryReadToken(Reader, Token) || Token[0] != '\"')
+						// @todo We support <> includes elsewhere but only "" ones here - should we add a check for < or remove the <> elsewhere?
+						if (!TryReadToken(reader, tempStorage) || tempStorage[0] != '\"')
 						{
 							break;
 						}
 
-						string IncludeFile = Token.ToString();
+						ReadOnlySpan<char> tokenSpan = GetSpanFromStringBuilder(tempStorage);
 
-						if (!IncludeFile.EndsWith(".h\"", StringComparison.Ordinal) && !IncludeFile.EndsWith(".h>", StringComparison.Ordinal))
+						if (!tokenSpan.EndsWith(".h\"", StringComparison.Ordinal) && !tokenSpan.EndsWith(".h>", StringComparison.Ordinal))
 						{
 							break;
 						}
-						if (IncludeFile.Equals("\"RequiredProgramMainCPPInclude.h\"", StringComparison.OrdinalIgnoreCase))
+						if (tokenSpan.Equals("\"RequiredProgramMainCPPInclude.h\"", StringComparison.OrdinalIgnoreCase))
 						{
 							break;
 						}
-						if (!TryReadToken(Reader, Token) || Token.ToString() != "\n")
+
+						int writerLength = writer.Length;
+						writer.Append("#include ").Append(tokenSpan).AppendLine();
+
+						if (!TryReadToken(reader, tempStorage) || tempStorage[0] != '\n')
 						{
+							// Undo the previous append as the above condition negates it.
+							writer.Length = writerLength;
 							break;
 						}
-
-						Writer.Append("#include ").AppendLine(IncludeFile);
 					}
 					else
 					{
@@ -76,20 +81,32 @@ namespace UnrealBuildTool
 			}
 		}
 
+		private static ReadOnlySpan<char> GetSpanFromStringBuilder(StringBuilder stringBuilder)
+		{
+			// If the headers fit into the initial allocation of the StringBuilder, we can just use its internal memory
+			// rather than allocating more memory.
+			StringBuilder.ChunkEnumerator ChunkEnumerator = stringBuilder.GetChunks();
+			ChunkEnumerator.MoveNext();
+			ReadOnlyMemory<char> MaybeOnlyChunk = ChunkEnumerator.Current;
+			// MoveNext will be false if there's only one chunk i.e. our one chunk contains all the contents.
+			ReadOnlySpan<char> TokenSpan = ChunkEnumerator.MoveNext() ? stringBuilder.ToString().AsSpan() : MaybeOnlyChunk.Span;
+			return TokenSpan;
+		}
+
 		/// <summary>
 		/// Reads an individual token from the input stream
 		/// </summary>
-		/// <param name="Reader">Stream to read from</param>
-		/// <param name="Token">Buffer to store token read from the stream</param>
+		/// <param name="reader">Stream to read from</param>
+		/// <param name="token">Buffer to store token read from the stream</param>
 		/// <returns>True if a token was read, false otherwise</returns>
-		static bool TryReadToken(StreamReader Reader, StringBuilder Token)
+		private static bool TryReadToken(StreamReader reader, StringBuilder token)
 		{
-			Token.Clear();
+			token.Clear();
 
 			int NextChar;
 			for (; ; )
 			{
-				NextChar = Reader.Read();
+				NextChar = reader.Read();
 				if (NextChar == -1)
 				{
 					return false;
@@ -100,12 +117,12 @@ namespace UnrealBuildTool
 					{
 						break;
 					}
-					else if (Reader.Peek() == '/')
+					else if (reader.Peek() == '/')
 					{
-						Reader.Read();
+						reader.Read();
 						for (; ; )
 						{
-							NextChar = Reader.Read();
+							NextChar = reader.Read();
 							if (NextChar == -1)
 							{
 								return false;
@@ -116,22 +133,22 @@ namespace UnrealBuildTool
 							}
 						}
 					}
-					else if (Reader.Peek() == '*')
+					else if (reader.Peek() == '*')
 					{
-						Reader.Read();
+						reader.Read();
 						for (; ; )
 						{
-							NextChar = Reader.Read();
+							NextChar = reader.Read();
 							if (NextChar == -1)
 							{
 								return false;
 							}
-							if (NextChar == '*' && Reader.Peek() == '/')
+							if (NextChar == '*' && reader.Peek() == '/')
 							{
 								break;
 							}
 						}
-						Reader.Read();
+						reader.Read();
 					}
 					else
 					{
@@ -140,30 +157,30 @@ namespace UnrealBuildTool
 				}
 			}
 
-			Token.Append((char)NextChar);
+			token.Append((char)NextChar);
 
 			if (Char.IsLetterOrDigit((char)NextChar))
 			{
 				for (; ; )
 				{
-					NextChar = Reader.Read();
+					NextChar = reader.Read();
 					if (NextChar == -1 || !Char.IsLetterOrDigit((char)NextChar))
 					{
 						break;
 					}
-					Token.Append((char)NextChar);
+					token.Append((char)NextChar);
 				}
 			}
 			else if (NextChar == '\"' || NextChar == '<')
 			{
 				for (; ; )
 				{
-					NextChar = Reader.Read();
+					NextChar = reader.Read();
 					if (NextChar == -1)
 					{
 						break;
 					}
-					Token.Append((char)NextChar);
+					token.Append((char)NextChar);
 					if (NextChar == '\"' || NextChar == '>')
 					{
 						break;
