Changed AdditionalPlugins to go to the <Project>Editor project, not the non-existant <Ptoject> project, when using bMakeProjectPerTarget

#Synced-CL 50611668

#Synced-CL 52482450

diff --git a/Engine/Source/Programs/UnrealBuildTool/ProjectFiles/ProjectFileGenerator.cs b/Engine/Source/Programs/UnrealBuildTool/ProjectFiles/ProjectFileGenerator.cs
--- a/Engine/Source/Programs/UnrealBuildTool/ProjectFiles/ProjectFileGenerator.cs
+++ b/Engine/Source/Programs/UnrealBuildTool/ProjectFiles/ProjectFileGenerator.cs
@@ -3069,7 +3069,14 @@ namespace UnrealBuildTool
 				{
 					foreach (FileReference ReferencingProject in PluginData.ReferencingProjects)
 					{
-						ProjectFile ProjectFile = FindOrAddProjectHelper(ReferencingProject.GetFileNameWithoutExtension(), ReferencingProject.Directory);
+						string ProjectToAddTo = ReferencingProject.GetFileNameWithoutExtension();
+						// if we are making a project for each target, there won't be a project named for the game, so add it to the Editor project
+						if (bMakeProjectPerTarget)
+						{
+							string? DefaultEditorTarget = GetProjectDefaultTargetNameForType(ReferencingProject.Directory, TargetType.Editor);
+							ProjectToAddTo = DefaultEditorTarget ?? ProjectToAddTo + "Editor";
+						}
+						ProjectFile ProjectFile = FindOrAddProjectHelper(ProjectToAddTo, ReferencingProject.Directory);
 						FileReference? GameUProject = ProjectFileToUProjectFile.GetValueOrDefault(ProjectFile);
 						if (GameUProject != null)
 						{
@@ -3305,7 +3312,15 @@ namespace UnrealBuildTool
 			{
 				FileReference ProjectFileRef = PluginData.ReferencingProjects[0];
 				BaseFolder = ProjectFileRef.Directory;
-				FoundProjects.Add(FindOrAddProjectHelper(ProjectFileRef.GetFileNameWithoutExtension(), BaseFolder));
+
+				string ProjectToAddTo = ProjectFileRef.GetFileNameWithoutExtension();
+				// if we are making a project for each target, there won't be a project named for the game, so add it to the Editor project
+				if (bMakeProjectPerTarget)
+				{
+					string? DefaultEditorTarget = GetProjectDefaultTargetNameForType(BaseFolder, TargetType.Editor);
+					ProjectToAddTo = DefaultEditorTarget ?? ProjectToAddTo + "Editor";
+				}
+				FoundProjects.Add(FindOrAddProjectHelper(ProjectToAddTo, BaseFolder));
 				return FoundProjects;
 			}
 
