UnrealBuildTool: Remove the duplicate logic for handling UBA traces from the UBA executor.

#Synced-CL 52536563

diff --git a/Engine/Source/Programs/Shared/EpicGames.UBA/ITrace.cs b/Engine/Source/Programs/Shared/EpicGames.UBA/ITrace.cs
--- a/Engine/Source/Programs/Shared/EpicGames.UBA/ITrace.cs
+++ b/Engine/Source/Programs/Shared/EpicGames.UBA/ITrace.cs
@@ -9,6 +9,11 @@ namespace EpicGames.UBA
 	/// </summary>
 	public interface ITrace : IBaseInterface
 	{
+		/// <summary>
+		/// The path to the trace file on disk, if any.
+		/// </summary>
+		public abstract string? Path { get; }
+
 		/// <summary>
 		/// Begin task. Will create a bar in uba visualizer
 		/// </summary>
@@ -18,7 +23,7 @@ namespace EpicGames.UBA
 		public abstract uint TaskBegin(string description, string details);
 
 		/// <summary>
-		/// Add hint to task.. will show with milliseconds since start or previous hint
+		/// Add hint to task. will show with milliseconds since start or previous hint
 		/// </summary>
 		/// <param name="taskId">Task id returned by TaskBegin</param>
 		/// <param name="hint">Hint text</param>
diff --git a/Engine/Source/Programs/Shared/EpicGames.UBA/Impl/TraceImpl.cs b/Engine/Source/Programs/Shared/EpicGames.UBA/Impl/TraceImpl.cs
--- a/Engine/Source/Programs/Shared/EpicGames.UBA/Impl/TraceImpl.cs
+++ b/Engine/Source/Programs/Shared/EpicGames.UBA/Impl/TraceImpl.cs
@@ -90,6 +90,8 @@ namespace EpicGames.UBA
 		}
 		#endregion
 
+		public string? Path => _traceFile;
+
 		public nint GetHandle() => _handle;
 
 		public uint TaskBegin(string description, string details) => Trace_TaskBegin(_handle, description, details);
@@ -105,6 +107,7 @@ namespace EpicGames.UBA
 	{
 		public nint GetHandle() => 0;
 		public void Dispose() { }
+		public string? Path => null;
 		public uint TaskBegin(string description, string details) => 0;
 		public void TaskHint(uint taskId, string hint) { }
 		public void TaskEnd(uint taskId) { }
diff --git a/Engine/Source/Programs/UnrealBuildTool/Executors/UnrealBuildAccelerator/UBAExecutor.cs b/Engine/Source/Programs/UnrealBuildTool/Executors/UnrealBuildAccelerator/UBAExecutor.cs
--- a/Engine/Source/Programs/UnrealBuildTool/Executors/UnrealBuildAccelerator/UBAExecutor.cs
+++ b/Engine/Source/Programs/UnrealBuildTool/Executors/UnrealBuildAccelerator/UBAExecutor.cs
@@ -309,27 +309,12 @@ namespace UnrealBuildTool
 
 			DirectoryReference.CreateDirectory(_rootDirRef);
 
-			FileReference ubaTraceFile;
-			if (!String.IsNullOrEmpty(UBAConfig.TraceFile))
+			string? ubaTraceFile = ITrace.GlobalTrace?.Path;
+			if (ubaTraceFile is null)
 			{
-				ubaTraceFile = new FileReference(UBAConfig.TraceFile);
-			}
-			else if (Unreal.IsBuildMachine())
-			{
-				ubaTraceFile = FileReference.Combine(Unreal.EngineProgramSavedDirectory, "AutomationTool", "Saved", "Logs", "Trace.uba");
-			}
-			else if (Log.OutputFile != null)
-			{
-				ubaTraceFile = Log.OutputFile.ChangeExtension(".uba");
-			}
-			else
-			{
-				ubaTraceFile = FileReference.Combine(Unreal.EngineProgramSavedDirectory, "UnrealBuildTool", "Trace.uba");
+				throw new Exception("UBA executor is not expected to be invoked from a recursive UBT call.");
 			}
 
-			DirectoryReference.CreateDirectory(ubaTraceFile.Directory);
-			Log.BackupLogFile(ubaTraceFile);
-
 			DirectoryReference ubaBinariesDir = GetUbaBinariesDir(GetHostPlatform());
 
 			FileReference configFile = FileReference.Combine(ubaBinariesDir, "UbaHost.toml");
@@ -368,7 +353,7 @@ namespace UnrealBuildTool
 			// SessionServerImpl iterates Utils.CrossArchitecturePaths in its constructor,
 			// so these paths must be registered before CreateSessionServer() is called.
 			if (OperatingSystem.IsMacOS())
-			{ 				
+			{
 				AppleToolChain.RegisterCrossArchitecturePathsForUBA(logger);
 				_config.AddValue("Scheduler", "DeferRemoteCapableProcesses", "true");
 			}
@@ -381,7 +366,7 @@ namespace UnrealBuildTool
 			EpicGames.UBA.ILogger ubaLogger = EpicGames.UBA.ILogger.GlobalLogger!;
 			_server = IServer.CreateServer(UBAConfig.MaxWorkers, UBAConfig.SendSize, ubaLogger, UBAConfig.bUseQuic);
 			_storageServer = IStorageServer.CreateStorageServer(_server, ubaLogger, new StorageServerCreateInfo(_rootDirRef.FullName, ((ulong)UBAConfig.StoreCapacityGb) * 1024 * 1024 * 1024, !UBAConfig.bStoreRaw, UBAConfig.Zone));
-			using ISessionServerCreateInfo serverCreateInfo = ISessionServerCreateInfo.CreateSessionServerCreateInfo(_storageServer, _server, ubaLogger, new SessionServerCreateInfo(_rootDirRef.FullName, ubaTraceFile.FullName.Replace('\\', '/'), UBAConfig.bDisableCustomAlloc, false, UBAConfig.bResetCas, UBAConfig.bWriteToDisk, UBAConfig.bDetailedTrace, !UBAConfig.bDisableWaitOnMem, UBAConfig.bAllowKillOnMem, UBAConfig.bStoreObjFilesCompressed));
+			using ISessionServerCreateInfo serverCreateInfo = ISessionServerCreateInfo.CreateSessionServerCreateInfo(_storageServer, _server, ubaLogger, new SessionServerCreateInfo(_rootDirRef.FullName, ubaTraceFile.Replace('\\', '/'), UBAConfig.bDisableCustomAlloc, false, UBAConfig.bResetCas, UBAConfig.bWriteToDisk, UBAConfig.bDetailedTrace, !UBAConfig.bDisableWaitOnMem, UBAConfig.bAllowKillOnMem, UBAConfig.bStoreObjFilesCompressed));
 			_session = ISessionServer.CreateSessionServer(serverCreateInfo);
 			_trace = _session.GetTrace();
 
diff --git a/Engine/Source/Programs/UnrealBuildTool/UnrealBuildTool.cs b/Engine/Source/Programs/UnrealBuildTool/UnrealBuildTool.cs
--- a/Engine/Source/Programs/UnrealBuildTool/UnrealBuildTool.cs
+++ b/Engine/Source/Programs/UnrealBuildTool/UnrealBuildTool.cs
@@ -215,10 +215,6 @@ namespace UnrealBuildTool
 				{
 					traceFileRef = FileReference.Combine(Unreal.EngineProgramSavedDirectory, "AutomationTool", "Saved", "Logs", "Trace.uba");
 				}
-				else if (Log.OutputFile != null)
-				{
-					traceFileRef = Log.OutputFile.ChangeExtension(".uba");
-				}
 				else
 				{
 					traceFileRef = FileReference.Combine(Unreal.EngineProgramSavedDirectory, "UnrealBuildTool", "Trace.uba");
