[NuGet] Update System.IdentityModel.Tokens.Jwt to 8.17.0, System.Data.SqlClient to 4.9.1, and various System packages to 10.0.5

#rnx
#Synced-CL 52554515

diff --git a/Engine/Source/Programs/AutomationTool/Apple/Apple.Automation.csproj b/Engine/Source/Programs/AutomationTool/Apple/Apple.Automation.csproj
--- a/Engine/Source/Programs/AutomationTool/Apple/Apple.Automation.csproj
+++ b/Engine/Source/Programs/AutomationTool/Apple/Apple.Automation.csproj
@@ -19,7 +19,7 @@
   <ItemGroup>
 
 	<PackageReference Include="Microsoft.IdentityModel.Tokens" Version="8.17.0" />
-	<PackageReference Include="System.IdentityModel.Tokens.Jwt" Version="6.36.0" />
+	<PackageReference Include="System.IdentityModel.Tokens.Jwt" Version="8.17.0" />
 	<Reference Include="BouncyCastle.Crypto, Version=1.8.3.0, Culture=neutral, PublicKeyToken=0e99375e54769942">
 		<HintPath>..\..\..\..\Binaries\ThirdParty\IOS\BouncyCastle.Crypto.dll</HintPath>
 	</Reference>
diff --git a/Engine/Source/Programs/AutomationTool/AutomationUtils/AutomationUtils.Automation.csproj b/Engine/Source/Programs/AutomationTool/AutomationUtils/AutomationUtils.Automation.csproj
--- a/Engine/Source/Programs/AutomationTool/AutomationUtils/AutomationUtils.Automation.csproj
+++ b/Engine/Source/Programs/AutomationTool/AutomationUtils/AutomationUtils.Automation.csproj
@@ -40,7 +40,7 @@
     <PackageReference Include="Magick.NET-Q16-HDRI-AnyCPU" Version="14.11.1" PrivateAssets="all" />
     <PackageReference Include="System.Drawing.Common" Version="10.0.5" />
     <!--Required by Ionic -->
-    <PackageReference Include="System.Security.Permissions" Version="8.0.0" />
+    <PackageReference Include="System.Security.Permissions" Version="10.0.5" />
 
     <PackageReference Include="Google.Apis.Drive.v3" Version="1.33.1.1232" PrivateAssets="all" />
     <PackageReference Include="Google.Apis.Sheets.v4" Version="1.33.1.1240" PrivateAssets="all" />
diff --git a/Engine/Source/Programs/AutomationTool/IOS/IOS.Automation.csproj b/Engine/Source/Programs/AutomationTool/IOS/IOS.Automation.csproj
--- a/Engine/Source/Programs/AutomationTool/IOS/IOS.Automation.csproj
+++ b/Engine/Source/Programs/AutomationTool/IOS/IOS.Automation.csproj
@@ -35,7 +35,7 @@
     <ProjectReference Include="..\Apple\Apple.Automation.csproj" />
   </ItemGroup>
   <ItemGroup>
-    <PackageReference Include="System.Resources.Extensions" Version="4.7.1" PrivateAssets="all" />
+    <PackageReference Include="System.Resources.Extensions" Version="10.0.5" PrivateAssets="all" />
     <PackageReference Include="Portable.BouncyCastle" Version="1.8.10" PrivateAssets="all" />
   </ItemGroup>
 </Project>
\ No newline at end of file
diff --git a/Engine/Source/Programs/Horde/Drivers/JobDriver/JobDriver.csproj b/Engine/Source/Programs/Horde/Drivers/JobDriver/JobDriver.csproj
--- a/Engine/Source/Programs/Horde/Drivers/JobDriver/JobDriver.csproj
+++ b/Engine/Source/Programs/Horde/Drivers/JobDriver/JobDriver.csproj
@@ -70,8 +70,8 @@
     <PackageReference Include="Serilog.Settings.Configuration" Version="8.0.4" />
     <PackageReference Include="Serilog.Sinks.Console" Version="5.0.1" />
     <PackageReference Include="Serilog.Sinks.File" Version="5.0.0" />
-    <PackageReference Include="System.Diagnostics.PerformanceCounter" Version="8.0.1" />
-    <PackageReference Include="System.ServiceProcess.ServiceController" Version="8.0.1" />
+    <PackageReference Include="System.Diagnostics.PerformanceCounter" Version="10.0.5" />
+    <PackageReference Include="System.ServiceProcess.ServiceController" Version="10.0.5" />
   </ItemGroup>
   
   <ItemGroup>
diff --git a/Engine/Source/Programs/Horde/HordeAgent/HordeAgent.csproj b/Engine/Source/Programs/Horde/HordeAgent/HordeAgent.csproj
--- a/Engine/Source/Programs/Horde/HordeAgent/HordeAgent.csproj
+++ b/Engine/Source/Programs/Horde/HordeAgent/HordeAgent.csproj
@@ -58,8 +58,8 @@
     <PackageReference Include="Serilog.Settings.Configuration" Version="8.0.4" />
     <PackageReference Include="Serilog.Sinks.Console" Version="5.0.1" />
     <PackageReference Include="Serilog.Sinks.File" Version="5.0.0" />
-    <PackageReference Include="System.Diagnostics.PerformanceCounter" Version="8.0.1" />
-    <PackageReference Include="System.ServiceProcess.ServiceController" Version="8.0.1" />
+    <PackageReference Include="System.Diagnostics.PerformanceCounter" Version="10.0.5" />
+    <PackageReference Include="System.ServiceProcess.ServiceController" Version="10.0.5" />
   </ItemGroup>
 
   <ItemGroup>
diff --git a/Engine/Source/Programs/Horde/HordeServer/HordeServer.csproj b/Engine/Source/Programs/Horde/HordeServer/HordeServer.csproj
--- a/Engine/Source/Programs/Horde/HordeServer/HordeServer.csproj
+++ b/Engine/Source/Programs/Horde/HordeServer/HordeServer.csproj
@@ -61,7 +61,7 @@
     <PackageReference Include="SharpZipLib" Version="1.4.2" />
     <PackageReference Include="StackExchange.Redis" Version="2.8.58" />
     <PackageReference Include="Swashbuckle.AspNetCore" Version="6.5.0" />
-    <PackageReference Include="System.Data.Odbc" Version="9.0.10" />
+    <PackageReference Include="System.Data.Odbc" Version="10.0.5" />
     <PackageReference Include="T-Digest.NET" Version="1.0.3" />
     <PackageReference Include="TimeZoneConverter" Version="6.1.0" />
     <PackageReference Include="ZstdSharp.Port" Version="0.8.6" />
diff --git a/Engine/Source/Programs/Horde/Plugins/PerformanceTrends/HordeServer.PerformanceTrends/HordeServer.PerformanceTrends.csproj b/Engine/Source/Programs/Horde/Plugins/PerformanceTrends/HordeServer.PerformanceTrends/HordeServer.PerformanceTrends.csproj
--- a/Engine/Source/Programs/Horde/Plugins/PerformanceTrends/HordeServer.PerformanceTrends/HordeServer.PerformanceTrends.csproj
+++ b/Engine/Source/Programs/Horde/Plugins/PerformanceTrends/HordeServer.PerformanceTrends/HordeServer.PerformanceTrends.csproj
@@ -5,7 +5,7 @@
   </PropertyGroup>
   <ItemGroup>
     <PackageReference Include="Polly" Version="8.2.1" />
-    <PackageReference Include="System.Data.Odbc" Version="9.0.10" />
+    <PackageReference Include="System.Data.Odbc" Version="10.0.5" />
   </ItemGroup>
   <ItemGroup>
     <ProjectReference Include="..\..\..\..\Shared\EpicGames.Analytics.PerformanceTrends\EpicGames.Analytics.PerformanceTrends.csproj" />
diff --git a/Engine/Source/Programs/Horde/Plugins/StructuredAnalytics/HordeServer.StructuredAnalytics/HordeServer.StructuredAnalytics.csproj b/Engine/Source/Programs/Horde/Plugins/StructuredAnalytics/HordeServer.StructuredAnalytics/HordeServer.StructuredAnalytics.csproj
--- a/Engine/Source/Programs/Horde/Plugins/StructuredAnalytics/HordeServer.StructuredAnalytics/HordeServer.StructuredAnalytics.csproj
+++ b/Engine/Source/Programs/Horde/Plugins/StructuredAnalytics/HordeServer.StructuredAnalytics/HordeServer.StructuredAnalytics.csproj
@@ -5,7 +5,7 @@
   </PropertyGroup>
   <ItemGroup>
     <PackageReference Include="Polly" Version="8.2.1" />
-    <PackageReference Include="System.Data.Odbc" Version="9.0.10" />
+    <PackageReference Include="System.Data.Odbc" Version="10.0.5" />
   </ItemGroup>
   <ItemGroup>
     <ProjectReference Include="..\..\..\..\Shared\EpicGames.Build\EpicGames.Build.csproj" />
diff --git a/Engine/Source/Programs/UnrealBuildTool/UnrealBuildTool.csproj b/Engine/Source/Programs/UnrealBuildTool/UnrealBuildTool.csproj
--- a/Engine/Source/Programs/UnrealBuildTool/UnrealBuildTool.csproj
+++ b/Engine/Source/Programs/UnrealBuildTool/UnrealBuildTool.csproj
@@ -73,12 +73,12 @@
     <PackageReference Include="Microsoft.VisualStudio.Setup.Configuration.Interop" Version="3.14.2075" />
     <PackageReference Include="Microsoft.VisualStudio.SolutionPersistence" Version="1.0.52" />
     <PackageReference Include="OpenTracing" Version="0.12.1" />
-    <PackageReference Include="System.CodeDom" Version="8.0.0" />
-    <PackageReference Include="System.Management" Version="8.0.0" />
-    <PackageReference Include="System.ServiceProcess.ServiceController" Version="8.0.1" />
+    <PackageReference Include="System.CodeDom" Version="10.0.5" />
+    <PackageReference Include="System.Management" Version="10.0.5" />
+    <PackageReference Include="System.ServiceProcess.ServiceController" Version="10.0.5" />
 
     <!-- Required by Ionic -->
-    <PackageReference Include="System.Security.Permissions" Version="8.0.0" />
+    <PackageReference Include="System.Security.Permissions" Version="10.0.5" />
   </ItemGroup>
 
 </Project>
