UnrealBuildTool: Add C# automation dependency and test projects in a loop, as tests may introduce dependencies and dependencies may introduce tests projects.

#Synced-CL 52569840

diff --git a/Engine/Source/Programs/UnrealBuildTool/ProjectFiles/ProjectFileGenerator.cs b/Engine/Source/Programs/UnrealBuildTool/ProjectFiles/ProjectFileGenerator.cs
--- a/Engine/Source/Programs/UnrealBuildTool/ProjectFiles/ProjectFileGenerator.cs
+++ b/Engine/Source/Programs/UnrealBuildTool/ProjectFiles/ProjectFileGenerator.cs
@@ -1564,20 +1564,27 @@ namespace UnrealBuildTool
 			Task.WaitAll([.. DotnetProjectTasks]);
 
 			PrimaryProjectFolder DependenciesFolder;
+			PrimaryProjectFolder TestsFolder;
 			lock (RootFolder)
 			{
 				DependenciesFolder = RootFolder.AddSubFolder("Dependencies");
+				TestsFolder = RootFolder.AddSubFolder("Tests");
 			}
 
-			while (DotnetProjectTasks.Count != 0)
+			bool runLoop = true;
+			// Tests may introduce extra dependencies, dependencies may introduce extra tests, so run both in a loop.
+			while (runLoop)
 			{
+				runLoop = false;
+
 				DotnetProjectTasks = GetDotnetDependencies(RootFolder, DependenciesFolder, Logger);
 				Task.WaitAll([.. DotnetProjectTasks]);
-			}
+				runLoop |= DotnetProjectTasks.Count > 0;
 
-			PrimaryProjectFolder TestsFolder = RootFolder.AddSubFolder("Tests");
-			DotnetProjectTasks = GetDotnetTests(RootFolder, TestsFolder);
-			Task.WaitAll([.. DotnetProjectTasks]);
+				DotnetProjectTasks = GetDotnetTests(RootFolder, TestsFolder);
+				Task.WaitAll([.. DotnetProjectTasks]);
+				runLoop |= DotnetProjectTasks.Count > 0;
+			}
 		}
 
 		/// <summary>
