UnrealBuildTool: Restrict Clang JSON compilation databases to TU source files.

#12745

#Synced-CL 52569840

diff --git a/Engine/Source/Programs/UnrealBuildTool/ProjectFiles/VisualStudioCode/VSCodeProjectFileGenerator.cs b/Engine/Source/Programs/UnrealBuildTool/ProjectFiles/VisualStudioCode/VSCodeProjectFileGenerator.cs
--- a/Engine/Source/Programs/UnrealBuildTool/ProjectFiles/VisualStudioCode/VSCodeProjectFileGenerator.cs
+++ b/Engine/Source/Programs/UnrealBuildTool/ProjectFiles/VisualStudioCode/VSCodeProjectFileGenerator.cs
@@ -887,7 +887,10 @@ namespace UnrealBuildTool
 					DirectoryToResponseFile.Add(Pair.Key, ResponseFile);
 				}
 
-				foreach (FileReference File in SourceFiles.OrderBy(x => x.FullName))
+				IEnumerable<FileReference> SourceFilesForClangCompilationDatabase = SourceFiles
+					.Where(x => x.HasExtension(".cpp") || x.HasExtension(".cc") || x.HasExtension(".cxx") || x.HasExtension(".c"))
+					.OrderBy(x => x.FullName);
+				foreach (FileReference File in SourceFilesForClangCompilationDatabase)
 				{
 					DirectoryReference Directory = File.Directory;
 
