UnrealBuildTool: Fix the VSCode project generator's `terminal.integrated.env` setting for Mac.

#13983

#Synced-CL 52578659

diff --git a/Engine/Source/Programs/UnrealBuildTool/ProjectFiles/VisualStudioCode/VSCodeProjectFileGenerator.cs b/Engine/Source/Programs/UnrealBuildTool/ProjectFiles/VisualStudioCode/VSCodeProjectFileGenerator.cs
--- a/Engine/Source/Programs/UnrealBuildTool/ProjectFiles/VisualStudioCode/VSCodeProjectFileGenerator.cs
+++ b/Engine/Source/Programs/UnrealBuildTool/ProjectFiles/VisualStudioCode/VSCodeProjectFileGenerator.cs
@@ -1737,7 +1737,7 @@ namespace UnrealBuildTool
 				// disable npm script autodetect to avoid lag populating tasks list 
 				WorkspaceFile.AddField("npm.autoDetect", "off");
 				// set environment variables so correct dotnet is used
-				WorkspaceFile.BeginObject($"terminal.integrated.env.{(OperatingSystem.IsWindows() ? "windows" : OperatingSystem.IsMacOS() ? "mac" : "linux")}");
+				WorkspaceFile.BeginObject($"terminal.integrated.env.{(OperatingSystem.IsWindows() ? "windows" : OperatingSystem.IsMacOS() ? "osx" : "linux")}");
 				{
 					WorkspaceFile.AddField("PATH", $"{Unreal.DotnetDirectory}{Path.PathSeparator}${{env:PATH}}");
 					WorkspaceFile.AddField("DOTNET_ROOT", Unreal.DotnetDirectory.FullName);
