UnrealBuildTool: Fix the MSVC unused parameter warning.

#14465

#Synced-CL 52578659

diff --git a/Engine/Source/Programs/UnrealBuildTool/Configuration/Rules/CppCompileWarnings.cs b/Engine/Source/Programs/UnrealBuildTool/Configuration/Rules/CppCompileWarnings.cs
--- a/Engine/Source/Programs/UnrealBuildTool/Configuration/Rules/CppCompileWarnings.cs
+++ b/Engine/Source/Programs/UnrealBuildTool/Configuration/Rules/CppCompileWarnings.cs
@@ -906,7 +906,7 @@ namespace UnrealBuildTool
 		///		https://clang.llvm.org/docs/DiagnosticsReference.html#wunused-parameter
 		/// </summary>
 		/// <remarks>Unused function parameter. A lot are named 'bUnused'...</remarks>
-		[WarningsMSCV(["/wd4100"], ["/w4100"], ["/we4100"])]
+		[WarningsMSCV(["/wd4100"], ["/w44100"], ["/we4100"])]
 		[WarningsVCClang(["-Wno-unused-parameter"], ["-Wunused-parameter", "-Wno-error=unused-parameter"], ["-Wunused-parameter"])]
 		[BasicWarningLevelDefault(WarningLevel.Off)]
 		public WarningLevel UnusedParameterWarningLevel
