Reduced linking time on clang with LTO by disabling debug info upgrading, sice we build everything with the same toolchain

#Synced-CL 51881412
#Synced-CL 51881412

diff --git a/Engine/Source/Programs/UnrealBuildTool/Platform/Clang/ClangToolChain.Link.cs b/Engine/Source/Programs/UnrealBuildTool/Platform/Clang/ClangToolChain.Link.cs
--- a/Engine/Source/Programs/UnrealBuildTool/Platform/Clang/ClangToolChain.Link.cs
+++ b/Engine/Source/Programs/UnrealBuildTool/Platform/Clang/ClangToolChain.Link.cs
@@ -93,6 +93,8 @@ namespace UnrealBuildTool
 				{
 					arguments.Add("-flto");
 				}
+
+				AddDisableAutoUpgradeDebugInfoFlags(arguments);
 			}
 		}
 
@@ -174,5 +176,11 @@ namespace UnrealBuildTool
 
 			GetLinkArguments_WarningsAndErrors(linkEnvironment, arguments);
 		}
+
+		protected virtual void AddDisableAutoUpgradeDebugInfoFlags(List<string> Arguments)
+		{
+			// we build everything with the same toolchain, so we don't want to upgrade debug info because it runs validation and is expensive
+			Arguments.Add("-Wl,-mllvm,-disable-auto-upgrade-debug-info");
+		}
 	}
 }
diff --git a/Engine/Source/Programs/UnrealBuildTool/Platform/Linux/LinuxToolChain.cs b/Engine/Source/Programs/UnrealBuildTool/Platform/Linux/LinuxToolChain.cs
--- a/Engine/Source/Programs/UnrealBuildTool/Platform/Linux/LinuxToolChain.cs
+++ b/Engine/Source/Programs/UnrealBuildTool/Platform/Linux/LinuxToolChain.cs
@@ -642,7 +642,7 @@ namespace UnrealBuildTool
 
 			if (Options.HasFlag(ClangToolChainOptions.EnableLinkTimeOptimization) && Options.HasFlag(ClangToolChainOptions.EnableThinLTO))
 			{
-				arguments.Add("-Wl,-mllvm,-disable-auto-upgrade-debug-info, -Wl,-mllvm,-enable-ext-tsp-block-placement=1");
+				arguments.Add("-Wl,-mllvm,-enable-ext-tsp-block-placement=1");
 			}
 		}
 
diff --git a/Engine/Source/Programs/UnrealBuildTool/Platform/Windows/VCToolChain.cs b/Engine/Source/Programs/UnrealBuildTool/Platform/Windows/VCToolChain.cs
--- a/Engine/Source/Programs/UnrealBuildTool/Platform/Windows/VCToolChain.cs
+++ b/Engine/Source/Programs/UnrealBuildTool/Platform/Windows/VCToolChain.cs
@@ -3845,7 +3845,7 @@ namespace UnrealBuildTool
 				{
 					if (Target.WindowsPlatform.bAllowClangLinker)
 					{
-						if (LinkEnvironment.bPGOOptimize || LinkEnvironment.bPGOProfile)
+						if (bPGOOptimize || bPGOProfile)
 						{
 							Log.TraceInformationOnce("Enabling Link-time optimization as Profile Guided Optimization (PGO) was requested. Linking will take a while.");
 						}
@@ -3857,6 +3857,9 @@ namespace UnrealBuildTool
 						// lld should consider logical cores when determining how many threads to use
 						Arguments.Add("/opt:lldltojobs=all");
 
+						// we build everything with the same toolchain, so we don't want to upgrade debug info because it runs validation and is expensive
+						Arguments.Add("-mllvm:-disable-auto-upgrade-debug-info");
+
 						// ThinLTO incremental cache
 						DirectoryReference? ThinLTOCacheDir = DirectoryReference.FromString(LinkEnvironment.ThinLTOCacheDirectory);
 						if (ThinLTOCacheDir != null)
