PR #14659: PVS_VERSION_MAJOR and PVS_VERSION_MINOR defines.

#Synced-CL 52566575

diff --git a/Engine/Source/Programs/UnrealBuildTool/Platform/Windows/PVSToolChain.cs b/Engine/Source/Programs/UnrealBuildTool/Platform/Windows/PVSToolChain.cs
--- a/Engine/Source/Programs/UnrealBuildTool/Platform/Windows/PVSToolChain.cs
+++ b/Engine/Source/Programs/UnrealBuildTool/Platform/Windows/PVSToolChain.cs
@@ -753,6 +753,8 @@ namespace UnrealBuildTool
 			};
 			preprocessCompileEnvironment.AdditionalArguments += " /wd4005 /wd4828 /wd5105";
 			preprocessCompileEnvironment.Definitions.Add("PVS_STUDIO");
+			preprocessCompileEnvironment.Definitions.Add($"PVS_VERSION_MAJOR={AnalyzerVersion.Major}");
+			preprocessCompileEnvironment.Definitions.Add($"PVS_VERSION_MINOR={AnalyzerVersion.Minor}");
 			preprocessCompileEnvironment.CppCompileWarnings.UndefinedIdentifierWarningLevel = WarningLevel.Off; // Not sure why THIRD_PARTY_INCLUDES_START doesn't pick this up; the _Pragma appears in the preprocessed output. Perhaps in preprocess-only mode the compiler doesn't respect these?
 
 			preprocessActions = [];
