PR #14491: [UBT] Implement .pvsconfig support for PVSToolchain

#Synced-CL 52688981

diff --git a/Engine/Source/Programs/UnrealBuildTool/Configuration/Rules/ReadOnlyTargetRules.cs b/Engine/Source/Programs/UnrealBuildTool/Configuration/Rules/ReadOnlyTargetRules.cs
--- a/Engine/Source/Programs/UnrealBuildTool/Configuration/Rules/ReadOnlyTargetRules.cs
+++ b/Engine/Source/Programs/UnrealBuildTool/Configuration/Rules/ReadOnlyTargetRules.cs
@@ -420,6 +420,8 @@ namespace UnrealBuildTool
 
 		public int StaticAnalyzerPVSPrintLevel => Inner.StaticAnalyzerPVSPrintLevel;
 
+		public IReadOnlyList<string> StaticAnalyzerPVSConfigs => Inner.StaticAnalyzerPVSConfigs;
+
 		public bool bStaticAnalyzerProjectOnly => Inner.bStaticAnalyzerProjectOnly;
 
 		public bool bStaticAnalyzerIncludeGenerated => Inner.bStaticAnalyzerIncludeGenerated;
diff --git a/Engine/Source/Programs/UnrealBuildTool/Configuration/Rules/TargetRules.cs b/Engine/Source/Programs/UnrealBuildTool/Configuration/Rules/TargetRules.cs
--- a/Engine/Source/Programs/UnrealBuildTool/Configuration/Rules/TargetRules.cs
+++ b/Engine/Source/Programs/UnrealBuildTool/Configuration/Rules/TargetRules.cs
@@ -2250,6 +2250,12 @@ namespace UnrealBuildTool
 		[XmlConfigFile(Category = "BuildConfiguration")]
 		public int StaticAnalyzerPVSPrintLevel { get; set; } = 1;
 
+		/// <summary>
+		/// Paths to PVS-Studio .pvsconfig configuration files
+		/// </summary>
+		[CommandLine("-StaticAnalyzerPVSConfigs=", ListSeparator = ';')]
+		public List<string> StaticAnalyzerPVSConfigs = [];
+
 		/// <summary>
 		/// Only run static analysis against project modules, skipping engine modules
 		/// </summary>
diff --git a/Engine/Source/Programs/UnrealBuildTool/Platform/Windows/PVSToolChain.cs b/Engine/Source/Programs/UnrealBuildTool/Platform/Windows/PVSToolChain.cs
--- a/Engine/Source/Programs/UnrealBuildTool/Platform/Windows/PVSToolChain.cs
+++ b/Engine/Source/Programs/UnrealBuildTool/Platform/Windows/PVSToolChain.cs
@@ -887,6 +887,39 @@ namespace UnrealBuildTool
 				$"--lic-name \"{ApplicationSettings?.UserName}\" --lic-key \"{ApplicationSettings?.SerialNumber}\"",
 			];
 
+			// PVS-Studio loads global configuration files prior to the ones specified by the user.
+			DirectoryReference? applicationData = DirectoryReference.GetSpecialFolder(Environment.SpecialFolder.ApplicationData);
+			if (applicationData != null && DirectoryReference.Exists(applicationData))
+			{
+				DirectoryReference globalPVSConfigsDir = DirectoryReference.Combine(applicationData, "PVS-Studio");
+
+				if (DirectoryReference.Exists(globalPVSConfigsDir))
+				{
+					foreach (FileReference configPath in DirectoryReference.EnumerateFiles(globalPVSConfigsDir, "*.pvsconfig", SearchOption.TopDirectoryOnly))
+					{
+						arguments.Add($"--rules-config=\"{configPath.FullName}\"");
+						analyzeAction.PrerequisiteItems.Add(FileItem.GetItemByFileReference(configPath));
+					}
+				}
+			}
+
+			foreach (string configPath in Target.StaticAnalyzerPVSConfigs.Where(path => path.EndsWith(".pvsconfig", StringComparison.InvariantCultureIgnoreCase)))
+			{
+				string fullPVSConfigPath = configPath;
+
+				if (!Path.IsPathRooted(fullPVSConfigPath) && Target.ProjectFile != null)
+				{
+					string projectDir = Path.GetDirectoryName(Target.ProjectFile.FullName)!;
+					fullPVSConfigPath = Path.Combine(projectDir, fullPVSConfigPath);
+				}
+
+				if (Path.Exists(fullPVSConfigPath))
+				{
+					arguments.Add($"--rules-config=\"{fullPVSConfigPath}\"");
+					analyzeAction.PrerequisiteItems.Add(FileItem.GetItemByPath(fullPVSConfigPath));
+				}
+			}
+
 			analyzeAction.CommandArguments = String.Join(' ', arguments);
 
 			analyzeAction.PrerequisiteItems.UnionWith(preprocessAction.AdditionalPrerequisiteItems);
