Add NoGameFeatureContent so the cooked build has the correct .uplugin information staged.

#Synced-CL 52608389

diff --git a/Engine/Source/Programs/UnrealBuildTool/Configuration/Descriptors/PluginDescriptor.cs b/Engine/Source/Programs/UnrealBuildTool/Configuration/Descriptors/PluginDescriptor.cs
--- a/Engine/Source/Programs/UnrealBuildTool/Configuration/Descriptors/PluginDescriptor.cs
+++ b/Engine/Source/Programs/UnrealBuildTool/Configuration/Descriptors/PluginDescriptor.cs
@@ -224,6 +224,11 @@ namespace UnrealBuildTool
 		/// </summary>
 		public bool bNoCode;
 
+		/// <summary>
+		/// If true, the GameFeaturePlugin this is on should have no content other than the GameFeatureData asset.
+		/// </summary>
+		public bool bNoGameFeatureContent;
+
 		/// <summary>
 		/// When true, this plugin's modules will not be loaded automatically nor will it's content be mounted automatically. It will load/mount when explicitly requested and LoadingPhases will be ignored
 		/// </summary>
@@ -255,11 +260,17 @@ namespace UnrealBuildTool
 		public string[]? DisallowedPlugins;
 
 		/// <summary>
-		/// The JsonObject created from reading a .uplugin on disk or from parsing a json text 
+		/// The JsonObject created from reading a .uplugin on disk or from parsing a json text
 		/// This preserves the order of all the fields from the source json as well as account for any custom fields.
 		/// </summary>
 		private readonly JsonObject CachedJson;
 
+		/// <summary>
+		/// Read-only access to the raw JSON object for this plugin descriptor.
+		/// Callers can use this to read non-standard fields that are not modelled by this class.
+		/// </summary>
+		public JsonObject RawJson => CachedJson;
+
 		/// <summary>
 		/// Reads a plugin descriptor from a json object
 		/// </summary>
@@ -397,6 +408,7 @@ namespace UnrealBuildTool
 			RawObject.TryGetBoolField("RequiresBuildPlatform", out bRequiresBuildPlatform);
 			RawObject.TryGetBoolField("Sealed", out bIsSealed);
 			RawObject.TryGetBoolField("NoCode", out bNoCode);
+			RawObject.TryGetBoolField("NoGameFeatureContent", out bNoGameFeatureContent);
 			RawObject.TryGetBoolField("ExplicitlyLoaded", out bExplicitlyLoaded);
 			RawObject.TryGetBoolField("HasExplicitPlatforms", out bHasExplicitPlatforms);
 
@@ -542,6 +554,11 @@ namespace UnrealBuildTool
 				Writer.WriteValue("NoCode", bNoCode);
 			}
 
+			if (bNoGameFeatureContent)
+			{
+				Writer.WriteValue("NoGameFeatureContent", bNoGameFeatureContent);
+			}
+
 			if (bExplicitlyLoaded)
 			{
 				Writer.WriteValue("ExplicitlyLoaded", bExplicitlyLoaded);
