Fix to JAVA_HOME Path During Manual Android Sdk Switching on Mac

#android
#Synced-CL 52700724

diff --git a/Engine/Source/Programs/UnrealBuildTool/Platform/Android/AndroidPlatformSDK.cs b/Engine/Source/Programs/UnrealBuildTool/Platform/Android/AndroidPlatformSDK.cs
--- a/Engine/Source/Programs/UnrealBuildTool/Platform/Android/AndroidPlatformSDK.cs
+++ b/Engine/Source/Programs/UnrealBuildTool/Platform/Android/AndroidPlatformSDK.cs
@@ -201,8 +201,8 @@ namespace UnrealBuildTool
 			}
 			else if (OperatingSystem.IsMacOS())
 			{
-				string Directory = Path.Combine(Environment.GetEnvironmentVariable("HOME")!, "Applications/Android Studio.app");
-				return System.IO.Directory.Exists(Directory) ? Directory : "/Applications/Android Studio.app";
+				string Directory = Path.Combine(Environment.GetEnvironmentVariable("HOME")!, "Applications/Android Studio.app/Contents");
+				return System.IO.Directory.Exists(Directory) ? Directory : "/Applications/Android Studio.app/Contents";
 			}
 			else
 			{
@@ -284,6 +284,11 @@ namespace UnrealBuildTool
 						}
 					}
 
+					if (OperatingSystem.IsMacOS())
+					{
+						JavaDir = Path.Combine(JavaDir, "Contents", "Home");
+					}
+
 					EnvVarValues["JAVA_HOME"] = JavaDir;
 					return true;
 				}
