Remove UEDeployAndroid.RunCommandLineProgramWithException as it unused now, which is needed with all Warnings on in Horde

#Synced-CL 52700724

diff --git a/Engine/Source/Programs/UnrealBuildTool/Platform/Android/UEDeployAndroid.cs b/Engine/Source/Programs/UnrealBuildTool/Platform/Android/UEDeployAndroid.cs
--- a/Engine/Source/Programs/UnrealBuildTool/Platform/Android/UEDeployAndroid.cs
+++ b/Engine/Source/Programs/UnrealBuildTool/Platform/Android/UEDeployAndroid.cs
@@ -1684,48 +1684,6 @@ LD_HWASAN=1 exec $cmd";
 			return Proc.ExitCode;
 		}
 
-		private static void RunCommandLineProgramWithException(string WorkingDirectory, string Command, string Params, ILogger Logger, string? OverrideDesc = null, bool bUseShellExecute = false)
-		{
-			// Process Arguments follow windows conventions in .NET Core
-			// Which means single quotes ' are not considered quotes.
-			// see [REDACTED_PROFILE_URL]/runtime/issues/29857
-			// also see UE-102580
-			// for rules see https://docs.microsoft.com/en-us/cpp/cpp/main-function-command-line-args
-			if (Params.Contains('\'', StringComparison.Ordinal))
-			{
-				Params = Params.Replace("\"", "\\\"", StringComparison.Ordinal);
-				Params = Params.Replace('\'', '\"');
-			}
-
-			if (OverrideDesc == null)
-			{
-				Logger.LogInformation("\nRunning: {Command} {Params}", Command, Params);
-			}
-			else if (!String.IsNullOrEmpty(OverrideDesc))
-			{
-				Logger.LogInformation("{Message}", OverrideDesc);
-				Logger.LogDebug("\nRunning: {Command} {Params}", Command, Params);
-			}
-
-			ProcessStartInfo StartInfo = new ProcessStartInfo();
-			StartInfo.WorkingDirectory = WorkingDirectory;
-			StartInfo.FileName = Command;
-			StartInfo.Arguments = Params;
-			StartInfo.UseShellExecute = bUseShellExecute;
-			StartInfo.WindowStyle = ProcessWindowStyle.Minimized;
-
-			Process Proc = new Process();
-			Proc.StartInfo = StartInfo;
-			Proc.Start();
-			Proc.WaitForExit();
-
-			// android bat failure
-			if (Proc.ExitCode != 0)
-			{
-				throw new BuildException("{0} failed with args {1}", Command, Params);
-			}
-		}
-
 		private enum FilterAction
 		{
 			Skip,
