UnrealBuildTool: Fix programs not under the Engine directory failing to compile in Xcode if they're in a path that includes the string "Programs".

#14556

#Synced-CL 52723323

diff --git a/Engine/Source/Programs/UnrealBuildTool/ProjectFiles/Xcode/XcodeProject.File.cs b/Engine/Source/Programs/UnrealBuildTool/ProjectFiles/Xcode/XcodeProject.File.cs
--- a/Engine/Source/Programs/UnrealBuildTool/ProjectFiles/Xcode/XcodeProject.File.cs
+++ b/Engine/Source/Programs/UnrealBuildTool/ProjectFiles/Xcode/XcodeProject.File.cs
@@ -226,7 +226,9 @@ namespace UnrealBuildTool.XcodeProjectXcconfig
 
 		public void SetUProjectLocation(FileReference? UProjectLocation)
 		{
-			RootDirectory = UProjectLocation == null || UProjectLocation.ContainsName("Programs", 0) ? (Unreal.EngineDirectory) : UProjectLocation!.Directory;
+			RootDirectory = UProjectLocation == null || (UProjectLocation.IsUnderDirectory(Unreal.EngineDirectory) && UProjectLocation.ContainsName("Programs", Unreal.EngineDirectory.FullName.Length))
+				? (Unreal.EngineDirectory)
+				: UProjectLocation!.Directory;
 		}
 
 		public void ProcessFile(XcodeSourceFile File, bool bIsForBuild, bool bIsFolder, string? GroupName = null, string ExtraResourceSettings = "", Dictionary<DirectoryReference, int>? SourceToBuildFileMap = null)
