[UBT] Log the location of the finished binary.

This feature request has been backlogged for several months, but it turns out to be one-shottable in Claude.

Results look like this:

```
1>  Compile [x64] Delegate.cpp [ Wall Time 16.28s / CPU Time 15.58s / Mem 721.15 MB ]
1>  Compile [x64] Benchmarks.cpp [ Wall Time 15.64s / CPU Time 15.03s / Mem 636.99 MB ]
1>  Compile [x64] API.cpp [ Wall Time 15.56s / CPU Time 14.98s / Mem 657.23 MB ]
1>  Compile [x64] Class.cpp [ Time 15.31s / Mem 595.21 MB ]
1>  Compile [x64] HttpCacheStore.cpp [ Time 15.19s / Mem 484.76 MB ]
1>  Compile [x64] RWLockTests.cpp [ Wall Time 15.08s / CPU Time 14.58s / Mem 629.14 MB ]
1>
1>  Total time in Unreal Build Accelerator executor: 45.06 seconds
1>  Output binary: E:\Epic\Release-41.00\Engine\Binaries\Win64\AutoRTFMTests.exe
1>
1>  Result: Succeeded
1>  Total execution time: 48.75 seconds
1>  Trace written to file E:\Epic\Release-41.00\Engine\Programs\UnrealBuildTool\Trace.uba with size 459.4kb
========== Build: 1 succeeded, 0 failed, 0 up-to-date, 0 skipped ==========

```

#Synced-CL 52735921

diff --git a/Engine/Source/Programs/UnrealBuildTool/Modes/BuildMode.cs b/Engine/Source/Programs/UnrealBuildTool/Modes/BuildMode.cs
--- a/Engine/Source/Programs/UnrealBuildTool/Modes/BuildMode.cs
+++ b/Engine/Source/Programs/UnrealBuildTool/Modes/BuildMode.cs
@@ -812,6 +812,12 @@ namespace UnrealBuildTool
 						await ActionGraph.ExecuteActionsAsync(BuildConfiguration, MergedActionsToExecute, TargetDescriptors, Logger, executor);
 					}
 
+					// Log the output binary path for each target
+					foreach (TargetMakefile Makefile in Makefiles)
+					{
+						Logger.LogInformation("Output binary: {ExecutableFile}", Makefile.ExecutableFile);
+					}
+
 					// these are parsed by external tools wishing to open this file directly
 					foreach (LinkedAction BuildAction in MergedActionsToExecute.Where(BuildAction => BuildAction.ActionType == ActionType.Compile && (BuildAction.Inner as VCCompileAction)?.bIsAnalyzing != true))
 					{
