UnrealBuildTool: Add AWS Graviton4 (Neoverse V2) CPU target support

#14667

#Synced-CL 52727718

diff --git a/Engine/Source/Programs/UnrealBuildTool/Platform/Clang/ClangToolChain.cs b/Engine/Source/Programs/UnrealBuildTool/Platform/Clang/ClangToolChain.cs
--- a/Engine/Source/Programs/UnrealBuildTool/Platform/Clang/ClangToolChain.cs
+++ b/Engine/Source/Programs/UnrealBuildTool/Platform/Clang/ClangToolChain.cs
@@ -993,6 +993,10 @@ namespace UnrealBuildTool
 						Arguments.Add("-mcpu=neoverse-512tvb+bf16+rng+crypto");
 						Arguments.Add("-mtune=neoverse-v1");
 						break;
+
+					case Arm64TargetCpuType.Graviton4:
+						Arguments.Add("-mcpu=neoverse-v2+crypto");
+						break;
 				}
 			}
 		}
diff --git a/Engine/Source/Programs/UnrealBuildTool/System/CppCompileEnvironment.cs b/Engine/Source/Programs/UnrealBuildTool/System/CppCompileEnvironment.cs
--- a/Engine/Source/Programs/UnrealBuildTool/System/CppCompileEnvironment.cs
+++ b/Engine/Source/Programs/UnrealBuildTool/System/CppCompileEnvironment.cs
@@ -177,9 +177,14 @@ namespace UnrealBuildTool
 		Graviton2,
 
 		/// <summary>
-		/// AWS Graviton3 CPU - ARM Neoverse V1 uArch with AMRv8.4-A ISA + SVE + bf16 + rng + crypto extensions
+		/// AWS Graviton3 CPU - ARM Neoverse V1 uArch with ARMv8.4-A ISA + SVE + bf16 + rng + crypto extensions
 		/// </summary>
 		Graviton3,
+
+		/// <summary>
+		/// AWS Graviton4 CPU - ARM Neoverse V2 uArch with ARMv9.0-A ISA + SVE2 + bf16 + rng + crypto extensions
+		/// </summary>
+		Graviton4,
 	}
 
 	/// <summary>
