UnrealBuildTool: Add GravitonSVE cross-generation ARM64 CPU target for Graviton3/4+

#14674
#Synced-CL 52782323

diff --git a/Engine/Source/Programs/UnrealBuildTool/Platform/Clang/ClangToolChain.cs b/Engine/Source/Programs/UnrealBuildTool/Platform/Clang/ClangToolChain.cs
--- a/Engine/Source/Programs/UnrealBuildTool/Platform/Clang/ClangToolChain.cs
+++ b/Engine/Source/Programs/UnrealBuildTool/Platform/Clang/ClangToolChain.cs
@@ -982,6 +982,9 @@ namespace UnrealBuildTool
 			{
 				Arguments.Add("-fsigned-char");
 
+				// For the Graviton targets, desired optional features are enabled by default,
+				// with the exception of crypto, which may be disabled by default due to export
+				// restrictions.
 				switch (CompileEnvironment.MinArm64CpuTarget)
 				{
 					case Arm64TargetCpuType.Graviton2:
@@ -989,6 +992,10 @@ namespace UnrealBuildTool
 						Arguments.Add("-mtune=neoverse-n1");
 						break;
 
+					case Arm64TargetCpuType.GravitonSVE:
+						Arguments.Add("-mcpu=neoverse-512tvb+crypto");
+						break;
+
 					case Arm64TargetCpuType.Graviton3:
 						Arguments.Add("-mcpu=neoverse-512tvb+bf16+rng+crypto");
 						Arguments.Add("-mtune=neoverse-v1");
diff --git a/Engine/Source/Programs/UnrealBuildTool/System/CppCompileEnvironment.cs b/Engine/Source/Programs/UnrealBuildTool/System/CppCompileEnvironment.cs
--- a/Engine/Source/Programs/UnrealBuildTool/System/CppCompileEnvironment.cs
+++ b/Engine/Source/Programs/UnrealBuildTool/System/CppCompileEnvironment.cs
@@ -176,6 +176,15 @@ namespace UnrealBuildTool
 		/// </summary>
 		Graviton2,
 
+		/// <summary>
+		/// AWS Graviton SVE (Graviton3+) - targets Neoverse 512TVB with SVE + bf16 + crypto.
+		/// Cross-generation target compatible with Graviton3, Graviton4, and future SVE-capable
+		/// Graviton processors. NOT compatible with Graviton2 (no SVE support).
+		/// For Graviton2, use the Graviton2 target. For full Graviton4 optimizations (SVE2,
+		/// ARMv9.0-A), use the Graviton4 target.
+		/// </summary>
+		GravitonSVE,
+
 		/// <summary>
 		/// AWS Graviton3 CPU - ARM Neoverse V1 uArch with ARMv8.4-A ISA + SVE + bf16 + rng + crypto extensions
 		/// </summary>
